COPY from //Dev-Anim to //Dev-Main

diff --git a/Engine/Shaders/Private/GpuSkinVertexFactory.ush b/Engine/Shaders/Private/GpuSkinVertexFactory.ush
--- a/Engine/Shaders/Private/GpuSkinVertexFactory.ush
+++ b/Engine/Shaders/Private/GpuSkinVertexFactory.ush
@@ -24,6 +24,10 @@
 #define GPUSKIN_USE_EXTRA_INFLUENCES	0	// 8 bones/vertex
 #endif
 
+#ifndef GPUSKIN_UNLIMITED_BONE_INFLUENCE
+#define GPUSKIN_UNLIMITED_BONE_INFLUENCE 0
+#endif
+
 #ifndef GPUSKIN_USE_BONES_SRV_BUFFER
 #define GPUSKIN_USE_BONES_SRV_BUFFER	0	// Use a Buffer<float4> if 1, otherwise a uniform buffer  
 #endif
@@ -44,6 +48,11 @@ float4x4 PreviousClothLocalToWorld;
 float ClothBlendWeight;
 #endif// #if GPUSKIN_APEX_CLOTH
 
+#if GPUSKIN_UNLIMITED_BONE_INFLUENCE
+uint InputWeightIndexSize; // bone index size: 1 - 8bit or 2 - 16bit
+Buffer<uint> InputWeightStream;
+#endif
+
 struct FVertexFactoryInput
 {
 	float4	Position		: ATTRIBUTE0;
@@ -53,6 +62,9 @@ struct FVertexFactoryInput
 	// TangentZ.w contains sign of tangent basis determinant
 	half4	TangentZ		: ATTRIBUTE2;	
 
+#if GPUSKIN_UNLIMITED_BONE_INFLUENCE
+	uint	BlendOffsetCount	: ATTRIBUTE3;
+#else
 #if FEATURE_LEVEL >= FEATURE_LEVEL_ES3_1 || COMPILER_METAL || COMPILER_VULKAN
 	uint4	BlendIndices	: ATTRIBUTE3;
 	#if GPUSKIN_USE_EXTRA_INFLUENCES
@@ -69,6 +81,7 @@ struct FVertexFactoryInput
 #if GPUSKIN_USE_EXTRA_INFLUENCES
 	float4	BlendWeightsExtra	: ATTRIBUTE15;
 #endif
+#endif // GPUSKIN_UNLIMITED_BONE_INFLUENCE
 
 
 #if NUM_MATERIAL_TEXCOORDS_VERTEX
@@ -334,6 +347,31 @@ FBoneMatrix GetBoneMatrix(int Index)
 
 FBoneMatrix CalcBoneMatrix( FVertexFactoryInput Input )
 {
+#if GPUSKIN_UNLIMITED_BONE_INFLUENCE
+	int NumBoneInfluences = Input.BlendOffsetCount & 0xff;
+	int StreamOffset = Input.BlendOffsetCount >> 8;
+	int WeightsOffset = StreamOffset + (InputWeightIndexSize * NumBoneInfluences);
+
+	FBoneMatrix BoneMatrix = FBoneMatrix(float4(0,0,0,0), float4(0,0,0,0), float4(0,0,0,0));
+	for (int InfluenceIdx = 0; InfluenceIdx < NumBoneInfluences; InfluenceIdx++)
+	{
+		int BoneIndexOffset = StreamOffset + (InputWeightIndexSize * InfluenceIdx);
+		int BoneIndex = InputWeightStream[BoneIndexOffset];
+		if (InputWeightIndexSize > 1)
+		{
+			BoneIndex = InputWeightStream[BoneIndexOffset + 1] << 8 | BoneIndex;
+			//@todo-lh: Workaround to fix issue in SPIRVEmitter of DXC; this block must be inside the if branch
+			float BoneWeight = float(InputWeightStream[WeightsOffset + InfluenceIdx]) / 255.0f;
+			BoneMatrix += BoneWeight * GetBoneMatrix(BoneIndex);
+		}
+		else
+		{
+			//@todo-lh: Workaround to fix issue in SPIRVEmitter of DXC; this block must be inside the if branch
+			float BoneWeight = float(InputWeightStream[WeightsOffset + InfluenceIdx]) / 255.0f;
+			BoneMatrix += BoneWeight * GetBoneMatrix(BoneIndex);
+		}
+	}
+#else // GPUSKIN_UNLIMITED_BONE_INFLUENCE
 	FBoneMatrix BoneMatrix = Input.BlendWeights.x * GetBoneMatrix(Input.BlendIndices.x);
 	BoneMatrix += Input.BlendWeights.y * GetBoneMatrix(Input.BlendIndices.y);
 
@@ -348,6 +386,7 @@ FBoneMatrix CalcBoneMatrix( FVertexFactoryInput Input )
 	BoneMatrix += Input.BlendWeightsExtra.w * GetBoneMatrix(Input.BlendIndicesExtra.w);
 #endif
 #endif//GPUSKIN_LIMIT_2BONE_INFLUENCES
+#endif // GPUSKIN_UNLIMITED_BONE_INFLUENCE
 	return BoneMatrix;
 }
 
@@ -365,6 +404,31 @@ FBoneMatrix GetPreviousBoneMatrix(int Index)
 
 FBoneMatrix CalcPreviousBoneMatrix( FVertexFactoryInput Input )
 {
+#if GPUSKIN_UNLIMITED_BONE_INFLUENCE
+	int NumBoneInfluences = Input.BlendOffsetCount & 0xff;
+	int StreamOffset = Input.BlendOffsetCount >> 8;
+	int WeightsOffset = StreamOffset + (InputWeightIndexSize * NumBoneInfluences);
+
+	FBoneMatrix BoneMatrix = FBoneMatrix(float4(0,0,0,0), float4(0,0,0,0), float4(0,0,0,0));
+	for (int InfluenceIdx = 0; InfluenceIdx < NumBoneInfluences; InfluenceIdx++)
+	{
+		int BoneIndexOffset = StreamOffset + (InputWeightIndexSize * InfluenceIdx);
+		int BoneIndex = InputWeightStream[BoneIndexOffset];
+		if (InputWeightIndexSize > 1)
+		{
+			BoneIndex = InputWeightStream[BoneIndexOffset + 1] << 8 | BoneIndex;
+			//@todo-lh: Workaround to fix issue in SPIRVEmitter of DXC; this block must be inside the if branch
+			float BoneWeight = float(InputWeightStream[WeightsOffset + InfluenceIdx]) / 255.0f;
+			BoneMatrix += BoneWeight * GetPreviousBoneMatrix(BoneIndex);
+		}
+		else
+		{
+			//@todo-lh: Workaround to fix issue in SPIRVEmitter of DXC; this block must be inside the if branch
+			float BoneWeight = float(InputWeightStream[WeightsOffset + InfluenceIdx]) / 255.0f;
+			BoneMatrix += BoneWeight * GetPreviousBoneMatrix(BoneIndex);
+		}
+	}
+#else // GPUSKIN_UNLIMITED_BONE_INFLUENCE
 	FBoneMatrix BoneMatrix = Input.BlendWeights.x * GetPreviousBoneMatrix(Input.BlendIndices.x);
 	BoneMatrix += Input.BlendWeights.y * GetPreviousBoneMatrix(Input.BlendIndices.y);
 
@@ -379,6 +443,7 @@ FBoneMatrix CalcPreviousBoneMatrix( FVertexFactoryInput Input )
 	BoneMatrix += Input.BlendWeightsExtra.w * GetPreviousBoneMatrix(Input.BlendIndicesExtra.w);
 #endif
 #endif//GPUSKIN_LIMIT_2BONE_INFLUENCES
+#endif // GPUSKIN_UNLIMITED_BONE_INFLUENCE
 	return BoneMatrix;
 }
 
diff --git a/Engine/Source/Developer/AssetTools/Private/AssetTools.cpp b/Engine/Source/Developer/AssetTools/Private/AssetTools.cpp
--- a/Engine/Source/Developer/AssetTools/Private/AssetTools.cpp
+++ b/Engine/Source/Developer/AssetTools/Private/AssetTools.cpp
@@ -38,6 +38,7 @@
 #include "AssetTypeActions/AssetTypeActions_VectorField.h"
 #include "AssetTypeActions/AssetTypeActions_AnimationAsset.h"
 #include "AssetTypeActions/AssetTypeActions_AnimBlueprint.h"
+#include "AssetTypeActions/AssetTypeActions_AnimBoneCompressionSettings.h"
 #include "AssetTypeActions/AssetTypeActions_AnimComposite.h"
 #include "AssetTypeActions/AssetTypeActions_AnimStreamable.h"
 #include "AssetTypeActions/AssetTypeActions_AnimCurveCompressionSettings.h"
@@ -186,6 +187,7 @@ UAssetToolsImpl::UAssetToolsImpl(const FObjectInitializer& ObjectInitializer)
 	// Register the built-in asset type actions
 	RegisterAssetTypeActions(MakeShareable(new FAssetTypeActions_AnimationAsset));
 	RegisterAssetTypeActions(MakeShareable(new FAssetTypeActions_AnimBlueprint));
+	RegisterAssetTypeActions(MakeShareable(new FAssetTypeActions_AnimBoneCompressionSettings));
 	RegisterAssetTypeActions(MakeShareable(new FAssetTypeActions_AnimComposite));
 	RegisterAssetTypeActions(MakeShareable(new FAssetTypeActions_AnimStreamable));
 	RegisterAssetTypeActions(MakeShareable(new FAssetTypeActions_AnimCurveCompressionSettings));
diff --git a/Engine/Source/Developer/MeshBoneReduction/Private/MeshBoneReduction.cpp b/Engine/Source/Developer/MeshBoneReduction/Private/MeshBoneReduction.cpp
--- a/Engine/Source/Developer/MeshBoneReduction/Private/MeshBoneReduction.cpp
+++ b/Engine/Source/Developer/MeshBoneReduction/Private/MeshBoneReduction.cpp
@@ -132,7 +132,9 @@ public:
 	{
 		// now you have list of bones, remove them from vertex influences
 		{
-			TMap<uint8, uint8> BoneMapRemapTable;
+			// FBoneIndexType/uint16 max range
+			const int32 FBoneIndexTypeMax = 65536;
+			TMap<FBoneIndexType, FBoneIndexType> BoneMapRemapTable;
 			// first go through bone map and see if this contains BonesToRemove
 			int32 BoneMapSize = Section.BoneMap.Num();
 			int32 AdjustIndex=0;
@@ -172,7 +174,7 @@ public:
 					// first fix up all indices of BoneMapRemapTable for the indices higher than BoneMapIndex, since BoneMapIndex is being removed
 					for (auto Iter = BoneMapRemapTable.CreateIterator(); Iter; ++Iter)
 					{
-						uint8& Value = Iter.Value();
+						FBoneIndexType& Value = Iter.Value();
 
 						check (Value != BoneMapIndex);
 						if (Value > BoneMapIndex)
@@ -186,10 +188,10 @@ public:
 					// you still have to add no matter what even if same since indices might change after added
 					{
 						// add to remap table
-						check (OldIndex < 256 && OldIndex >= 0);
-						check (NewIndex < 256 && NewIndex >= 0);
-						check (BoneMapRemapTable.Contains((uint8)OldIndex) == false);
-						BoneMapRemapTable.Add((uint8)OldIndex, (uint8)NewIndex);
+						check (OldIndex < FBoneIndexTypeMax && OldIndex >= 0);
+						check (NewIndex < FBoneIndexTypeMax && NewIndex >= 0);
+						check (BoneMapRemapTable.Contains((FBoneIndexType)OldIndex) == false);
+						BoneMapRemapTable.Add((FBoneIndexType)OldIndex, (FBoneIndexType)NewIndex);
 					}
 
 					// reduce index since the item is removed
@@ -203,11 +205,10 @@ public:
 				{
 					int32 OldIndex = BoneMapIndex+AdjustIndex;
 					int32 NewIndex = BoneMapIndex;
-
-					check (OldIndex < 256 && OldIndex >= 0);
-					check (NewIndex < 256 && NewIndex >= 0);
-					check (BoneMapRemapTable.Contains((uint8)OldIndex) == false);
-					BoneMapRemapTable.Add((uint8)OldIndex, (uint8)NewIndex);
+					check (OldIndex < FBoneIndexTypeMax && OldIndex >= 0);
+					check (NewIndex < FBoneIndexTypeMax && NewIndex >= 0);
+					check (BoneMapRemapTable.Contains((FBoneIndexType)OldIndex) == false);
+					BoneMapRemapTable.Add((FBoneIndexType)OldIndex, (FBoneIndexType)NewIndex);
 				}
 			}
 
@@ -219,13 +220,13 @@ public:
 				{
 					FSoftSkinVertex & Vert = Section.SoftVertices[VertIndex];
 
-					auto RemapBoneInfluenceVertexIndex = [&BoneMapRemapTable](uint8 InfluenceBones[MAX_TOTAL_INFLUENCES], uint8 InfluenceWeights[MAX_TOTAL_INFLUENCES])
+					auto RemapBoneInfluenceVertexIndex = [&BoneMapRemapTable](FBoneIndexType InfluenceBones[MAX_TOTAL_INFLUENCES], uint8 InfluenceWeights[MAX_TOTAL_INFLUENCES])
 					{
 						bool ShouldRenormalize = false;
 
 						for (int32 InfluenceIndex = 0; InfluenceIndex < MAX_TOTAL_INFLUENCES; InfluenceIndex++)
 						{
-							uint8 *RemappedBone = BoneMapRemapTable.Find(InfluenceBones[InfluenceIndex]);
+							FBoneIndexType *RemappedBone = BoneMapRemapTable.Find(InfluenceBones[InfluenceIndex]);
 							if (RemappedBone)
 							{
 								InfluenceBones[InfluenceIndex] = *RemappedBone;
diff --git a/Engine/Source/Developer/MeshUtilities/Private/MeshUtilities.cpp b/Engine/Source/Developer/MeshUtilities/Private/MeshUtilities.cpp
--- a/Engine/Source/Developer/MeshUtilities/Private/MeshUtilities.cpp
+++ b/Engine/Source/Developer/MeshUtilities/Private/MeshUtilities.cpp
@@ -790,6 +790,7 @@ void FMeshUtilities::BuildSkeletalModelFromChunks(FSkeletalMeshLODModel& LODMode
 
 		// update max bone influences
 		Section.CalcMaxBoneInfluences();
+		Section.CalcUse16BitBoneIndex();
 
 		// Log info about the chunk.
 		UE_LOG(LogSkeletalMesh, Log, TEXT("Section %u: %u vertices, %u active bones"),
@@ -4060,7 +4061,13 @@ public:
 						InfluenceCount++;
 						LookIdx++;
 					}
+
 					InfluenceCount = FMath::Min<uint32>(InfluenceCount, MAX_TOTAL_INFLUENCES);
+					if (InfluenceCount > EXTRA_BONE_INFLUENCES && !FGPUBaseSkinVertexFactory::UseUnlimitedBoneInfluences(InfluenceCount))
+					{
+						InfluenceCount = EXTRA_BONE_INFLUENCES;
+						UE_LOG(LogSkeletalMesh, Warning, TEXT("Skeletal mesh of %d bone influences requires unlimited bone influence mode on. Influence truncated to %d."), InfluenceCount, EXTRA_BONE_INFLUENCES);
+					}
 
 					// Setup the vertex influences.
 					Vertex.InfluenceBones[0] = 0;
@@ -5737,8 +5744,7 @@ void FMeshUtilities::GenerateRuntimeSkinWeightData(const FSkeletalMeshLODModel*
 		TargetLODModel.GetVertices(TargetVertices);
 
 		// Determine how many influences each skinweight can contain
-		const bool bTargetExtraBoneInfluences = TargetLODModel.DoSectionsNeedExtraBoneInfluences();
-		const int32 NumInfluences = bTargetExtraBoneInfluences ? MAX_TOTAL_INFLUENCES : MAX_INFLUENCES_PER_STREAM;
+		const int32 NumInfluences = TargetLODModel.GetMaxBoneInfluences();
 
 		TArray<FRawSkinWeight> UniqueWeights;
 		for (int32 VertexIndex = 0; VertexIndex < TargetVertices.Num(); ++VertexIndex)
@@ -5787,9 +5793,9 @@ void FMeshUtilities::GenerateRuntimeSkinWeightData(const FSkeletalMeshLODModel*
 					{
 						if (SourceSkinWeight.InfluenceWeights[InfluenceIndex] > 0)
 						{
-							const uint16 Index = SourceSkinWeight.InfluenceBones[InfluenceIndex] << 8;
-							const uint16 Weight = SourceSkinWeight.InfluenceWeights[InfluenceIndex];
-							const uint16 Value = Index | Weight;
+							const uint32 Index = SourceSkinWeight.InfluenceBones[InfluenceIndex] << 16;
+							const uint32 Weight = SourceSkinWeight.InfluenceWeights[InfluenceIndex];
+							const uint32 Value = Index | Weight;
 
 							InOutSkinWeightOverrideData.Weights.Add(Value);
 							++DeltaOverride.NumInfluences;
diff --git a/Engine/Source/Developer/RigVMDeveloper/Private/RigVMCompiler/RigVMAST.cpp b/Engine/Source/Developer/RigVMDeveloper/Private/RigVMCompiler/RigVMAST.cpp
new file mode 100644
--- /dev/null
+++ b/Engine/Source/Developer/RigVMDeveloper/Private/RigVMCompiler/RigVMAST.cpp
@@ -0,0 +1,1321 @@
+// Copyright Epic Games, Inc. All Rights Reserved.
+
+#include "RigVMCompiler/RigVMAST.h"
+#include "RigVMModel/Nodes/RigVMStructNode.h"
+#include "RigVMModel/Nodes/RigVMParameterNode.h"
+#include "RigVMModel/Nodes/RigVMVariableNode.h"
+#include "RigVMModel/Nodes/RigVMCommentNode.h"
+#include "RigVMModel/Nodes/RigVMRerouteNode.h"
+#include "RigVMModel/RigVMGraph.h"
+#include "RigVMCore/RigVMExecuteContext.h"
+#include "Stats/StatsHierarchical.h"
+
+FRigVMExprAST::FRigVMExprAST(const FRigVMParserAST* InParser, EType InType)
+	: Name(NAME_None)
+	, Type(InType)
+	, Index(INDEX_NONE)
+	, ParserPtr(InParser)
+{
+	if (InParser)
+	{
+		Index = InParser->Expressions.Num();
+		((FRigVMParserAST*)ParserPtr)->Expressions.Add(this);
+	}
+}
+
+FName FRigVMExprAST::GetTypeName() const
+{
+	switch (GetType())
+	{
+		case EType::Block:
+		{
+			return TEXT("[.Block.]");
+		}
+		case EType::Entry:
+		{
+			return TEXT("[.Entry.]");
+		}
+		case EType::CallExtern:
+		{
+			return TEXT("[.Call..]");
+		}
+		case EType::NoOp:
+		{
+			return TEXT("[.NoOp..]");
+		}
+		case EType::Var:
+		{
+			return TEXT("[.Var...]");
+		}
+		case EType::Literal:
+		{
+			return TEXT("[Literal]");
+		}
+		case EType::Assign:
+		{
+			return TEXT("[.Assign]");
+		}
+		case EType::Copy:
+		{
+			return TEXT("[.Copy..]");
+		}
+		case EType::CachedValue:
+		{
+			return TEXT("[.Cache.]");
+		}
+		case EType::Exit:
+		{
+			return TEXT("[.Exit..]");
+		}
+		case EType::Invalid:
+		{
+			return TEXT("[Invalid]");
+		}
+		default:
+		{
+			ensure(false);
+		}
+	}
+	return NAME_None;
+}
+
+const FRigVMExprAST* FRigVMExprAST::GetParent() const
+{
+	if (Parents.Num() > 0)
+	{
+		return ParentAt(0);
+	}
+	return nullptr;
+}
+
+const FRigVMBlockExprAST* FRigVMExprAST::GetBlock() const
+{
+	if (Parents.Num() == 0)
+	{
+		if (IsA(EType::Block))
+		{
+			return To<FRigVMBlockExprAST>();
+		}
+		return nullptr;
+	}
+
+	const FRigVMExprAST* Parent = GetParent();
+	if (Parent->IsA(EType::Block))
+	{
+		return Parent->To<FRigVMBlockExprAST>();
+	}
+
+	return Parent->GetBlock();
+}
+
+const FRigVMBlockExprAST* FRigVMExprAST::GetRootBlock() const
+{
+	const FRigVMBlockExprAST* Block = GetBlock();
+
+	if (IsA(EType::Block))
+	{
+		if (Block && NumParents() > 0)
+		{
+			return Block->GetRootBlock();
+		}
+		return To<FRigVMBlockExprAST>();
+	}
+
+	if (Block)
+	{
+		return Block->GetRootBlock();
+	}
+
+	return nullptr;
+}
+
+void FRigVMExprAST::AddParent(FRigVMExprAST* InParent)
+{
+	ensure(InParent != this);
+	if (Parents.Contains(InParent))
+	{
+		return;
+	}
+	InParent->Children.Add(this);
+	Parents.Add(InParent);
+}
+
+void FRigVMExprAST::RemoveParent(FRigVMExprAST* InParent)
+{
+	if (Parents.Remove(InParent) > 0)
+	{
+		InParent->Children.Remove(this);
+	}
+}
+
+void FRigVMExprAST::RemoveChild(FRigVMExprAST* InChild)
+{
+	InChild->RemoveParent(this);
+}
+
+void FRigVMExprAST::ReplaceParent(FRigVMExprAST* InCurrentParent, FRigVMExprAST* InNewParent)
+{
+	for (int32 ParentIndex = 0; ParentIndex < Parents.Num(); ParentIndex++)
+	{
+		if (Parents[ParentIndex] == InCurrentParent)
+		{
+			Parents[ParentIndex] = InNewParent;
+			InCurrentParent->Children.Remove(this);
+			InNewParent->Children.Add(this);
+		}
+	}
+}
+
+void FRigVMExprAST::ReplaceChild(FRigVMExprAST* InCurrentChild, FRigVMExprAST* InNewChild)
+{
+	for (int32 ChildIndex = 0; ChildIndex < Children.Num(); ChildIndex++)
+	{
+		if (Children[ChildIndex] == InCurrentChild)
+		{
+			Children[ChildIndex] = InNewChild;
+			InCurrentChild->Parents.Remove(this);
+			InNewChild->Parents.Add(this);
+		}
+	}
+}
+
+FString FRigVMExprAST::DumpText(const FString& InPrefix) const
+{
+	FString Result;
+	if (Name.IsNone())
+	{
+		Result = FString::Printf(TEXT("%s%s"), *InPrefix, *GetTypeName().ToString());
+	}
+	else
+	{
+		Result = FString::Printf(TEXT("%s%s %s"), *InPrefix, *GetTypeName().ToString(), *Name.ToString());
+	}
+
+	if (Children.Num() > 0)
+	{
+		FString Prefix = InPrefix;
+		if (Prefix.IsEmpty())
+		{
+			Prefix = TEXT("-- ");
+		}
+		else
+		{
+			Prefix = TEXT("---") + Prefix;
+		}
+		for (FRigVMExprAST* Child : Children)
+		{
+			Result += TEXT("\n") + Child->DumpText(Prefix);
+		}
+	}
+	return Result;
+}
+
+FString FRigVMExprAST::DumpDot(TArray<bool>& OutExpressionDefined, const FString& InPrefix) const
+{
+	FString Prefix = InPrefix;
+
+	FString Result;
+	bool bWasDefined = true;
+	if (!OutExpressionDefined[GetIndex()])
+	{
+		bWasDefined = false;
+
+		FString Label = GetName().ToString();
+		FString AdditionalNodeSettings;
+		switch (GetType())
+		{
+			case EType::Literal:
+			{
+				Label = FString::Printf(TEXT("%s(Literal)"), *To<FRigVMLiteralExprAST>()->GetPin()->GetName());
+				break;
+			}
+			case EType::Var:
+			{
+				if (To<FRigVMVarExprAST>()->IsGraphParameter())
+				{
+					URigVMParameterNode* ParameterNode = Cast<URigVMParameterNode>(To<FRigVMVarExprAST>()->GetPin()->GetNode());
+					check(ParameterNode);
+					Label = FString::Printf(TEXT("Param %s"), *ParameterNode->GetParameterName().ToString());
+				}
+				else if (To<FRigVMVarExprAST>()->IsGraphVariable())
+				{
+					URigVMVariableNode* VariableNode = Cast<URigVMVariableNode>(To<FRigVMVarExprAST>()->GetPin()->GetNode());
+					check(VariableNode);
+					Label = FString::Printf(TEXT("Variable %s"), *VariableNode->GetVariableName().ToString());
+				}
+				else
+				{
+					Label = To<FRigVMVarExprAST>()->GetPin()->GetName();
+				}
+				if (To<FRigVMVarExprAST>()->IsExecuteContext())
+				{
+					AdditionalNodeSettings += TEXT(", shape = cds");
+				}
+				break;
+			}
+			case EType::Block:
+			{
+				if (GetParent() == nullptr)
+				{
+					Label = TEXT("Unused");
+					Result += FString::Printf(TEXT("\n%ssubgraph unused_%d {"), *Prefix, GetIndex());
+					Prefix += TEXT("  ");
+				}
+				else
+				{
+					Label = TEXT("Block");
+				}
+				break;
+			}
+			case EType::Assign:
+			{
+				Label = TEXT("=");
+				break;
+			}
+			case EType::Copy:
+			{
+				Label = TEXT("Copy");
+				break;
+			}
+			case EType::CachedValue:
+			{
+				Label = TEXT("Cache");
+				break;
+			}
+			case EType::CallExtern:
+			{
+				if (URigVMStructNode* Node = Cast<URigVMStructNode>(To<FRigVMCallExternExprAST>()->GetNode()))
+				{
+					Label = Node->GetScriptStruct()->GetName();
+				}
+				break;
+			}
+			case EType::NoOp:
+			{
+				Label = TEXT("NoOp");
+				break;
+			}
+			case EType::Exit:
+			{
+				Label = TEXT("Exit");
+				break;
+			}
+			case EType::Entry:
+			{
+				Result += FString::Printf(TEXT("\n%ssubgraph %s_%d {"), *Prefix, *GetName().ToString(), GetIndex());
+				Prefix += TEXT("  ");
+				break;
+			}
+			default:
+			{
+				break;
+			}
+		}
+
+		if (!Label.IsEmpty())
+		{
+			Result += FString::Printf(TEXT("\n%snode_%d [label = \"%s\"%s];"), *Prefix, GetIndex(), *Label, *AdditionalNodeSettings);
+		}
+
+		switch (GetType())
+		{
+			case EType::Entry:
+			case EType::Exit:
+			{
+				Result += FString::Printf(TEXT("\n%snode_%d [shape = Mdiamond];"), *Prefix, GetIndex());
+				break;
+			}
+			case EType::Assign:
+			case EType::Copy:
+			case EType::CallExtern:
+			case EType::NoOp:
+			{
+				Result += FString::Printf(TEXT("\n%snode_%d [shape = box];"), *Prefix, GetIndex());
+				break;
+			}
+			default:
+			{
+				break;
+			}
+		}
+
+	}
+
+	for (FRigVMExprAST* Child : Children)
+	{
+		Result += Child->DumpDot(OutExpressionDefined, Prefix);
+		if(!bWasDefined)
+		{
+			Result += FString::Printf(TEXT("\n%snode_%d -> node_%d;"), *Prefix, GetIndex(), Child->GetIndex());
+		}
+	}
+
+	if (!OutExpressionDefined[GetIndex()])
+	{
+		switch (GetType())
+		{
+			case EType::Block:
+			{
+				if (GetParent() == nullptr)
+				{
+					Prefix = Prefix.LeftChop(2);
+					Result += FString::Printf(TEXT("\n%s}"), *Prefix, *GetName().ToString(), GetIndex());
+				}
+				break;
+			}
+			case EType::Entry:
+			{
+				Prefix = Prefix.LeftChop(2);
+				Result += FString::Printf(TEXT("\n%s}"), *Prefix, *GetName().ToString(), GetIndex());
+				break;
+			}
+			default:
+			{
+				break;
+			}
+		}
+	}
+
+	OutExpressionDefined[GetIndex()] = true;
+
+	return Result;
+}
+
+bool FRigVMBlockExprAST::ShouldExecute() const
+{
+	return ContainsEntry();
+}
+
+bool FRigVMBlockExprAST::ContainsEntry() const
+{
+	if (IsA(FRigVMExprAST::EType::Entry))
+	{
+		return true;
+	}
+	for (FRigVMExprAST* Expression : *this)
+	{
+		if (Expression->IsA(EType::Entry))
+		{
+			return true;
+		}
+	}
+	return false;
+}
+
+bool FRigVMBlockExprAST::Contains(const FRigVMExprAST* InExpression) const
+{
+	for (FRigVMExprAST* Expression : *this)
+	{
+		if (Expression == InExpression)
+		{
+			return true;
+		}
+
+		if (Expression->IsA(EType::Block))
+		{
+			if (Expression->To<FRigVMBlockExprAST>()->Contains(InExpression))
+			{
+				return true;
+			}
+		}
+	}
+	return false;
+}
+
+FName FRigVMEntryExprAST::GetEventName() const
+{
+	if (URigVMNode* EventNode = GetNode())
+	{
+		return EventNode->GetEventName();
+	}
+	return NAME_None;
+}
+
+FString FRigVMVarExprAST::GetCPPType() const
+{
+	return Pin->GetCPPType();
+}
+
+UObject* FRigVMVarExprAST::GetCPPTypeObject() const
+{
+	return Pin->GetCPPTypeObject();
+}
+
+ERigVMPinDirection FRigVMVarExprAST::GetPinDirection() const
+{
+	return Pin->GetDirection();
+}
+
+FString FRigVMVarExprAST::GetDefaultValue() const
+{
+	return Pin->GetDefaultValue();
+}
+
+bool FRigVMVarExprAST::IsExecuteContext() const
+{
+	return Pin->IsExecuteContext();
+}
+
+bool FRigVMVarExprAST::IsGraphParameter() const
+{
+	if (Cast<URigVMParameterNode>(Pin->GetNode()))
+	{
+		return Pin->GetName() == TEXT("Value");
+	}
+	return false;
+}
+
+bool FRigVMVarExprAST::IsGraphVariable() const
+{
+	if (Cast<URigVMVariableNode>(Pin->GetNode()))
+	{
+		return Pin->GetName() == TEXT("Value");
+	}
+	return false;
+}
+
+FRigVMParserAST::FRigVMParserAST(URigVMGraph* InGraph, const FRigVMParserASTSettings& InSettings)
+{
+	DECLARE_SCOPE_HIERARCHICAL_COUNTER_FUNC()
+	
+	LastCycleCheckExpr = nullptr;
+
+	const TArray<URigVMNode*> Nodes = InGraph->GetNodes();
+
+	NodeExpressionIndex.SetNumUninitialized(Nodes.Num());
+	for (int32& Index : NodeExpressionIndex)
+	{
+		Index = INDEX_NONE;
+	}
+
+	for (URigVMNode* Node : Nodes)
+	{
+		if(Node->IsEvent())
+		{
+			TraverseMutableNode(Node, nullptr);
+		}
+	}
+
+	// traverse all remaining mutable nodes,
+	// followed by a pass for all remaining non-mutable nodes
+	FRigVMBlockExprAST* ObsoleteBlock = nullptr;
+	for (int32 PassIndex = 0; PassIndex < 2; PassIndex++)
+	{
+		const bool bTraverseMutable = PassIndex == 0;
+		for (int32 NodeIndex = 0; NodeIndex < NodeExpressionIndex.Num(); NodeIndex++)
+		{
+			if (NodeExpressionIndex[NodeIndex] != INDEX_NONE)
+			{
+				continue;
+			}
+
+			if (Nodes[NodeIndex]->IsMutable() == bTraverseMutable)
+			{
+				if (ObsoleteBlock == nullptr)
+				{
+					ObsoleteBlock = new FRigVMBlockExprAST(this);
+					RootExpressions.Add(ObsoleteBlock);
+				}
+				if (bTraverseMutable)
+				{
+					TraverseMutableNode(Nodes[NodeIndex], ObsoleteBlock);
+				}
+				else
+				{
+					TraverseNode(Nodes[NodeIndex], ObsoleteBlock);
+				}
+			}
+		}
+	}
+
+	FoldEntries();
+	InjectExitsToEntries();
+
+	if (InSettings.bFoldReroutes || InSettings.bFoldAssignments)
+	{
+		FoldNoOps();
+	}
+
+	if (InSettings.bFoldAssignments)
+	{
+		FoldAssignments();
+	}
+
+	if (InSettings.bFoldLiterals)
+	{
+		FoldLiterals();
+	}
+}
+
+FRigVMParserAST::~FRigVMParserAST()
+{
+	for (FRigVMExprAST* Expression : Expressions)
+	{
+		delete(Expression);
+	}
+	Expressions.Empty();
+
+	// root expressions are a subset of the
+	// expressions array, so no cleanup necessary
+	RootExpressions.Empty();
+}
+
+FRigVMExprAST* FRigVMParserAST::TraverseMutableNode(URigVMNode* InNode, FRigVMExprAST* InParentExpr)
+{
+	if (SubjectToExpression.Contains(InNode))
+	{
+		return SubjectToExpression.FindChecked(InNode);
+	}
+
+	FRigVMExprAST* NodeExpr = nullptr;
+	if (InNode->IsEvent())
+	{
+		NodeExpr = new FRigVMEntryExprAST(this, InNode);
+		NodeExpr->Name = InNode->GetEventName();
+	}
+	else
+	{
+		if (Cast<URigVMRerouteNode>(InNode) ||
+			Cast<URigVMParameterNode>(InNode) ||
+			Cast<URigVMVariableNode>(InNode))
+		{
+			NodeExpr = new FRigVMNoOpExprAST(this, InNode);
+		}
+		else
+		{
+			NodeExpr = new FRigVMCallExternExprAST(this, InNode);
+		}
+		NodeExpr->Name = InNode->GetFName();
+	}
+
+	if (InParentExpr != nullptr)
+	{
+		NodeExpr->AddParent(InParentExpr);
+	}
+	else
+	{
+		RootExpressions.Add(NodeExpr);
+		InParentExpr = NodeExpr;
+	}
+	SubjectToExpression.Add(InNode, NodeExpr);
+	NodeExpressionIndex[InNode->GetNodeIndex()] = NodeExpr->GetIndex();
+
+	TraversePins(InNode, NodeExpr);
+
+	for (URigVMPin* SourcePin : InNode->GetPins())
+	{
+		if (SourcePin->GetDirection() == ERigVMPinDirection::Output || SourcePin->GetDirection() == ERigVMPinDirection::IO)
+		{
+			if (SourcePin->GetScriptStruct()->IsChildOf(FRigVMExecuteContext::StaticStruct()))
+			{
+				TArray<URigVMPin*> TargetPins = SourcePin->GetLinkedTargetPins();
+				for (URigVMPin* TargetPin : TargetPins)
+				{
+					TraverseMutableNode(TargetPin->GetNode(), InParentExpr);
+				}
+			}
+		}
+	}
+
+	return NodeExpr;
+}
+
+FRigVMExprAST* FRigVMParserAST::TraverseNode(URigVMNode* InNode, FRigVMExprAST* InParentExpr)
+{
+	if (Cast<URigVMCommentNode>(InNode))
+	{
+		return nullptr;
+	
+	}
+	if (SubjectToExpression.Contains(InNode))
+	{
+		FRigVMExprAST* NodeExpr = SubjectToExpression.FindChecked(InNode);
+		NodeExpr->AddParent(InParentExpr);
+		return NodeExpr;
+	}
+
+	FRigVMExprAST* NodeExpr = nullptr;
+	if (Cast<URigVMRerouteNode>(InNode) ||
+		Cast<URigVMParameterNode>(InNode) ||
+		Cast<URigVMVariableNode>(InNode))
+	{
+		NodeExpr = new FRigVMNoOpExprAST(this, InNode);
+	}
+	else
+	{
+		NodeExpr = new FRigVMCallExternExprAST(this, InNode);
+	}
+	NodeExpr->Name = InNode->GetFName();
+	NodeExpr->AddParent(InParentExpr);
+	SubjectToExpression.Add(InNode, NodeExpr);
+	NodeExpressionIndex[InNode->GetNodeIndex()] = NodeExpr->GetIndex();
+
+	TraversePins(InNode, NodeExpr);
+
+	return NodeExpr;
+}
+
+TArray<FRigVMExprAST*> FRigVMParserAST::TraversePins(URigVMNode* InNode, FRigVMExprAST* InParentExpr)
+{
+	TArray<FRigVMExprAST*> PinExpressions;
+
+	for (URigVMPin* Pin : InNode->GetPins())
+	{
+		PinExpressions.Add(TraversePin(Pin, InParentExpr));
+	}
+
+	return PinExpressions;
+}
+
+FRigVMExprAST* FRigVMParserAST::TraversePin(URigVMPin* InPin, FRigVMExprAST* InParentExpr)
+{
+	ensure(!SubjectToExpression.Contains(InPin));
+
+	TArray<URigVMLink*> SourceLinks = InPin->GetSourceLinks(true);
+
+	FRigVMExprAST* PinExpr = nullptr;
+
+	if (Cast<URigVMParameterNode>(InPin->GetNode()) ||
+		Cast<URigVMVariableNode>(InPin->GetNode()))
+	{
+		if (InPin->GetDirection() == ERigVMPinDirection::Visible)
+		{
+			return nullptr;
+		}
+	}
+
+	if ((InPin->GetDirection() == ERigVMPinDirection::Input ||
+		InPin->GetDirection() == ERigVMPinDirection::Visible) &&
+		SourceLinks.Num() == 0)
+	{
+		PinExpr = new FRigVMLiteralExprAST(this, InPin);
+	}
+	else
+	{
+		PinExpr = new FRigVMVarExprAST(this, InPin);
+	}
+	PinExpr->AddParent(InParentExpr);
+	PinExpr->Name = *InPin->GetPinPath();
+	SubjectToExpression.Add(InPin, PinExpr);
+
+	if (InPin->GetScriptStruct()->IsChildOf(FRigVMExecuteContext::StaticStruct()))
+	{
+		return PinExpr;
+	}
+
+	if ((InPin->GetDirection() == ERigVMPinDirection::IO ||
+		InPin->GetDirection() == ERigVMPinDirection::Input) 
+		&& !InPin->IsExecuteContext())
+	{
+		bool bHasSourceLinkToRoot = false;
+		URigVMPin* RootPin = InPin->GetRootPin();
+		for (URigVMLink* SourceLink : SourceLinks)
+		{
+			if (SourceLink->GetTargetPin() == RootPin)
+			{
+				bHasSourceLinkToRoot = true;
+				break;
+			}
+		}
+
+		if (!bHasSourceLinkToRoot && (InPin->GetDirection() == ERigVMPinDirection::IO || SourceLinks.Num() > 0))
+		{
+			FRigVMLiteralExprAST* LiteralExpr = new FRigVMLiteralExprAST(this, InPin);
+			FRigVMCopyExprAST* LiteralCopyExpr = new FRigVMCopyExprAST(this, InPin, InPin);
+			LiteralCopyExpr->Name = *FString::Printf(TEXT("%s -> %s"), *InPin->GetPinPath(), *InPin->GetPinPath());
+			LiteralCopyExpr->AddParent(PinExpr);
+			LiteralExpr->AddParent(LiteralCopyExpr);
+			LiteralExpr->Name = *InPin->GetPinPath();
+
+			SubjectToExpression.Remove(InPin);
+			SubjectToExpression.Add(InPin, LiteralExpr);
+		}
+	}
+
+
+	for (URigVMLink* SourceLink : SourceLinks)
+	{
+		TraverseLink(SourceLink, PinExpr);
+	}
+
+	return PinExpr;
+}
+
+FRigVMExprAST* FRigVMParserAST::TraverseLink(URigVMLink* InLink, FRigVMExprAST* InParentExpr)
+{
+	ensure(!SubjectToExpression.Contains(InLink));
+
+	URigVMPin* SourcePin = InLink->GetSourcePin();
+	URigVMPin* TargetPin = InLink->GetTargetPin();
+	URigVMPin* SourceRootPin = SourcePin->GetRootPin();
+	URigVMPin* TargetRootPin = TargetPin->GetRootPin();
+
+	bool bRequiresCopy = SourceRootPin != SourcePin || TargetRootPin != TargetPin;
+	if (!bRequiresCopy)
+	{
+		if(Cast<URigVMParameterNode>(TargetRootPin->GetNode()) ||
+			Cast<URigVMVariableNode>(TargetRootPin->GetNode()))
+		{
+			bRequiresCopy = true;
+		}
+	}
+	
+	FRigVMAssignExprAST* AssignExpr = nullptr;
+	if (bRequiresCopy)
+	{
+		AssignExpr = new FRigVMCopyExprAST(this, SourcePin, TargetPin);
+	}
+	else
+	{
+		AssignExpr = new FRigVMAssignExprAST(this, SourcePin, TargetPin);
+	}
+
+	AssignExpr->Name = *InLink->GetPinPathRepresentation();
+	AssignExpr->AddParent(InParentExpr);
+	SubjectToExpression.Add(InLink, AssignExpr);
+
+	FRigVMExprAST* NodeExpr = TraverseNode(SourcePin->GetNode(), AssignExpr);
+
+	// if this is a copy expression - we should require the copy to use a ref instead
+	if (NodeExpr->IsA(FRigVMExprAST::EType::CallExtern))
+	{
+		for (FRigVMExprAST* ChildExpr : *NodeExpr)
+		{
+			if (ChildExpr->IsA(FRigVMExprAST::EType::Var))
+			{
+				FRigVMVarExprAST* VarExpr = ChildExpr->To<FRigVMVarExprAST>();
+				if (VarExpr->GetPin() == SourceRootPin)
+				{
+					FRigVMCachedValueExprAST* CacheExpr = nullptr;
+					for (FRigVMExprAST* VarExprParent : VarExpr->Parents)
+					{
+						if (VarExprParent->IsA(FRigVMExprAST::EType::CachedValue))
+						{
+							CacheExpr = VarExprParent->To<FRigVMCachedValueExprAST>();
+							break;
+						}
+					}
+
+					if(CacheExpr == nullptr)
+					{
+						CacheExpr = new FRigVMCachedValueExprAST(this);
+						CacheExpr->Name = AssignExpr->GetName();
+						VarExpr->AddParent(CacheExpr);
+						NodeExpr->AddParent(CacheExpr);
+					}
+
+					AssignExpr->ReplaceChild(NodeExpr, CacheExpr);
+					return AssignExpr;
+				}
+			}
+		}
+		checkNoEntry();
+	}
+
+	return AssignExpr;
+}
+
+void FRigVMParserAST::FoldEntries()
+{
+	DECLARE_SCOPE_HIERARCHICAL_COUNTER_FUNC()
+
+	TArray<FRigVMExprAST*> FoldRootExpressions;
+	TArray<FRigVMExprAST*> ExpressionsToRemove;
+	TMap<FName, FRigVMEntryExprAST*> EntryByName;
+
+	for (FRigVMExprAST* RootExpr : RootExpressions)
+	{
+		if (RootExpr->IsA(FRigVMExprAST::EType::Entry))
+		{
+			FRigVMEntryExprAST* Entry = RootExpr->To<FRigVMEntryExprAST>();
+			if (EntryByName.Contains(Entry->GetEventName()))
+			{
+				FRigVMEntryExprAST* FoldEntry = EntryByName.FindChecked(Entry->GetEventName());
+
+				// replace the original entry with a noop
+				FRigVMNoOpExprAST* NoOpExpr = new FRigVMNoOpExprAST(this, Entry->GetNode());
+				NoOpExpr->AddParent(FoldEntry);
+				NoOpExpr->Name = Entry->Name;
+				SubjectToExpression.FindChecked(Entry->GetNode()) = NoOpExpr;
+
+				TArray<FRigVMExprAST*> Children = Entry->Children; // copy since the loop changes the array
+				for (FRigVMExprAST* ChildExpr : Children)
+				{
+					ChildExpr->RemoveParent(Entry);
+					if (ChildExpr->IsA(FRigVMExprAST::Var))
+					{
+						if (ChildExpr->To<FRigVMVarExprAST>()->IsExecuteContext())
+						{
+							ExpressionsToRemove.Add(ChildExpr);
+							continue;
+						}
+					}
+					ChildExpr->AddParent(FoldEntry);
+				}
+				ExpressionsToRemove.Add(Entry);
+			}
+			else
+			{
+				FoldRootExpressions.Add(Entry);
+				EntryByName.Add(Entry->GetEventName(), Entry);
+			}
+		}
+		else
+		{
+			FoldRootExpressions.Add(RootExpr);
+		}
+	}
+
+	RootExpressions = FoldRootExpressions;
+
+	RemoveExpressions(ExpressionsToRemove);
+}
+
+void FRigVMParserAST::InjectExitsToEntries()
+{
+	DECLARE_SCOPE_HIERARCHICAL_COUNTER_FUNC()
+
+	for (FRigVMExprAST* RootExpr : RootExpressions)
+	{
+		if (RootExpr->IsA(FRigVMExprAST::EType::Entry))
+		{
+			bool bHasExit = false;
+			if (RootExpr->Children.Num() > 0)
+			{
+				if (RootExpr->Children.Last()->IsA(FRigVMExprAST::EType::Exit))
+				{
+					bHasExit = true;
+					break;
+				}
+			}
+
+			if (!bHasExit)
+			{
+				FRigVMExprAST* ExitExpr = new FRigVMExitExprAST(this);
+				ExitExpr->AddParent(RootExpr);
+				Expressions.AddUnique(ExitExpr);
+			}
+		}
+	}
+}
+
+void FRigVMParserAST::RefreshExprIndices()
+{
+	for (int32 Index = 0; Index < Expressions.Num(); Index++)
+	{
+		Expressions[Index]->Index = Index;
+	}
+}
+
+void FRigVMParserAST::FoldNoOps()
+{
+	DECLARE_SCOPE_HIERARCHICAL_COUNTER_FUNC()
+
+	for (FRigVMExprAST* Expression : Expressions)
+	{
+		if (Expression->IsA(FRigVMExprAST::EType::NoOp))
+		{
+			if (URigVMNode* Node = Expression->To<FRigVMNoOpExprAST>()->GetNode())
+			{
+				if (URigVMParameterNode* ParameterNode = Cast<URigVMParameterNode>(Node))
+				{
+					if (!ParameterNode->IsInput())
+					{
+						continue;
+					}
+				}
+				if(URigVMVariableNode* VariableNode = Cast<URigVMVariableNode>(Node))
+				{
+					if (!VariableNode->IsGetter())
+					{
+						continue;
+					}
+				}
+			}
+			// copy since we are changing the content during iteration below
+			TArray<FRigVMExprAST*> Children = Expression->Children;
+			TArray<FRigVMExprAST*> Parents = Expression->Parents;
+
+			for (FRigVMExprAST* Parent : Parents)
+			{
+				Expression->RemoveParent(Parent);
+			}
+
+			for (FRigVMExprAST* Child : Children)
+			{
+				Child->RemoveParent(Expression);
+				for (FRigVMExprAST* Parent : Parents)
+				{
+					Child->AddParent(Parent);
+				}
+			}
+		}
+	}
+}
+
+void FRigVMParserAST::FoldAssignments()
+{
+	DECLARE_SCOPE_HIERARCHICAL_COUNTER_FUNC()
+
+	TArray<FRigVMExprAST*> ExpressionsToRemove;
+
+	// first - Fold all assignment chains
+	for (FRigVMExprAST* Expression : Expressions)
+	{
+		if (Expression->Parents.Num() == 0)
+		{
+			continue;
+		}
+
+		if (Expression->GetType() != FRigVMExprAST::EType::Assign)
+		{
+			continue;
+		}
+
+		FRigVMAssignExprAST* AssignExpr = Expression->To<FRigVMAssignExprAST>();
+		ensure(AssignExpr->Parents.Num() == 1);
+		ensure(AssignExpr->Children.Num() == 1);
+
+		// non-input pins on anything but a reroute node should be skipped
+		if (AssignExpr->GetTargetPin()->GetDirection() != ERigVMPinDirection::Input &&
+			Cast<URigVMRerouteNode>(AssignExpr->GetTargetPin()->GetNode()) == nullptr)
+		{
+			continue;
+		}
+
+		FRigVMExprAST* Parent = AssignExpr->Parents[0];
+		FRigVMExprAST* Child = AssignExpr->Children[0];
+		AssignExpr->RemoveParent(Parent);
+		Child->RemoveParent(AssignExpr);
+
+		TArray<FRigVMExprAST*> GrandParents = Parent->Parents;
+		for (FRigVMExprAST* GrandParent : GrandParents)
+		{
+			GrandParent->ReplaceChild(Parent, Child);
+			if (GrandParent->IsA(FRigVMExprAST::EType::Assign))
+			{
+				FRigVMAssignExprAST* GrandParentAssign = GrandParent->To<FRigVMAssignExprAST>();
+				GrandParentAssign->SourcePin = AssignExpr->SourcePin;
+				GrandParentAssign->Name = *FString::Printf(TEXT("%s -> %s"), *GrandParentAssign->SourcePin->GetPinPath(), *GrandParentAssign->TargetPin->GetPinPath());
+			}
+		}
+
+		ExpressionsToRemove.Add(AssignExpr);
+		if (Parent->Parents.Num() == 0)
+		{
+			ExpressionsToRemove.Add(Parent);
+		}
+	}
+
+	RemoveExpressions(ExpressionsToRemove);
+}
+
+void FRigVMParserAST::FoldLiterals()
+{
+	DECLARE_SCOPE_HIERARCHICAL_COUNTER_FUNC()
+
+	TMap<FString, FRigVMLiteralExprAST*> ValueToLiteral;
+	TArray<FRigVMExprAST*> ExpressionsToRemove;
+
+	for (int32 ExpressionIndex = 0; ExpressionIndex < Expressions.Num(); ExpressionIndex++)
+	{
+		FRigVMExprAST* Expression = Expressions[ExpressionIndex];
+		if (Expression->Parents.Num() == 0)
+		{
+			continue;
+		}
+
+		if (Expression->GetType() == FRigVMExprAST::EType::Literal)
+		{
+			ensure(Expression->Children.Num() == 0);
+
+			FRigVMLiteralExprAST* LiteralExpr = Expression->To<FRigVMLiteralExprAST>();
+			FString Hash = FString::Printf(TEXT("[%s] %s"), *LiteralExpr->GetCPPType(), *LiteralExpr->GetDefaultValue());
+
+			FRigVMLiteralExprAST* const* MappedExpr = ValueToLiteral.Find(Hash);
+			if (MappedExpr)
+			{
+				TArray<FRigVMExprAST*> Parents = Expression->Parents;
+				for (FRigVMExprAST* Parent : Parents)
+				{
+					Parent->ReplaceChild(Expression, *MappedExpr);
+				}
+				ExpressionsToRemove.Add(Expression);
+			}
+			else
+			{
+				ValueToLiteral.Add(Hash, LiteralExpr);
+			}
+		}
+	}
+
+	RemoveExpressions(ExpressionsToRemove);
+}
+
+const FRigVMExprAST* FRigVMParserAST::GetExprForSubject(UObject* InSubject)
+{
+	if (FRigVMExprAST* const* ExpressionPtr = SubjectToExpression.Find(InSubject))
+	{
+		return *ExpressionPtr;
+	}
+	return nullptr;
+}
+
+void FRigVMParserAST::PrepareCycleChecking(URigVMPin* InPin)
+{
+	if (InPin == nullptr)
+	{
+		LastCycleCheckExpr = nullptr;
+		CycleCheckFlags.Reset();
+		return;
+	}
+
+	const FRigVMExprAST* Expression = nullptr;
+	if (FRigVMExprAST* const* ExpressionPtr = SubjectToExpression.Find(InPin->GetNode()))
+	{
+		Expression = *ExpressionPtr;
+	}
+	else
+	{
+		return;
+	}
+
+	if (LastCycleCheckExpr != Expression)
+	{
+		LastCycleCheckExpr = Expression;
+		CycleCheckFlags.SetNumZeroed(Expressions.Num());
+		CycleCheckFlags[LastCycleCheckExpr->GetIndex()] = ETraverseRelationShip_Self;
+	}
+}
+
+bool FRigVMParserAST::CanLink(URigVMPin* InSourcePin, URigVMPin* InTargetPin, FString* OutFailureReason)
+{
+	if (InSourcePin == nullptr || InTargetPin == nullptr || InSourcePin == InTargetPin)
+	{
+		if (OutFailureReason)
+		{
+			*OutFailureReason = FString(TEXT("Provided objects contain nullptr."));
+		}
+		return false;
+	}
+
+	URigVMNode* SourceNode = InSourcePin->GetNode();
+	URigVMNode* TargetNode = InTargetPin->GetNode();
+	if (SourceNode == TargetNode)
+	{
+		if (OutFailureReason)
+		{
+			*OutFailureReason = FString(TEXT("Source and Target Nodes are identical."));
+		}
+		return false;
+	}
+
+	const FRigVMExprAST* SourceExpression = nullptr;
+	if (FRigVMExprAST* const* SourceExpressionPtr = SubjectToExpression.Find(SourceNode))
+	{
+		SourceExpression = *SourceExpressionPtr;
+	}
+	else
+	{
+		if (OutFailureReason)
+		{
+			*OutFailureReason = FString(TEXT("Source node is not part of AST."));
+		}
+		return false;
+	}
+
+	const FRigVMExprAST* TargetExpression = nullptr;
+	if (FRigVMExprAST* const* TargetExpressionPtr = SubjectToExpression.Find(TargetNode))
+	{
+		TargetExpression = *TargetExpressionPtr;
+	}
+	else
+	{
+		if (OutFailureReason)
+		{
+			*OutFailureReason = FString(TEXT("Target node is not part of AST."));
+		}
+		return false;
+	}
+
+	const FRigVMBlockExprAST* SourceBlock = SourceExpression->GetBlock();
+	const FRigVMBlockExprAST* TargetBlock = TargetExpression->GetBlock();
+	if (SourceBlock == nullptr || TargetBlock == nullptr)
+	{
+		return false;
+	}
+
+	if (SourceBlock == TargetBlock ||
+		SourceBlock->Contains(TargetBlock) ||
+		TargetBlock->Contains(SourceBlock))
+	{
+		if (LastCycleCheckExpr != SourceExpression && LastCycleCheckExpr != TargetExpression)
+		{
+			PrepareCycleChecking(InSourcePin);
+		}
+
+		TArray<ETraverseRelationShip>& Flags = CycleCheckFlags;
+		TraverseParents(LastCycleCheckExpr, [&Flags](const FRigVMExprAST* InExpr) -> bool {
+			if (Flags[InExpr->GetIndex()] == ETraverseRelationShip_Self)
+			{
+				return true;
+			}
+			if (Flags[InExpr->GetIndex()] != ETraverseRelationShip_Unknown)
+			{
+				return false;
+			}
+			Flags[InExpr->GetIndex()] = ETraverseRelationShip_Parent;
+			return true;
+		});
+
+		TraverseChildren(LastCycleCheckExpr, [&Flags](const FRigVMExprAST* InExpr) -> bool {
+			if (Flags[InExpr->GetIndex()] == ETraverseRelationShip_Self)
+			{
+				return true;
+			}
+			if (Flags[InExpr->GetIndex()] != ETraverseRelationShip_Unknown)
+			{
+				return false;
+			}
+			Flags[InExpr->GetIndex()] = ETraverseRelationShip_Child;
+			return true;
+		});
+
+		bool bFoundCycle = false;
+		if (LastCycleCheckExpr == SourceExpression)
+		{
+			bFoundCycle = Flags[TargetExpression->GetIndex()] == ETraverseRelationShip_Child;
+		}
+		else
+		{
+			bFoundCycle = Flags[SourceExpression->GetIndex()] == ETraverseRelationShip_Parent;
+		}
+
+		if (bFoundCycle)
+		{
+			if (OutFailureReason)
+			{
+				*OutFailureReason = FString(TEXT("Cycles are not allowed."));
+			}
+			return false;
+		}
+	}
+	else
+	{
+		if (SourceBlock->GetRootBlock()->ContainsEntry() != 
+			TargetBlock->GetRootBlock()->ContainsEntry())
+		{
+			return true;
+		}
+
+		if (OutFailureReason)
+		{
+			*OutFailureReason = FString(TEXT("Blocks are not compatible."));
+		}
+		return false;
+	}
+
+	return true;
+}
+
+FString FRigVMParserAST::DumpText() const
+{
+	DECLARE_SCOPE_HIERARCHICAL_COUNTER_FUNC()
+
+	FString Result;
+	for (FRigVMExprAST* RootExpr : RootExpressions)
+	{
+		Result += TEXT("\n") + RootExpr->DumpText();
+	}
+	return Result;
+}
+
+FString FRigVMParserAST::DumpDot() const
+{
+	DECLARE_SCOPE_HIERARCHICAL_COUNTER_FUNC()
+
+	TArray<bool> OutExpressionDefined;
+	OutExpressionDefined.AddZeroed(Expressions.Num());
+
+	FString Result = TEXT("digraph AST {\n  node [style=filled];\n  rankdir=\"RL\";");
+
+	for (FRigVMExprAST* RootExpr : RootExpressions)
+	{
+		Result += RootExpr->DumpDot(OutExpressionDefined, TEXT("  "));
+	}
+	Result += TEXT("\n}");
+	return Result;
+}
+
+void FRigVMParserAST::RemoveExpression(FRigVMExprAST* InExpr, bool bRefreshIndices)
+{
+	TArray<FRigVMExprAST*> Parents = InExpr->Parents;
+	for (FRigVMExprAST* Parent : Parents)
+	{
+		InExpr->RemoveParent(Parent);
+	}
+	TArray<FRigVMExprAST*> Children = InExpr->Children;
+	for (FRigVMExprAST* Child : Children)
+	{
+		Child->RemoveParent(InExpr);
+	}
+
+	Expressions.Remove(InExpr);
+
+	TArray<UObject*> KeysToRemove;
+	for (TPair<UObject*, FRigVMExprAST*> Pair : SubjectToExpression)
+	{
+		if (Pair.Value == InExpr)
+		{
+			KeysToRemove.Add(Pair.Key);
+		}
+	}
+	for (UObject* KeyToRemove : KeysToRemove)
+	{
+		SubjectToExpression.Remove(KeyToRemove);
+	}
+
+	delete InExpr;
+
+	if (bRefreshIndices)
+	{
+		RefreshExprIndices();
+	}
+}
+
+void FRigVMParserAST::RemoveExpressions(TArray<FRigVMExprAST*> InExprs, bool bRefreshIndices)
+{
+	for (FRigVMExprAST* InExpr : InExprs)
+	{
+		RemoveExpression(InExpr, false);
+	}
+	if (bRefreshIndices)
+	{
+		RefreshExprIndices();
+	}
+}
+
+void FRigVMParserAST::TraverseParents(const FRigVMExprAST* InExpr, TFunctionRef<bool(const FRigVMExprAST*)> InContinuePredicate)
+{
+	if (!InContinuePredicate(InExpr))
+	{
+		return;
+	}
+	for (const FRigVMExprAST* ParentExpr : InExpr->Parents)
+	{
+		TraverseParents(ParentExpr, InContinuePredicate);
+	}
+}
+
+void FRigVMParserAST::TraverseChildren(const FRigVMExprAST* InExpr, TFunctionRef<bool(const FRigVMExprAST*)> InContinuePredicate)
+{
+	if (!InContinuePredicate(InExpr))
+	{
+		return;
+	}
+	for (const FRigVMExprAST* ChildExpr : InExpr->Children)
+	{
+		TraverseChildren(ChildExpr, InContinuePredicate);
+	}
+}
+
diff --git a/Engine/Source/Developer/RigVMDeveloper/Private/RigVMCompiler/RigVMCompiler.cpp b/Engine/Source/Developer/RigVMDeveloper/Private/RigVMCompiler/RigVMCompiler.cpp
new file mode 100644
--- /dev/null
+++ b/Engine/Source/Developer/RigVMDeveloper/Private/RigVMCompiler/RigVMCompiler.cpp
@@ -0,0 +1,1060 @@
+// Copyright Epic Games, Inc. All Rights Reserved.
+
+#include "RigVMCompiler/RigVMCompiler.h"
+#include "RigVMModel/RigVMController.h"
+#include "RigVMCore/RigVMExecuteContext.h"
+#include "RigVMDeveloperModule.h"
+#include "UObject/PropertyPortFlags.h"
+#include "Stats/StatsHierarchical.h"
+
+class FRigVMCompilerImportErrorContext : public FOutputDevice
+{
+public:
+
+	URigVMCompiler* Compiler;
+	int32 NumErrors;
+
+	FRigVMCompilerImportErrorContext(URigVMCompiler* InCompiler)
+		: FOutputDevice()
+		, Compiler(InCompiler)
+		, NumErrors(0)
+	{
+	}
+
+	virtual void Serialize(const TCHAR* V, ELogVerbosity::Type Verbosity, const class FName& Category) override
+	{
+		switch (Verbosity)
+		{
+		case ELogVerbosity::Error:
+		case ELogVerbosity::Fatal:
+		{
+			Compiler->ReportError(V);
+			break;
+		}
+		case ELogVerbosity::Warning:
+		{
+			Compiler->ReportWarning(V);
+			break;
+		}
+		default:
+		{
+			Compiler->ReportInfo(V);
+			break;
+		}
+		}
+		NumErrors++;
+	}
+};
+
+FRigVMCompileSettings::FRigVMCompileSettings()
+	: SurpressInfoMessages(true)
+	, SurpressWarnings(false)
+	, SurpressErrors(false)
+	, EnablePinWatches(true)
+	, SplitLiteralsFromWorkMemory(true)
+	, ConsolidateWorkRegisters(true)
+	, ASTSettings(FRigVMParserASTSettings::Optimized())
+{
+}
+
+int32 FRigVMCompilerWorkData::IncRefRegister(int32 InRegister, int32 InIncrement)
+{
+	ensure(InIncrement > 0);
+	if (int32* RefCountPtr = RegisterRefCount.Find(InRegister))
+	{
+		*RefCountPtr += InIncrement;
+		return *RefCountPtr;
+	}
+	RegisterRefCount.Add(InRegister, InIncrement);
+	return InIncrement;
+}
+
+int32 FRigVMCompilerWorkData::DecRefRegister(int32 InRegister, int32 InDecrement)
+{
+	ensure(InDecrement > 0);
+	if (int32* RefCountPtr = RegisterRefCount.Find(InRegister))
+	{
+		int32& RefCount = *RefCountPtr;
+		RefCount -= InDecrement;
+		ensure(RefCount >= 0);
+		return RefCount;
+	}
+	return 0;
+}
+
+URigVMCompiler::URigVMCompiler()
+{
+}
+
+bool URigVMCompiler::Compile(URigVMGraph* InGraph, URigVM* OutVM, const FRigVMUserDataArray& InRigVMUserData, TMap<FString, FRigVMOperand>* OutOperands)
+{
+	if (InGraph == nullptr)
+	{
+		ReportError(TEXT("Provided graph is nullptr."));
+		return false;
+	}
+	if (OutVM == nullptr)
+	{
+		ReportError(TEXT("Provided vm is nullptr."));
+		return false;
+	}
+
+	DECLARE_SCOPE_HIERARCHICAL_COUNTER_FUNC()
+
+	OutVM->Reset();
+
+	for (URigVMNode* Node : InGraph->Nodes)
+	{
+		Node->InstructionIndex = INDEX_NONE;
+	}
+
+	TMap<FString, FRigVMOperand> LocalOperands;
+	if (OutOperands == nullptr)
+	{
+		OutOperands = &LocalOperands;
+	}
+	OutOperands->Reset();
+
+	TSharedPtr<FRigVMParserAST> AST = InGraph->GetRuntimeAST(Settings.ASTSettings, true /* force */);
+	ensure(AST.IsValid());
+
+	FRigVMCompilerWorkData WorkData;
+	WorkData.VM = OutVM;
+	WorkData.PinPathToOperand = OutOperands;
+	WorkData.RigVMUserData = InRigVMUserData;
+	WorkData.bSetupMemory = true;
+	WorkData.NumInstructions = 0;
+
+	// define all parameters independent from sorted nodes
+	for (URigVMNode* Node : InGraph->GetNodes())
+	{
+		if (URigVMParameterNode* ParameterNode = Cast<URigVMParameterNode>(Node))
+		{
+			const FRigVMExprAST* ParameterExpr = AST->GetExprForSubject(ParameterNode);
+			check(ParameterExpr);
+
+			for (const FRigVMExprAST* ChildExpr : *ParameterExpr)
+			{
+				if (ChildExpr->IsA(FRigVMExprAST::EType::Var))
+				{
+					const FRigVMVarExprAST* VarExpr = ChildExpr->To<FRigVMVarExprAST>();
+					if (VarExpr->GetPin()->GetName() == TEXT("Value"))
+					{
+						FindOrAddRegister(VarExpr, WorkData, false /* watchvalue */);
+					}
+				}
+			}
+		}
+	}
+
+	// If a parameter node has no corresponding AST node, because it's not on the execution path,
+	// then add a dummy parameter with no register index, so that we can properly add the pins on the 
+	// ControlRig node later and not lose any existing external connections.
+	for (URigVMNode* Node : InGraph->GetNodes())
+	{
+		if (URigVMParameterNode* ParameterNode = Cast<URigVMParameterNode>(Node))
+		{
+			const FRigVMExprAST* ParameterExpr = AST->GetExprForSubject(ParameterNode);
+			check(ParameterExpr);
+
+			if (!ParameterExpr || ParameterExpr->IsA(FRigVMExprAST::NoOp))
+			{
+				FName Name = ParameterNode->GetParameterName();
+
+				if (!WorkData.VM->ParametersNameMap.Contains(Name))
+				{
+					ERigVMParameterType ParameterType = ParameterNode->IsInput() ? ERigVMParameterType::Input : ERigVMParameterType::Output;
+					FRigVMParameter Parameter(ParameterType, Name, INDEX_NONE, ParameterNode->GetCPPType(), nullptr);
+					WorkData.VM->ParametersNameMap.Add(Parameter.Name, WorkData.VM->Parameters.Add(Parameter));
+				}
+			}
+		}
+	}
+
+	WorkData.ExprComplete.Reset();
+	for (FRigVMExprAST* RootExpr : *AST)
+	{
+		TraverseExpression(RootExpr, WorkData);
+	}
+
+	WorkData.bSetupMemory = false;
+	WorkData.ExprComplete.Reset();
+	for (FRigVMExprAST* RootExpr : *AST)
+	{
+		TraverseExpression(RootExpr, WorkData);
+	}
+
+	if (WorkData.VM->ByteCode.GetInstructions().Num() == 0)
+	{
+		WorkData.VM->ByteCode.AddExitOp();
+		WorkData.NumInstructions++;
+	}
+
+	return true;
+}
+
+void URigVMCompiler::TraverseExpression(const FRigVMExprAST* InExpr, FRigVMCompilerWorkData& WorkData)
+{
+	if (WorkData.ExprComplete.Contains(InExpr))
+	{
+		return;
+	}
+	WorkData.ExprComplete.Add(InExpr, true);
+
+	switch (InExpr->GetType())
+	{
+		case FRigVMExprAST::EType::Block:
+		{
+			TraverseBlock(InExpr->To<FRigVMBlockExprAST>(), WorkData);
+			break;
+		}
+		case FRigVMExprAST::EType::Entry:
+		{
+			TraverseEntry(InExpr->To<FRigVMEntryExprAST>(), WorkData);
+			break;
+		}
+		case FRigVMExprAST::EType::CallExtern:
+		{
+			TraverseCallExtern(InExpr->To<FRigVMCallExternExprAST>(), WorkData);
+			break;
+		}
+		case FRigVMExprAST::EType::NoOp:
+		{
+			TraverseNoOp(InExpr->To<FRigVMNoOpExprAST>(), WorkData);
+			break;
+		}
+		case FRigVMExprAST::EType::Var:
+		{
+			TraverseVar(InExpr->To<FRigVMVarExprAST>(), WorkData);
+			break;
+		}
+		case FRigVMExprAST::EType::Literal:
+		{
+			TraverseLiteral(InExpr->To<FRigVMLiteralExprAST>(), WorkData);
+			break;
+		}
+		case FRigVMExprAST::EType::Assign:
+		{
+			TraverseAssign(InExpr->To<FRigVMAssignExprAST>(), WorkData);
+			break;
+		}
+		case FRigVMExprAST::EType::Copy:
+		{
+			TraverseCopy(InExpr->To<FRigVMCopyExprAST>(), WorkData);
+			break;
+		}
+		case FRigVMExprAST::EType::CachedValue:
+		{
+			TraverseCachedValue(InExpr->To<FRigVMCachedValueExprAST>(), WorkData);
+			break;
+		}
+		case FRigVMExprAST::EType::Exit:
+		{
+			TraverseExit(InExpr->To<FRigVMExitExprAST>(), WorkData);
+			break;
+		}
+		default:
+		{
+			ensure(false);
+			break;
+		}
+	}
+}
+
+void URigVMCompiler::TraverseChildren(const FRigVMExprAST* InExpr, FRigVMCompilerWorkData& WorkData)
+{
+	for (FRigVMExprAST* ChildExpr : *InExpr)
+	{
+		TraverseExpression(ChildExpr, WorkData);
+	}
+}
+
+void URigVMCompiler::TraverseBlock(const FRigVMBlockExprAST* InExpr, FRigVMCompilerWorkData& WorkData)
+{
+	// root blocks are used to store unused expressions
+	if (InExpr->GetParent() == nullptr)
+	{
+		return;
+	}
+	TraverseChildren(InExpr, WorkData);
+}
+
+void URigVMCompiler::TraverseEntry(const FRigVMEntryExprAST* InExpr, FRigVMCompilerWorkData& WorkData)
+{
+	URigVMStructNode* StructNode = Cast<URigVMStructNode>(InExpr->GetNode());
+
+	if (WorkData.bSetupMemory)
+	{
+		TArray<uint8, TAlignedHeapAllocator<16>> TempBuffer;
+		FRigVMStruct* DefaultStruct = nullptr;
+		if (UScriptStruct* Struct = StructNode->GetScriptStruct())
+		{
+			TempBuffer.AddUninitialized(Struct->GetStructureSize());
+			Struct->InitializeDefaultValue(TempBuffer.GetData());
+			DefaultStruct = (FRigVMStruct*)TempBuffer.GetData();
+		}
+
+		WorkData.DefaultStructs.Add(DefaultStruct);
+		TraverseChildren(InExpr, WorkData);
+		WorkData.DefaultStructs.Pop();
+
+		if (UScriptStruct* Struct = StructNode->GetScriptStruct())
+		{
+			Struct->DestroyStruct(TempBuffer.GetData(), 1);
+		}
+	}
+	else
+	{
+		// todo: define the entry in the VM
+		TArray<FRigVMOperand> Operands;
+		for (FRigVMExprAST* ChildExpr : *InExpr)
+		{
+			if (ChildExpr->IsA(FRigVMExprAST::EType::Var))
+			{
+				Operands.Add(WorkData.ExprToOperand.FindChecked(ChildExpr->To<FRigVMVarExprAST>()));
+			}
+			else
+			{
+				break;
+			}
+		}
+
+		// setup the instruction
+		int32 FunctionIndex = WorkData.VM->AddRigVMFunction(StructNode->GetScriptStruct(), StructNode->GetMethodName());
+		WorkData.VM->ByteCode.AddExecuteOp(FunctionIndex, Operands);
+		StructNode->InstructionIndex = WorkData.NumInstructions++;
+
+		TraverseChildren(InExpr, WorkData);
+	}
+}
+
+void URigVMCompiler::TraverseCallExtern(const FRigVMCallExternExprAST* InExpr, FRigVMCompilerWorkData& WorkData)
+{
+	URigVMStructNode* StructNode = Cast<URigVMStructNode>(InExpr->GetNode());
+
+	if (WorkData.bSetupMemory)
+	{
+		TArray<uint8, TAlignedHeapAllocator<16>> TempBuffer;
+		FRigVMStruct* DefaultStruct = nullptr;
+		if (UScriptStruct* Struct = StructNode->GetScriptStruct())
+		{
+			TempBuffer.AddUninitialized(Struct->GetStructureSize());
+			Struct->InitializeDefaultValue(TempBuffer.GetData());
+			DefaultStruct = (FRigVMStruct*)TempBuffer.GetData();
+		}
+
+		WorkData.DefaultStructs.Add(DefaultStruct);
+		TraverseChildren(InExpr, WorkData);
+		WorkData.DefaultStructs.Pop();
+
+		if (Settings.EnablePinWatches)
+		{
+			for (URigVMPin* Pin : StructNode->Pins)
+			{
+				if (Pin->RequiresWatch())
+				{
+					FRigVMVarExprAST TempVarExpr(nullptr, Pin);
+					FindOrAddRegister(&TempVarExpr, WorkData, true);
+				}
+			}
+		}
+
+		if (Settings.ConsolidateWorkRegisters)
+		{
+			for (const FRigVMExprAST* Child : *InExpr)
+			{
+				const FRigVMExprAST* VarExpr = Child;
+				if (VarExpr->IsA(FRigVMExprAST::EType::CachedValue))
+				{
+					VarExpr = VarExpr->To<FRigVMCachedValueExprAST>()->GetVarExpr();
+				}
+				if (VarExpr->IsA(FRigVMExprAST::EType::Var))
+				{
+					if (const FRigVMOperand* Operand = WorkData.ExprToOperand.Find(VarExpr->To<FRigVMVarExprAST>()))
+					{
+						if (Operand->GetMemoryType() != ERigVMMemoryType::Literal)
+						{
+							WorkData.DecRefRegister(Operand->GetRegisterIndex());
+						}
+					}
+				}
+			}
+		}
+
+		if (UScriptStruct* Struct = StructNode->GetScriptStruct())
+		{
+			Struct->DestroyStruct(TempBuffer.GetData(), 1);
+		}
+	}
+	else
+	{
+		TArray<FRigVMOperand> Operands;
+		for (FRigVMExprAST* ChildExpr : *InExpr)
+		{
+			if (ChildExpr->GetType() == FRigVMExprAST::EType::CachedValue)
+			{
+				Operands.Add(WorkData.ExprToOperand.FindChecked(ChildExpr->To<FRigVMCachedValueExprAST>()->GetVarExpr()));
+			}
+			else if (ChildExpr->IsA(FRigVMExprAST::EType::Var))
+			{
+				Operands.Add(WorkData.ExprToOperand.FindChecked(ChildExpr->To<FRigVMVarExprAST>()));
+			}
+			else
+			{
+				break;
+			}
+		}
+
+		TraverseChildren(InExpr, WorkData);
+
+		// setup the instruction
+		int32 FunctionIndex = WorkData.VM->AddRigVMFunction(StructNode->GetScriptStruct(), StructNode->GetMethodName());
+		WorkData.VM->ByteCode.AddExecuteOp(FunctionIndex, Operands);
+		StructNode->InstructionIndex = WorkData.NumInstructions++;
+
+		ensure(InExpr->NumChildren() == StructNode->Pins.Num());
+		for (int32 PinIndex = 0; PinIndex < StructNode->Pins.Num(); PinIndex++)
+		{
+			URigVMPin* Pin = StructNode->Pins[PinIndex];
+
+			// ensure to copy the debug values
+			if (Pin->RequiresWatch() && Settings.EnablePinWatches)
+			{
+				const FRigVMExprAST* PinExpr = InExpr->ChildAt(PinIndex);
+				if (PinExpr->IsA(FRigVMExprAST::EType::CachedValue))
+				{
+					PinExpr = PinExpr->To<FRigVMCachedValueExprAST>()->GetVarExpr();
+				}
+				ensure(PinExpr->IsA(FRigVMExprAST::EType::Var));
+				FRigVMOperand Source = WorkData.ExprToOperand.FindChecked(PinExpr->To<FRigVMVarExprAST>());
+				FString PinHash = GetPinHash(Pin, PinExpr->To<FRigVMVarExprAST>(), true);
+				if (const FRigVMOperand* Target = WorkData.PinPathToOperand->Find(PinHash))
+				{
+					WorkData.VM->ByteCode.AddCopyOp(Source, *Target);
+					WorkData.NumInstructions++;
+				}
+			}
+		}
+	}
+}
+
+void URigVMCompiler::TraverseNoOp(const FRigVMNoOpExprAST* InExpr, FRigVMCompilerWorkData& WorkData)
+{
+	TraverseChildren(InExpr, WorkData);
+}
+
+void URigVMCompiler::TraverseVar(const FRigVMVarExprAST* InExpr, FRigVMCompilerWorkData& WorkData)
+{
+	TraverseChildren(InExpr, WorkData);
+
+	if (WorkData.bSetupMemory)
+	{
+		FindOrAddRegister(InExpr, WorkData);
+	}
+}
+
+void URigVMCompiler::TraverseLiteral(const FRigVMVarExprAST* InExpr, FRigVMCompilerWorkData& WorkData)
+{
+	TraverseVar(InExpr, WorkData);
+}
+
+void URigVMCompiler::TraverseAssign(const FRigVMAssignExprAST* InExpr, FRigVMCompilerWorkData& WorkData)
+{
+	TraverseChildren(InExpr, WorkData);
+
+	ensure(InExpr->NumChildren() > 0);
+
+	const FRigVMVarExprAST* SourceExpr = nullptr;
+
+	const FRigVMExprAST* ChildExpr = InExpr->ChildAt(0);
+	if (ChildExpr->IsA(FRigVMExprAST::EType::Var))
+	{
+		SourceExpr = ChildExpr->To<FRigVMVarExprAST>();
+	}
+	else if (ChildExpr->GetType() == FRigVMExprAST::EType::CachedValue)
+	{
+		SourceExpr = ChildExpr->To<FRigVMCachedValueExprAST>()->GetVarExpr();
+	}
+	else if (ChildExpr->GetType() == FRigVMExprAST::EType::NoOp)
+	{
+		ensure(ChildExpr->NumChildren() > 0);
+
+		for (FRigVMExprAST* GrandChild : *ChildExpr)
+		{
+			if (GrandChild->IsA(FRigVMExprAST::EType::Var))
+			{
+				const FRigVMVarExprAST* VarExpr = GrandChild->To<FRigVMVarExprAST>();
+				if (VarExpr->GetPin()->GetName() == TEXT("Value"))
+				{
+					SourceExpr = VarExpr;
+					break;
+				}
+			}
+		}
+
+		check(SourceExpr);
+	}
+	else
+	{
+		checkNoEntry();
+	}
+
+	FRigVMOperand Source = WorkData.ExprToOperand.FindChecked(SourceExpr);
+
+	if (WorkData.bSetupMemory)
+	{
+		if (Settings.ConsolidateWorkRegisters)
+		{
+			if (Source.GetMemoryType() != ERigVMMemoryType::Literal)
+			{
+				WorkData.DecRefRegister(Source.GetRegisterIndex());
+			}
+		}
+	}
+	else
+	{
+		const FRigVMVarExprAST* TargetExpr = InExpr->GetParent()->To<FRigVMVarExprAST>();
+		FRigVMOperand Target = WorkData.ExprToOperand.FindChecked(TargetExpr);
+
+		// if this is a copy - we should check if operands need offsets
+		if (InExpr->GetType() == FRigVMExprAST::EType::Copy)
+		{
+			struct Local
+			{
+				static void SetupRegisterOffset(URigVM* VM, URigVMPin* Pin, FRigVMOperand& Operand)
+				{
+					URigVMPin* RootPin = Pin->GetRootPin();
+					if (Pin == RootPin)
+					{
+						return;
+					}
+
+					FRigVMMemoryContainer& Memory = Operand.GetMemoryType() == ERigVMMemoryType::Literal ? VM->LiteralMemory : VM->WorkMemory;
+					FString SegmentPath = Pin->GetSegmentPath();
+
+					int32 ArrayIndex = INDEX_NONE;
+					if (RootPin->IsArray())
+					{
+						FString SegmentArrayIndex = SegmentPath;
+						int32 DotIndex = INDEX_NONE;
+						if (SegmentPath.FindChar('.', DotIndex))
+						{
+							SegmentArrayIndex = SegmentPath.Left(DotIndex);
+							SegmentPath = SegmentPath.Mid(DotIndex + 1);
+						}
+						else
+						{
+							SegmentPath = FString();
+						}
+						ArrayIndex = FCString::Atoi(*SegmentArrayIndex);
+					}
+
+					Operand = Memory.GetOperand(Operand.GetRegisterIndex(), SegmentPath, ArrayIndex);
+				}
+			};
+
+			Local::SetupRegisterOffset(WorkData.VM, InExpr->GetSourcePin(), Source);
+			Local::SetupRegisterOffset(WorkData.VM, InExpr->GetTargetPin(), Target);
+		}
+
+		WorkData.VM->ByteCode.AddCopyOp(Source, Target);
+		WorkData.NumInstructions++;
+	}
+}
+
+void URigVMCompiler::TraverseCopy(const FRigVMCopyExprAST* InExpr, FRigVMCompilerWorkData& WorkData)
+{
+	TraverseAssign(InExpr->To<FRigVMAssignExprAST>(), WorkData);
+}
+
+void URigVMCompiler::TraverseCachedValue(const FRigVMCachedValueExprAST* InExpr, FRigVMCompilerWorkData& WorkData)
+{
+	TraverseChildren(InExpr, WorkData);
+}
+
+void URigVMCompiler::TraverseExit(const FRigVMExitExprAST* InExpr, FRigVMCompilerWorkData& WorkData)
+{
+	ensure(InExpr->NumChildren() == 0);
+	if (!WorkData.bSetupMemory)
+	{
+		WorkData.VM->ByteCode.AddExitOp();
+		WorkData.NumInstructions++;
+	}
+}
+
+FString URigVMCompiler::GetPinHash(URigVMPin* InPin, const FRigVMVarExprAST* InVarExpr, bool bIsDebugValue)
+{
+	FString Prefix = bIsDebugValue ? TEXT("DEBUG::") : TEXT("");
+	FString Suffix;
+
+	if (InPin->IsExecuteContext())
+	{
+		return TEXT("ExecuteContext!");
+	}
+
+	if (InVarExpr != nullptr && !bIsDebugValue)
+	{
+		bool bIsExecutePin = InPin->IsExecuteContext();
+		bool bIsLiteral = InVarExpr->GetType() == FRigVMExprAST::EType::Literal;
+
+		// determine if this is an initialization for an IO pin
+		if (!bIsLiteral && !bIsExecutePin && (InPin->GetDirection() == ERigVMPinDirection::IO ||
+			(InPin->GetDirection() == ERigVMPinDirection::Input && InPin->GetSourceLinks().Num() == 0)))
+		{
+			Suffix = TEXT("::IO");
+		}
+	}
+
+	URigVMNode* Node = InPin->GetNode();
+	if (URigVMParameterNode* ParameterNode = Cast<URigVMParameterNode>(Node))
+	{
+		if (InPin->GetName() == TEXT("Value"))
+		{
+			return FString::Printf(TEXT("%sParameter::%s%s"), *Prefix, *ParameterNode->GetParameterName().ToString(), *Suffix);
+		}
+	}
+	else if (URigVMVariableNode* VariableNode = Cast<URigVMVariableNode>(Node))
+	{
+		if (InPin->GetName() == TEXT("Value"))
+		{
+			return FString::Printf(TEXT("%sVariable::%s%s"), *Prefix, *VariableNode->GetVariableName().ToString(), *Suffix);
+		}
+	}
+	return FString::Printf(TEXT("%s%s%s"), *Prefix, *InPin->GetPinPath(), *Suffix);
+}
+
+UScriptStruct* URigVMCompiler::GetScriptStructForCPPType(const FString& InCPPType)
+{
+	if (InCPPType == TEXT("FRotator"))
+	{
+		return TBaseStructure<FRotator>::Get();
+	}
+	if (InCPPType == TEXT("FQuat"))
+	{
+		return TBaseStructure<FQuat>::Get();
+	}
+	if (InCPPType == TEXT("FTransform"))
+	{
+		return TBaseStructure<FTransform>::Get();
+	}
+	if (InCPPType == TEXT("FLinearColor"))
+	{
+		return TBaseStructure<FLinearColor>::Get();
+	}
+	if (InCPPType == TEXT("FColor"))
+	{
+		return TBaseStructure<FColor>::Get();
+	}
+	if (InCPPType == TEXT("FPlane"))
+	{
+		return TBaseStructure<FPlane>::Get();
+	}
+	if (InCPPType == TEXT("FVector"))
+	{
+		return TBaseStructure<FVector>::Get();
+	}
+	if (InCPPType == TEXT("FVector2D"))
+	{
+		return TBaseStructure<FVector2D>::Get();
+	}
+	if (InCPPType == TEXT("FVector4"))
+	{
+		return TBaseStructure<FVector4>::Get();
+	}
+	return nullptr;
+}
+
+TArray<URigVMPin*> URigVMCompiler::GetLinkedPins(URigVMPin* InPin, bool bInputs, bool bOutputs, bool bRecursive)
+{
+	TArray<URigVMPin*> LinkedPins;
+	for (URigVMLink* Link : InPin->GetLinks())
+	{
+		if (bInputs && Link->GetTargetPin() == InPin)
+		{
+			LinkedPins.Add(Link->GetSourcePin());
+		}
+		else if (bOutputs && Link->GetSourcePin() == InPin)
+		{
+			LinkedPins.Add(Link->GetTargetPin());
+		}
+	}
+
+	if (bRecursive)
+	{
+		for (URigVMPin* SubPin : InPin->GetSubPins())
+		{
+			LinkedPins.Append(GetLinkedPins(SubPin, bInputs, bOutputs, bRecursive));
+		}
+	}
+
+	return LinkedPins;
+}
+
+uint16 URigVMCompiler::GetElementSizeFromCPPType(const FString& InCPPType, UScriptStruct* InScriptStruct)
+{
+	if (InScriptStruct == nullptr)
+	{
+		InScriptStruct = GetScriptStructForCPPType(InCPPType);
+	}
+	if (InScriptStruct != nullptr)
+	{
+		return InScriptStruct->GetStructureSize();
+	}
+	if (InCPPType == TEXT("bool"))
+	{
+		return sizeof(bool);
+	}
+	else if (InCPPType == TEXT("int32"))
+	{
+		return sizeof(int32);
+	}
+	if (InCPPType == TEXT("float"))
+	{
+		return sizeof(float);
+	}
+	if (InCPPType == TEXT("FName"))
+	{
+		return sizeof(FName);
+	}
+	if (InCPPType == TEXT("FString"))
+	{
+		return sizeof(FString);
+	}
+
+	ensure(false);
+	return 0;
+}
+
+FRigVMOperand URigVMCompiler::FindOrAddRegister(const FRigVMVarExprAST* InVarExpr, FRigVMCompilerWorkData& WorkData, bool bIsDebugValue)
+{
+	if (!bIsDebugValue)
+	{
+		FRigVMOperand const* ExistingOperand = WorkData.ExprToOperand.Find(InVarExpr);
+		if (ExistingOperand)
+		{
+			return *ExistingOperand;
+		}
+	}
+
+	URigVMPin* Pin = InVarExpr->GetPin();
+	FString BaseCPPType = Pin->IsArray() ? Pin->GetArrayElementCppType() : Pin->GetCPPType();
+	FString Hash = GetPinHash(Pin, InVarExpr, bIsDebugValue);
+	FRigVMOperand Operand;
+	FString RegisterKey = Hash;
+
+	bool bIsExecutePin = Pin->IsExecuteContext();
+	bool bIsLiteral = InVarExpr->GetType() == FRigVMExprAST::EType::Literal;
+	bool bIsParameter = InVarExpr->IsGraphParameter();
+
+	FRigVMOperand const* ExistingOperand = WorkData.PinPathToOperand->Find(Hash);
+	if (ExistingOperand)
+	{
+		if (!bIsDebugValue)
+		{
+			WorkData.ExprToOperand.Add(InVarExpr, *ExistingOperand);
+		}
+		return *ExistingOperand;
+	}
+
+	// create remaining operands / registers
+	if (!Operand.IsValid())
+	{
+		FName RegisterName = *RegisterKey;
+
+		TArray<FString> DefaultValues;
+		if (Pin->IsArray())
+		{
+			if (Pin->GetDirection() == ERigVMPinDirection::Hidden)
+			{
+				ensure(WorkData.DefaultStructs.Num() > 0);
+				FRigVMStruct* VMStruct = WorkData.DefaultStructs.Last();
+
+				int32 DesiredArraySize = VMStruct->GetMaxArraySize(Pin->GetFName(), WorkData.RigVMUserData);
+				ensure(DesiredArraySize > 0);
+
+				DefaultValues = URigVMController::SplitDefaultValue(Pin->GetDefaultValue());
+
+				if (DefaultValues.Num() != DesiredArraySize)
+				{
+					FString DefaultValue;
+					if (Pin->GetArraySize() > 0)
+					{
+						DefaultValue = Pin->GetSubPins()[0]->GetDefaultValue();
+					}
+
+					DefaultValues.Reset();
+					for (int32 Index = 0; Index < DesiredArraySize; Index++)
+					{
+						DefaultValues.Add(DefaultValue);
+					}
+				}
+			}
+			else
+			{
+				DefaultValues = URigVMController::SplitDefaultValue(Pin->GetDefaultValue());
+			}
+
+			while (DefaultValues.Num() < Pin->GetSubPins().Num())
+			{
+				DefaultValues.Add(FString());
+			}
+		}
+		else
+		{
+			DefaultValues.Add(Pin->GetDefaultValue());
+		}
+
+		UScriptStruct* ScriptStruct = Pin->GetScriptStruct();
+		if (ScriptStruct == nullptr)
+		{
+			ScriptStruct = GetScriptStructForCPPType(BaseCPPType);
+		}
+
+		FRigVMMemoryContainer& Memory = (bIsLiteral && Settings.SplitLiteralsFromWorkMemory) ? WorkData.VM->LiteralMemory : WorkData.VM->WorkMemory;
+
+		// look for a register to reuse
+		if (!bIsLiteral && !bIsDebugValue && Settings.ConsolidateWorkRegisters &&
+			(Pin->GetDirection() == ERigVMPinDirection::Output || Pin->GetDirection() == ERigVMPinDirection::IO) &&
+			!bIsExecutePin)
+		{
+			for (TPair<int32, int32> Pair : WorkData.RegisterRefCount)
+			{
+				int32 ExistingRegisterIndex = Pair.Key;
+				if (Pair.Value > 0)
+				{
+					continue;
+				}
+
+				if (Pin->IsArray())
+				{
+					if (!Memory[ExistingRegisterIndex].IsArray())
+					{
+						continue;
+					}
+
+					if (Pin->GetSubPins().Num() != Memory[ExistingRegisterIndex].ElementCount)
+					{
+						continue;
+					}
+				}
+
+				if (ScriptStruct)
+				{
+					if (Memory[ExistingRegisterIndex].ScriptStructIndex == INDEX_NONE)
+					{
+						continue;
+					}
+					if (Memory.GetScriptStruct(ExistingRegisterIndex) != ScriptStruct)
+					{
+						continue;
+					}
+				}
+				else if (Pin->GetCPPType() == TEXT("FString"))
+				{
+					if (Memory[ExistingRegisterIndex].Type != ERigVMRegisterType::String)
+					{
+						continue;
+					}
+				}
+				else if (Pin->GetCPPType() == TEXT("FName"))
+				{
+					if (Memory[ExistingRegisterIndex].Type != ERigVMRegisterType::Name)
+					{
+						continue;
+					}
+				}
+				else
+				{
+					if (Memory[ExistingRegisterIndex].Type != ERigVMRegisterType::Plain)
+					{
+						continue;
+					}
+
+					if (Pin->GetCPPType() == TEXT("bool"))
+					{
+						if (Memory[ExistingRegisterIndex].ElementSize != sizeof(bool))
+						{
+							continue;
+						}
+					}
+					else if (Pin->GetCPPType() == TEXT("int32"))
+					{
+						if (Memory[ExistingRegisterIndex].ElementSize != sizeof(int32))
+						{
+							continue;
+						}
+					}
+					else if (Pin->GetCPPType() == TEXT("float"))
+					{
+						if (Memory[ExistingRegisterIndex].ElementSize != sizeof(float))
+						{
+							continue;
+						}
+					}
+					else if (UEnum* Enum = Pin->GetEnum())
+					{
+						if (Memory[ExistingRegisterIndex].ElementSize != sizeof(uint8))
+						{
+							continue;
+						}
+					}
+					else
+					{
+						ensure(false);
+					}
+				}
+
+				WorkData.IncRefRegister(ExistingRegisterIndex, 1 + Pin->GetLinkedTargetPins(true).Num());
+				Operand = Memory.GetOperand(ExistingRegisterIndex);
+				break;
+			}
+		}
+
+		if (!Operand.IsValid())
+		{
+			if (ScriptStruct)
+			{
+				TArray<uint8, TAlignedHeapAllocator<16>> Data;
+				Data.AddUninitialized(ScriptStruct->GetStructureSize() * DefaultValues.Num());
+				uint8* Ptr = Data.GetData();
+
+				for (FString DefaultValue : DefaultValues)
+				{
+					ScriptStruct->InitializeStruct(Ptr, 1);
+					if (!DefaultValue.IsEmpty() && DefaultValue != TEXT("()"))
+					{
+						FRigVMCompilerImportErrorContext ErrorPipe(this);
+						ScriptStruct->ImportText(*DefaultValue, Ptr, nullptr, PPF_None, &ErrorPipe, ScriptStruct->GetName());
+					}
+					Ptr = Ptr + ScriptStruct->GetStructureSize();
+				}
+
+				int32 Register = Memory.AddStructArray(RegisterName, ScriptStruct, DefaultValues.Num(), Data.GetData(), 1);
+				ScriptStruct->DestroyStruct(Data.GetData(), DefaultValues.Num());
+
+				Operand = Memory.GetOperand(Register);
+			}
+			else if (UEnum* Enum = Pin->GetEnum())
+			{
+				TArray<uint8> Values;
+				for (FString DefaultValue : DefaultValues)
+				{
+					if (DefaultValue.IsEmpty())
+					{
+						Values.Add((uint8)Enum->GetValueByIndex(0));
+					}
+					else
+					{
+						Values.Add((uint8)Enum->GetValueByNameString(DefaultValue));
+					}
+				}
+				int32 Register = Memory.AddPlainArray<uint8>(RegisterName, Values, 1);
+				Operand = Memory.GetOperand(Register);
+			}
+			else if (BaseCPPType == TEXT("bool"))
+			{
+				TArray<bool> Values;
+				for (FString DefaultValue : DefaultValues)
+				{
+					Values.Add((DefaultValue == TEXT("True")) || (DefaultValue == TEXT("true")) || (DefaultValue == TEXT("1")));
+				}
+				int32 Register = Memory.AddPlainArray<bool>(RegisterName, Values, 1);
+				Operand = Memory.GetOperand(Register);
+			}
+			else if (BaseCPPType == TEXT("int32"))
+			{
+				TArray<int32> Values;
+				for (FString DefaultValue : DefaultValues)
+				{
+					if (DefaultValue.IsEmpty())
+					{
+						Values.Add(0);
+					}
+					else
+					{
+						Values.Add(FCString::Atoi(*DefaultValue));
+					}
+				}
+				int32 Register = Memory.AddPlainArray<int32>(RegisterName, Values, 1);
+				Operand = Memory.GetOperand(Register);
+			}
+			else if (BaseCPPType == TEXT("float"))
+			{
+				TArray<float> Values;
+				for (FString DefaultValue : DefaultValues)
+				{
+					if (DefaultValue.IsEmpty())
+					{
+						Values.Add(0.f);
+					}
+					else
+					{
+						Values.Add(FCString::Atof(*DefaultValue));
+					}
+				}
+				int32 Register = Memory.AddPlainArray<float>(RegisterName, Values, 1);
+				Operand = Memory.GetOperand(Register);
+			}
+			else if (BaseCPPType == TEXT("FName"))
+			{
+				TArray<FName> Values;
+				for (FString DefaultValue : DefaultValues)
+				{
+					Values.Add(*DefaultValue);
+				}
+				int32 Register = Memory.AddPlainArray<FName>(RegisterName, Values, 1);
+				Operand = Memory.GetOperand(Register);
+			}
+			else if (BaseCPPType == TEXT("FString"))
+			{
+				int32 Register = Memory.AddPlainArray<FString>(RegisterName, DefaultValues, 1);
+				Operand = Memory.GetOperand(Register);
+			}
+			else
+			{
+				ensure(false);
+			}
+
+			Memory[Operand.GetRegisterIndex()].bIsArray = Pin->IsArray();
+
+			if (!bIsParameter && !bIsLiteral && !bIsDebugValue && Settings.ConsolidateWorkRegisters && !bIsExecutePin && Pin->GetDirection() != ERigVMPinDirection::Hidden)
+			{
+				WorkData.IncRefRegister(Operand.GetRegisterIndex(), 1 + Pin->GetLinkedTargetPins(true).Num());
+			}
+		}
+
+		if (bIsParameter)
+		{
+			URigVMParameterNode* ParameterNode = Cast<URigVMParameterNode>(Pin->GetNode());
+			check(ParameterNode);
+			FName Name = ParameterNode->GetParameterName();
+			ERigVMParameterType ParameterType = ParameterNode->IsInput() ? ERigVMParameterType::Input : ERigVMParameterType::Output;
+			FRigVMParameter Parameter(ParameterType, Name, Operand.GetRegisterIndex(), Pin->GetCPPType(), ScriptStruct);
+			WorkData.VM->ParametersNameMap.Add(Parameter.Name, WorkData.VM->Parameters.Add(Parameter));
+		}
+	}
+	ensure(Operand.IsValid());
+
+	WorkData.PinPathToOperand->Add(Hash, Operand);
+	WorkData.ExprToOperand.Add(InVarExpr, Operand);
+
+	return Operand;
+}
+
+void URigVMCompiler::ReportInfo(const FString& InMessage)
+{
+	if (!Settings.SurpressInfoMessages)
+	{
+		UE_LOG(LogRigVMDeveloper, Display, TEXT("%s"), *InMessage);
+	}
+}
+
+void URigVMCompiler::ReportWarning(const FString& InMessage)
+{
+	FScriptExceptionHandler::Get().HandleException(ELogVerbosity::Warning, *InMessage, *FString());
+}
+
+void URigVMCompiler::ReportError(const FString& InMessage)
+{
+	FScriptExceptionHandler::Get().HandleException(ELogVerbosity::Error, *InMessage, *FString());
+}
diff --git a/Engine/Source/Developer/RigVMDeveloper/Private/RigVMModel/Nodes/RigVMRerouteNode.cpp b/Engine/Source/Developer/RigVMDeveloper/Private/RigVMModel/Nodes/RigVMRerouteNode.cpp
new file mode 100644
--- /dev/null
+++ b/Engine/Source/Developer/RigVMDeveloper/Private/RigVMModel/Nodes/RigVMRerouteNode.cpp
@@ -0,0 +1,21 @@
+// Copyright Epic Games, Inc. All Rights Reserved.
+
+#include "RigVMModel/Nodes/RigVMRerouteNode.h"
+
+const FString URigVMRerouteNode::RerouteName = TEXT("Reroute");
+const FString URigVMRerouteNode::ValueName = TEXT("Value");
+
+URigVMRerouteNode::URigVMRerouteNode()
+: bShowAsFullNode(true)
+{
+}
+
+bool URigVMRerouteNode::GetShowsAsFullNode() const
+{
+	return bShowAsFullNode;
+}
+
+FLinearColor URigVMRerouteNode::GetNodeColor() const
+{
+	return FLinearColor::White;
+}
diff --git a/Engine/Source/Developer/RigVMDeveloper/Private/RigVMModel/Nodes/RigVMStructNode.cpp b/Engine/Source/Developer/RigVMDeveloper/Private/RigVMModel/Nodes/RigVMStructNode.cpp
new file mode 100644
--- /dev/null
+++ b/Engine/Source/Developer/RigVMDeveloper/Private/RigVMModel/Nodes/RigVMStructNode.cpp
@@ -0,0 +1,109 @@
+// Copyright Epic Games, Inc. All Rights Reserved.
+
+#include "RigVMModel/Nodes/RigVMStructNode.h"
+
+FText URigVMStructNode::GetToolTipText() const
+{
+	if(UScriptStruct* Struct = GetScriptStruct())
+	{
+		return Struct->GetToolTipText();
+	}
+	return URigVMNode::GetToolTipText();
+}
+
+FText URigVMStructNode::GetToolTipTextForPin(const URigVMPin* InPin) const
+{
+	if(UScriptStruct* Struct = GetScriptStruct())
+	{
+		TArray<FString> Parts;
+		URigVMPin::SplitPinPath(InPin->GetPinPath(), Parts);
+
+		for (int32 PartIndex = 1; PartIndex < Parts.Num(); PartIndex++)
+		{
+			FProperty* Property = Struct->FindPropertyByName(*Parts[PartIndex]);
+			if (!Property)
+			{
+				break;
+			}
+
+			if (FArrayProperty* ArrayProperty = CastField<FArrayProperty>(Property))
+			{
+				if (PartIndex < Parts.Num() - 1)
+				{
+					Property = ArrayProperty->Inner;
+					PartIndex++;
+				}
+			}
+
+			if (PartIndex == Parts.Num() - 1)
+			{
+				return Property->GetToolTipText();
+			}
+
+			if (FStructProperty* StructProperty = CastField<FStructProperty>(Property))
+			{
+				Struct = StructProperty->Struct;
+			}
+		}
+
+	}
+	return URigVMNode::GetToolTipTextForPin(InPin);
+}
+
+UScriptStruct* URigVMStructNode::GetScriptStruct() const
+{
+	return ScriptStruct;
+}
+
+FName URigVMStructNode::GetMethodName() const
+{
+	return MethodName;
+}
+
+FString URigVMStructNode::GetStructDefaultValue() const
+{
+	TArray<FString> PinDefaultValues;
+	for (URigVMPin* Pin : GetPins())
+	{
+		if (Pin->GetDirection() == ERigVMPinDirection::Hidden)
+		{
+			continue;
+		}
+		FString PinDefaultValue = Pin->GetDefaultValue();
+		if (Pin->IsStringType())
+		{
+			PinDefaultValue = TEXT("\"") + PinDefaultValue + TEXT("\"");
+		}
+		else if (PinDefaultValue.IsEmpty() || PinDefaultValue == TEXT("()"))
+		{
+			continue;
+		}
+		PinDefaultValues.Add(FString::Printf(TEXT("%s=%s"), *Pin->GetName(), *PinDefaultValue));
+	}
+	if (PinDefaultValues.Num() == 0)
+	{
+		return TEXT("()");
+	}
+	return FString::Printf(TEXT("(%s)"), *FString::Join(PinDefaultValues, TEXT(",")));
+
+}
+
+TSharedPtr<FStructOnScope> URigVMStructNode::ConstructStructInstance(bool bUseDefault) const
+{
+	if (ScriptStruct)
+	{
+		TSharedPtr<FStructOnScope> StructOnScope = MakeShareable(new FStructOnScope(ScriptStruct));
+		FRigVMStruct* StructMemory = (FRigVMStruct*)StructOnScope->GetStructMemory();
+		if (bUseDefault)
+		{
+			ScriptStruct->InitializeDefaultValue((uint8*)StructMemory);
+		}
+		else
+		{
+			FString StructDefaultValue = GetStructDefaultValue();
+			ScriptStruct->ImportText(*StructDefaultValue, StructMemory, nullptr, PPF_None, nullptr, ScriptStruct->GetName());
+		}
+		return StructOnScope;
+	}
+	return nullptr;
+}
