COPY from //Dev-Anim to //Dev-Main

diff --git a/Engine/Source/Developer/RigVMDeveloper/Private/RigVMModel/RigVMController.cpp b/Engine/Source/Developer/RigVMDeveloper/Private/RigVMModel/RigVMController.cpp
new file mode 100644
--- /dev/null
+++ b/Engine/Source/Developer/RigVMDeveloper/Private/RigVMModel/RigVMController.cpp
@@ -0,0 +1,3864 @@
+// Copyright Epic Games, Inc. All Rights Reserved.
+
+#include "RigVMModel/RigVMController.h"
+#include "RigVMModel/RigVMControllerActions.h"
+#include "RigVMCore/RigVMRegistry.h"
+#include "RigVMCore/RigVMExecuteContext.h"
+#include "RigVMCompiler/RigVMCompiler.h"
+#include "RigVMDeveloperModule.h"
+#include "UObject/PropertyPortFlags.h"
+#include "UObject/Package.h"
+#include "Misc/CoreMisc.h"
+
+#if WITH_EDITOR
+#include "Exporters/Exporter.h"
+#include "UnrealExporter.h"
+#include "Factories.h"
+#endif
+
+URigVMController::URigVMController()
+	: bSuspendNotifications(false)
+	, bReportWarningsAndErrors(true)
+{
+	SetExecuteContextStruct(FRigVMExecuteContext::StaticStruct());
+}
+
+URigVMController::URigVMController(const FObjectInitializer& ObjectInitializer)
+	: Super(ObjectInitializer)
+	, bReportWarningsAndErrors(true)
+{
+	ActionStack = CreateDefaultSubobject<URigVMActionStack>(TEXT("ActionStack"));
+	SetExecuteContextStruct(FRigVMExecuteContext::StaticStruct());
+
+	ActionStack->OnModified().AddLambda([&](ERigVMGraphNotifType NotifType, URigVMGraph* InGraph, UObject* InSubject) -> void {
+		Notify(NotifType, InSubject);
+	});
+}
+
+URigVMController::~URigVMController()
+{
+}
+
+URigVMGraph* URigVMController::GetGraph() const
+{
+	return Graph;
+}
+
+void URigVMController::SetGraph(URigVMGraph* InGraph)
+{
+	if (Graph)
+	{
+		Graph->OnModified().RemoveAll(this);
+	}
+
+	Graph = InGraph;
+
+	if (Graph)
+	{
+		Graph->OnModified().AddUObject(this, &URigVMController::HandleModifiedEvent);
+	}
+
+	HandleModifiedEvent(ERigVMGraphNotifType::GraphChanged, Graph, nullptr);
+}
+
+FRigVMGraphModifiedEvent& URigVMController::OnModified()
+{
+	return ModifiedEventStatic;
+}
+
+void URigVMController::Notify(ERigVMGraphNotifType InNotifType, UObject* InSubject)
+{
+	if (bSuspendNotifications)
+	{
+		return;
+	}
+	if (Graph)
+	{
+		Graph->Notify(InNotifType, InSubject);
+	}
+}
+
+void URigVMController::ResendAllNotifications()
+{
+	if (Graph)
+	{
+		for (URigVMLink* Link : Graph->Links)
+		{
+			Notify(ERigVMGraphNotifType::LinkRemoved, Link);
+		}
+
+		for (URigVMNode* Node : Graph->Nodes)
+		{
+			Notify(ERigVMGraphNotifType::NodeRemoved, Node);
+		}
+
+		for (URigVMNode* Node : Graph->Nodes)
+		{
+			Notify(ERigVMGraphNotifType::NodeAdded, Node);
+		}
+
+		for (URigVMLink* Link : Graph->Links)
+		{
+			Notify(ERigVMGraphNotifType::LinkAdded, Link);
+		}
+	}
+}
+
+void URigVMController::HandleModifiedEvent(ERigVMGraphNotifType InNotifType, URigVMGraph* InGraph, UObject* InSubject)
+{
+	switch (InNotifType)
+	{
+		case ERigVMGraphNotifType::GraphChanged:
+		case ERigVMGraphNotifType::NodeAdded:
+		case ERigVMGraphNotifType::NodeRemoved:
+		case ERigVMGraphNotifType::LinkAdded:
+		case ERigVMGraphNotifType::LinkRemoved:
+		case ERigVMGraphNotifType::PinArraySizeChanged:
+		case ERigVMGraphNotifType::VariableAdded:
+		case ERigVMGraphNotifType::VariableRemoved:
+		case ERigVMGraphNotifType::ParameterAdded:
+		case ERigVMGraphNotifType::ParameterRemoved:
+		{
+			if (InGraph)
+			{
+				InGraph->ClearAST();
+			}
+			break;
+		}
+		case ERigVMGraphNotifType::PinDefaultValueChanged:
+		{
+			if (InGraph->RuntimeAST.IsValid())
+			{
+				URigVMPin* RootPin = CastChecked<URigVMPin>(InSubject)->GetRootPin();
+				const FRigVMExprAST* Expression = InGraph->GetRuntimeAST()->GetExprForSubject(RootPin);
+				if (Expression == nullptr)
+				{
+					InGraph->ClearAST();
+					break;
+				}
+				else if(Expression->NumParents() > 1)
+				{
+					InGraph->ClearAST();
+					break;
+				}
+			}
+			break;
+		}
+	}
+
+	ModifiedEventStatic.Broadcast(InNotifType, InGraph, InSubject);
+	if (ModifiedEventDynamic.IsBound())
+	{
+		ModifiedEventDynamic.Broadcast(InNotifType, InGraph, InSubject);
+	}
+}
+
+#if WITH_EDITOR
+
+URigVMStructNode* URigVMController::AddStructNode(UScriptStruct* InScriptStruct, const FName& InMethodName, const FVector2D& InPosition, const FString& InNodeName, bool bUndo)
+{
+	if(!IsValidGraph())
+	{
+		return nullptr;
+	}
+	if (InScriptStruct == nullptr)
+	{
+		ReportError(TEXT("InScriptStruct is null."));
+		return nullptr;
+	}
+	if (InMethodName == NAME_None)
+	{
+		ReportError(TEXT("InMethodName is None."));
+		return nullptr;
+	}
+
+	FString FunctionName = FString::Printf(TEXT("F%s::%s"), *InScriptStruct->GetName(), *InMethodName.ToString());
+	FRigVMFunctionPtr Function = FRigVMRegistry::Get().Find(*FunctionName);
+	if (Function == nullptr)
+	{
+		ReportErrorf(TEXT("RIGVM_METHOD '%s' cannot be found."), *FunctionName);
+		return nullptr;
+	}
+
+	FString Name = GetValidNodeName(InNodeName.IsEmpty() ? InScriptStruct->GetName() : InNodeName);
+	URigVMStructNode* Node = NewObject<URigVMStructNode>(Graph, *Name);
+	Node->ScriptStruct = InScriptStruct;
+	Node->MethodName = InMethodName;
+	Node->Position = InPosition;
+	Node->NodeTitle = InScriptStruct->GetMetaData(TEXT("DisplayName"));
+	
+	FString NodeColorMetadata;
+	InScriptStruct->GetStringMetaDataHierarchical(*URigVMNode::NodeColorName, &NodeColorMetadata);
+	if (!NodeColorMetadata.IsEmpty())
+	{
+		Node->NodeColor = GetColorFromMetadata(NodeColorMetadata);
+	}
+
+	FString ExportedDefaultValue;
+	CreateDefaultValueForStructIfRequired(InScriptStruct, ExportedDefaultValue);
+	AddPinsForStruct(InScriptStruct, Node, nullptr, ERigVMPinDirection::Invalid, ExportedDefaultValue, true);
+
+	Graph->Nodes.Add(Node);
+	Graph->MarkPackageDirty();
+
+	FRigVMAddStructNodeAction Action;
+	if (bUndo)
+	{
+		Action = FRigVMAddStructNodeAction(Node);
+		Action.Title = FString::Printf(TEXT("Add %s Node"), *Node->GetNodeTitle());
+		ActionStack->BeginAction(Action);
+	}
+
+	Notify(ERigVMGraphNotifType::NodeAdded, Node);
+
+	if (bUndo)
+	{
+		ActionStack->EndAction(Action);
+	}
+
+	return Node;
+}
+
+URigVMStructNode* URigVMController::AddStructNodeFromStructPath(const FString& InScriptStructPath, const FName& InMethodName, const FVector2D& InPosition, const FString& InNodeName, bool bUndo)
+{
+	if(!IsValidGraph())
+	{
+		return nullptr;
+	}
+
+	UScriptStruct* ScriptStruct = URigVMPin::FindObjectFromCPPTypeObjectPath<UScriptStruct>(InScriptStructPath);
+	if (ScriptStruct == nullptr)
+	{
+		ReportErrorf(TEXT("Cannot find struct for path '%s'."), *InScriptStructPath);
+		return nullptr;
+	}
+
+	return AddStructNode(ScriptStruct, InMethodName, InPosition, InNodeName, bUndo);
+}
+
+URigVMVariableNode* URigVMController::AddVariableNode(const FName& InVariableName, const FString& InCPPType, UObject* InCPPTypeObject, bool bIsGetter, const FString& InDefaultValue, const FVector2D& InPosition, const FString& InNodeName, bool bUndo)
+{
+	if (!IsValidGraph())
+	{
+		return nullptr;
+	}
+
+	if (InCPPTypeObject == nullptr)
+	{
+		InCPPTypeObject = URigVMCompiler::GetScriptStructForCPPType(InCPPType);
+	}
+	if (InCPPTypeObject == nullptr)
+	{
+		InCPPTypeObject = URigVMPin::FindObjectFromCPPTypeObjectPath<UObject>(InCPPType);
+	}
+
+	TArray<FRigVMGraphVariableDescription> ExistingVariables = Graph->GetVariableDescriptions();
+	for (const FRigVMGraphVariableDescription& ExistingVariable : ExistingVariables)
+	{
+		if (ExistingVariable.Name == InVariableName)
+		{
+			if (ExistingVariable.CPPType != InCPPType ||
+				ExistingVariable.CPPTypeObject != InCPPTypeObject)
+			{
+				ReportErrorf(TEXT("Cannot add variable '%s' - variable already exists."), *InVariableName.ToString());
+				return nullptr;
+			}
+		}
+	}
+
+	FString Name = GetValidNodeName(InNodeName.IsEmpty() ? FString(TEXT("VariableNode")) : InNodeName);
+	URigVMVariableNode* Node = NewObject<URigVMVariableNode>(Graph, *Name);
+	Node->Position = InPosition;
+
+	if (!bIsGetter)
+	{
+		URigVMPin* ExecutePin = NewObject<URigVMPin>(Node, *URigVMNode::ExecuteContextName);
+		ExecutePin->CPPType = FString::Printf(TEXT("F%s"), *ExecuteContextStruct->GetName());
+		ExecutePin->CPPTypeObject = ExecuteContextStruct;
+		ExecutePin->CPPTypeObjectPath = *ExecutePin->CPPTypeObject->GetPathName();
+		ExecutePin->Direction = ERigVMPinDirection::IO;
+		Node->Pins.Add(ExecutePin);
+	}
+
+	URigVMPin* VariablePin = NewObject<URigVMPin>(Node, *URigVMVariableNode::VariableName);
+	VariablePin->CPPType = TEXT("FName");
+	VariablePin->Direction = ERigVMPinDirection::Visible;
+	VariablePin->DefaultValue = InVariableName.ToString();
+	VariablePin->CustomWidgetName = TEXT("VariableName");
+	Node->Pins.Add(VariablePin);
+
+	URigVMPin* ValuePin = NewObject<URigVMPin>(Node, *URigVMVariableNode::ValueName);
+	ValuePin->CPPType = InCPPType;
+
+	if (UScriptStruct* ScriptStruct = Cast<UScriptStruct>(InCPPTypeObject))
+	{
+		ValuePin->CPPTypeObject = ScriptStruct;
+		ValuePin->CPPTypeObjectPath = *ValuePin->CPPTypeObject->GetPathName();
+	}
+	else if (UEnum* Enum = Cast<UEnum>(InCPPTypeObject))
+	{
+		ValuePin->CPPTypeObject = Enum;
+		ValuePin->CPPTypeObjectPath = *ValuePin->CPPTypeObject->GetPathName();
+	}
+
+	ValuePin->Direction = bIsGetter ? ERigVMPinDirection::Output : ERigVMPinDirection::Input;
+	Node->Pins.Add(ValuePin);
+
+	Graph->Nodes.Add(Node);
+
+	if (ValuePin->IsStruct())
+	{
+		FString DefaultValue = InDefaultValue;
+		CreateDefaultValueForStructIfRequired(ValuePin->GetScriptStruct(), DefaultValue);
+		AddPinsForStruct(ValuePin->GetScriptStruct(), Node, ValuePin, ValuePin->Direction, DefaultValue, false);
+	}
+	else if (!InDefaultValue.IsEmpty() && InDefaultValue != TEXT("()"))
+	{
+		SetPinDefaultValue(ValuePin, InDefaultValue, true, false, false);
+	}
+
+	ForEveryPinRecursively(Node, [](URigVMPin* Pin) {
+		Pin->bIsExpanded = false;
+	});
+
+	Graph->MarkPackageDirty();
+
+	FRigVMAddVariableNodeAction Action;
+	if (bUndo)
+	{
+		Action = FRigVMAddVariableNodeAction(Node);
+		Action.Title = FString::Printf(TEXT("Add %s Variable"), *InVariableName.ToString());
+		ActionStack->BeginAction(Action);
+	}
+
+	Notify(ERigVMGraphNotifType::NodeAdded, Node);
+	Notify(ERigVMGraphNotifType::VariableAdded, Node);
+
+	if (bUndo)
+	{
+		ActionStack->EndAction(Action);
+	}
+
+	return Node;
+}
+
+URigVMVariableNode* URigVMController::AddVariableNodeFromObjectPath(const FName& InVariableName, const FString& InCPPType, const FString& InCPPTypeObjectPath, bool bIsGetter, const FString& InDefaultValue, const FVector2D& InPosition, const FString& InNodeName, bool bUndo)
+{
+	if (!IsValidGraph())
+	{
+		return nullptr;
+	}
+
+	UObject* CPPTypeObject = nullptr;
+	if (!InCPPTypeObjectPath.IsEmpty())
+	{
+		CPPTypeObject = URigVMPin::FindObjectFromCPPTypeObjectPath<UObject>(InCPPTypeObjectPath);
+		if (CPPTypeObject == nullptr)
+		{
+			ReportErrorf(TEXT("Cannot find cpp type object for path '%s'."), *InCPPTypeObjectPath);
+			return nullptr;
+		}
+	}
+
+	return AddVariableNode(InVariableName, InCPPType, CPPTypeObject, bIsGetter, InDefaultValue, InPosition, InNodeName, bUndo);
+}
+
+URigVMParameterNode* URigVMController::AddParameterNode(const FName& InParameterName, const FString& InCPPType, UObject* InCPPTypeObject, bool bIsInput, const FString& InDefaultValue, const FVector2D& InPosition, const FString& InNodeName, bool bUndo)
+{
+	if (!IsValidGraph())
+	{
+		return nullptr;
+	}
+
+	if (InCPPTypeObject == nullptr)
+	{
+		InCPPTypeObject = URigVMCompiler::GetScriptStructForCPPType(InCPPType);
+	}
+	if (InCPPTypeObject == nullptr)
+	{
+		InCPPTypeObject = URigVMPin::FindObjectFromCPPTypeObjectPath<UObject>(InCPPType);
+	}
+
+	TArray<FRigVMGraphParameterDescription> ExistingParameters = Graph->GetParameterDescriptions();
+	for (const FRigVMGraphParameterDescription& ExistingParameter : ExistingParameters)
+	{
+		if (ExistingParameter.Name == InParameterName)
+		{
+			if (ExistingParameter.CPPType != InCPPType ||
+				ExistingParameter.CPPTypeObject != InCPPTypeObject ||
+				ExistingParameter.bIsInput != bIsInput)
+			{
+				ReportErrorf(TEXT("Cannot add parameter '%s' - parameter already exists."), *InParameterName.ToString());
+				return nullptr;
+			}
+		}
+	}
+
+	FString Name = GetValidNodeName(InNodeName.IsEmpty() ? FString(TEXT("ParameterNode")) : InNodeName);
+	URigVMParameterNode* Node = NewObject<URigVMParameterNode>(Graph, *Name);
+	Node->Position = InPosition;
+
+	if (!bIsInput)
+	{
+		URigVMPin* ExecutePin = NewObject<URigVMPin>(Node, *URigVMNode::ExecuteContextName);
+		ExecutePin->CPPType = FString::Printf(TEXT("F%s"), *ExecuteContextStruct->GetName());
+		ExecutePin->CPPTypeObject = ExecuteContextStruct;
+		ExecutePin->CPPTypeObjectPath = *ExecutePin->CPPTypeObject->GetPathName();
+		ExecutePin->Direction = ERigVMPinDirection::IO;
+		Node->Pins.Add(ExecutePin);
+	}
+
+	URigVMPin* ParameterPin = NewObject<URigVMPin>(Node, *URigVMParameterNode::ParameterName);
+	ParameterPin->CPPType = TEXT("FName");
+	ParameterPin->Direction = ERigVMPinDirection::Visible;
+	ParameterPin->DefaultValue = InParameterName.ToString();
+	ParameterPin->CustomWidgetName = TEXT("ParameterName");
+
+	Node->Pins.Add(ParameterPin);
+
+	URigVMPin* DefaultValuePin = nullptr;
+	if (bIsInput)
+	{
+		DefaultValuePin = NewObject<URigVMPin>(Node, *URigVMParameterNode::DefaultName);
+	}
+	URigVMPin* ValuePin = NewObject<URigVMPin>(Node, *URigVMParameterNode::ValueName);
+
+	if (DefaultValuePin)
+	{
+		DefaultValuePin->CPPType = InCPPType;
+	}
+	ValuePin->CPPType = InCPPType;
+
+	if (UScriptStruct* ScriptStruct = Cast<UScriptStruct>(InCPPTypeObject))
+	{
+		if (DefaultValuePin)
+		{
+			DefaultValuePin->CPPTypeObject = ScriptStruct;
+			DefaultValuePin->CPPTypeObjectPath = *DefaultValuePin->CPPTypeObject->GetPathName();
+		}
+		ValuePin->CPPTypeObject = ScriptStruct;
+		ValuePin->CPPTypeObjectPath = *ValuePin->CPPTypeObject->GetPathName();
+	}
+	else if (UEnum* Enum = Cast<UEnum>(InCPPTypeObject))
+	{
+		if (DefaultValuePin)
+		{
+			DefaultValuePin->CPPTypeObject = Enum;
+			DefaultValuePin->CPPTypeObjectPath = *DefaultValuePin->CPPTypeObject->GetPathName();
+		}
+		ValuePin->CPPTypeObject = Enum;
+		ValuePin->CPPTypeObjectPath = *ValuePin->CPPTypeObject->GetPathName();
+	}
+
+	if (DefaultValuePin)
+	{
+		DefaultValuePin->Direction = ERigVMPinDirection::Visible;
+	}
+	ValuePin->Direction = bIsInput ? ERigVMPinDirection::Output : ERigVMPinDirection::Input;
+
+	if (bIsInput)
+	{
+		if (ValuePin->CPPType == TEXT("FName"))
+		{
+			ValuePin->bIsConstant = true;
+		}
+	}
+
+	if (DefaultValuePin)
+	{
+		Node->Pins.Add(DefaultValuePin);
+	}
+	Node->Pins.Add(ValuePin);
+
+	Graph->Nodes.Add(Node);
+
+	if (ValuePin->IsStruct())
+	{
+		FString DefaultValue = InDefaultValue;
+		CreateDefaultValueForStructIfRequired(ValuePin->GetScriptStruct(), DefaultValue);
+		if (DefaultValuePin)
+		{
+			AddPinsForStruct(DefaultValuePin->GetScriptStruct(), Node, DefaultValuePin, DefaultValuePin->Direction, DefaultValue, false);
+		}
+		AddPinsForStruct(ValuePin->GetScriptStruct(), Node, ValuePin, ValuePin->Direction, DefaultValue, false);
+	}
+	else if (!InDefaultValue.IsEmpty() && InDefaultValue != TEXT("()"))
+	{
+		if (DefaultValuePin)
+		{
+			SetPinDefaultValue(DefaultValuePin, InDefaultValue, true, false, false);
+		}
+		SetPinDefaultValue(ValuePin, InDefaultValue, true, false, false);
+	}
+
+	ForEveryPinRecursively(Node, [](URigVMPin* Pin) {
+		Pin->bIsExpanded = false;
+	});
+
+	Graph->MarkPackageDirty();
+
+	FRigVMAddParameterNodeAction Action;
+	if (bUndo)
+	{
+		Action = FRigVMAddParameterNodeAction(Node);
+		Action.Title = FString::Printf(TEXT("Add %s Parameter"), *InParameterName.ToString());
+		ActionStack->BeginAction(Action);
+	}
+
+	Notify(ERigVMGraphNotifType::NodeAdded, Node);
+	Notify(ERigVMGraphNotifType::ParameterAdded, Node);
+
+	if (bUndo)
+	{
+		ActionStack->EndAction(Action);
+	}
+
+	return Node;
+}
+
+URigVMParameterNode* URigVMController::AddParameterNodeFromObjectPath(const FName& InParameterName, const FString& InCPPType, const FString& InCPPTypeObjectPath, bool bIsInput, const FString& InDefaultValue, const FVector2D& InPosition, const FString& InNodeName, bool bUndo)
+{
+	if (!IsValidGraph())
+	{
+		return nullptr;
+	}
+
+	UObject* CPPTypeObject = nullptr;
+	if (!InCPPTypeObjectPath.IsEmpty())
+	{
+		CPPTypeObject = URigVMPin::FindObjectFromCPPTypeObjectPath<UObject>(InCPPTypeObjectPath);
+		if (CPPTypeObject == nullptr)
+		{
+			ReportErrorf(TEXT("Cannot find cpp type object for path '%s'."), *InCPPTypeObjectPath);
+			return nullptr;
+		}
+	}
+
+	return AddParameterNode(InParameterName, InCPPType, CPPTypeObject, bIsInput, InDefaultValue, InPosition, InNodeName, bUndo);
+}
+
+URigVMCommentNode* URigVMController::AddCommentNode(const FString& InCommentText, const FVector2D& InPosition, const FVector2D& InSize, const FLinearColor& InColor, const FString& InNodeName, bool bUndo)
+{
+	if (!IsValidGraph())
+	{
+		return nullptr;
+	}
+
+	FString Name = GetValidNodeName(InNodeName.IsEmpty() ? FString(TEXT("CommentNode")) : InNodeName);
+	URigVMCommentNode* Node = NewObject<URigVMCommentNode>(Graph, *Name);
+	Node->Position = InPosition;
+	Node->Size = InSize;
+	Node->NodeColor = InColor;
+	Node->CommentText = InCommentText;
+
+	Graph->Nodes.Add(Node);
+	Graph->MarkPackageDirty();
+
+	FRigVMAddCommentNodeAction Action;
+	if (bUndo)
+	{
+		Action = FRigVMAddCommentNodeAction(Node);
+		Action.Title = FString::Printf(TEXT("Add Comment"));
+		ActionStack->BeginAction(Action);
+	}
+
+	Notify(ERigVMGraphNotifType::NodeAdded, Node);
+
+	if (bUndo)
+	{
+		ActionStack->EndAction(Action);
+	}
+
+	return Node;
+}
+
+URigVMRerouteNode* URigVMController::AddRerouteNodeOnLink(URigVMLink* InLink, bool bShowAsFullNode, const FVector2D& InPosition, const FString& InNodeName, bool bUndo)
+{
+	if(!IsValidLinkForGraph(InLink))
+	{
+		return nullptr;
+	}
+
+	URigVMPin* SourcePin = InLink->GetSourcePin();
+	URigVMPin* TargetPin = InLink->GetTargetPin();
+
+	FRigVMBaseAction Action;
+	if (bUndo)
+	{
+		Action.Title = FString::Printf(TEXT("Add Reroute"));
+		ActionStack->BeginAction(Action);
+	}
+
+	URigVMRerouteNode* Node = AddRerouteNodeOnPin(TargetPin->GetPinPath(), true, bShowAsFullNode, InPosition, InNodeName, bUndo);
+	if (Node == nullptr)
+	{
+		if (bUndo)
+		{
+			ActionStack->CancelAction(Action);
+		}
+		return nullptr;
+	}
+
+	URigVMPin* ValuePin = Node->Pins[0];
+	AddLink(SourcePin, ValuePin, bUndo);
+
+	if (bUndo)
+	{
+		ActionStack->EndAction(Action);
+	}
+
+	return Node;
+}
+
+URigVMRerouteNode* URigVMController::AddRerouteNodeOnLinkPath(const FString& InLinkPinPathRepresentation, bool bShowAsFullNode, const FVector2D& InPosition, const FString& InNodeName, bool bUndo)
+{
+	if (!IsValidGraph())
+	{
+		return nullptr;
+	}
+
+	URigVMLink* Link = Graph->FindLink(InLinkPinPathRepresentation);
+	return AddRerouteNodeOnLink(Link, bShowAsFullNode, InPosition, InNodeName, bUndo);
+}
+
+URigVMRerouteNode* URigVMController::AddRerouteNodeOnPin(const FString& InPinPath, bool bAsInput, bool bShowAsFullNode, const FVector2D& InPosition, const FString& InNodeName, bool bUndo)
+{
+	if (!IsValidGraph())
+	{
+		return nullptr;
+	}
+
+	URigVMPin* Pin = Graph->FindPin(InPinPath);
+	if(Pin == nullptr)
+	{
+		return nullptr;
+	}
+
+	FRigVMBaseAction Action;
+	if (bUndo)
+	{
+		Action.Title = FString::Printf(TEXT("Add Reroute"));
+		ActionStack->BeginAction(Action);
+	}
+
+	BreakAllLinks(Pin, bAsInput, bUndo);
+
+	FString Name = GetValidNodeName(InNodeName.IsEmpty() ? FString(TEXT("RerouteNode")) : InNodeName);
+	URigVMRerouteNode* Node = NewObject<URigVMRerouteNode>(Graph, *Name);
+	Node->Position = InPosition;
+	Node->bShowAsFullNode = bShowAsFullNode;
+
+	URigVMPin* ValuePin = NewObject<URigVMPin>(Node, *URigVMRerouteNode::ValueName);
+	ConfigurePinFromPin(ValuePin, Pin);
+	ValuePin->Direction = ERigVMPinDirection::IO;
+	Node->Pins.Add(ValuePin);
+
+	if (ValuePin->IsStruct())
+	{
+		AddPinsForStruct(ValuePin->GetScriptStruct(), Node, ValuePin, ValuePin->Direction, FString(), false);
+	}
+
+	FString DefaultValue = Pin->GetDefaultValue();
+	if (!DefaultValue.IsEmpty())
+	{
+		SetPinDefaultValue(ValuePin, Pin->GetDefaultValue(), true, false, false);
+	}
+
+	ForEveryPinRecursively(ValuePin, [](URigVMPin* Pin) {
+		Pin->bIsExpanded = true;
+		if (Pin->GetParentPin() != nullptr)
+		{
+			Pin->Direction = ERigVMPinDirection::Visible;
+		}
+	});
+
+	Graph->Nodes.Add(Node);
+	Graph->MarkPackageDirty();
+
+	if (bUndo)
+	{
+		ActionStack->AddAction(FRigVMAddRerouteNodeAction(Node));
+	}
+
+	Notify(ERigVMGraphNotifType::NodeAdded, Node);
+
+	if (bAsInput)
+	{
+		AddLink(ValuePin, Pin, bUndo);
+	}
+	else
+	{
+		AddLink(Pin, ValuePin, bUndo);
+	}
+
+	if (bUndo)
+	{
+		ActionStack->EndAction(Action);
+	}
+
+	return Node;
+}
+
+URigVMInjectionInfo* URigVMController::AddInjectedNode(const FString& InPinPath, bool bAsInput, UScriptStruct* InScriptStruct, const FName& InMethodName, const FName& InInputPinName, const FName& InOutputPinName, const FString& InNodeName, bool bUndo)
+{
+	if (!IsValidGraph())
+	{
+		return nullptr;
+	}
+
+	URigVMPin* Pin = Graph->FindPin(InPinPath);
+	if (Pin == nullptr)
+	{
+		return nullptr;
+	}
+
+	if (Pin->IsArray())
+	{
+		return nullptr;
+	}
+
+	if (bAsInput && !(Pin->GetDirection() == ERigVMPinDirection::Input || Pin->GetDirection() == ERigVMPinDirection::IO))
+	{
+		ReportError(TEXT("Pin is not an input / cannot add injected input node."));
+		return nullptr;
+	}
+	if (!bAsInput && !(Pin->GetDirection() == ERigVMPinDirection::Output))
+	{
+		ReportError(TEXT("Pin is not an output / cannot add injected output node."));
+		return nullptr;
+	}
+
+	if (InScriptStruct == nullptr)
+	{
+		ReportError(TEXT("InScriptStruct is null."));
+		return nullptr;
+	}
+
+	if (InMethodName == NAME_None)
+	{
+		ReportError(TEXT("InMethodName is None."));
+		return nullptr;
+	}
+
+	// find the input and output pins to use
+	FProperty* InputProperty = InScriptStruct->FindPropertyByName(InInputPinName);
+	if (InputProperty == nullptr)
+	{
+		ReportErrorf(TEXT("Cannot find property '%s' on struct type '%s'."), *InInputPinName.ToString(), *InScriptStruct->GetName());
+		return nullptr;
+	}
+	if (!InputProperty->HasMetaData(TEXT("Input")))
+	{
+		ReportErrorf(TEXT("Property '%s' on struct type '%s' is not marked as an input."), *InInputPinName.ToString(), *InScriptStruct->GetName());
+		return nullptr;
+	}
+	FProperty* OutputProperty = InScriptStruct->FindPropertyByName(InOutputPinName);
+	if (OutputProperty == nullptr)
+	{
+		ReportErrorf(TEXT("Cannot find property '%s' on struct type '%s'."), *InOutputPinName.ToString(), *InScriptStruct->GetName());
+		return nullptr;
+	}
+	if (!OutputProperty->HasMetaData(TEXT("Output")))
+	{
+		ReportErrorf(TEXT("Property '%s' on struct type '%s' is not marked as an output."), *InOutputPinName.ToString(), *InScriptStruct->GetName());
+		return nullptr;
+	}
+
+	FRigVMBaseAction Action;
+	if (bUndo)
+	{
+		Action.Title = FString::Printf(TEXT("Add Injected Node"));
+		ActionStack->BeginAction(Action);
+	}
+
+	URigVMStructNode* StructNode = nullptr;
+	{
+		TGuardValue<bool> GuardNotifications(bSuspendNotifications, true);
+		StructNode = AddStructNode(InScriptStruct, InMethodName, FVector2D::ZeroVector, InNodeName, false);
+	}
+	if (StructNode == nullptr)
+	{
+		if (bUndo)
+		{
+			ActionStack->CancelAction(Action);
+		}
+		return nullptr;
+	}
+	else if (StructNode->IsMutable())
+	{
+		ReportErrorf(TEXT("Injected node %s is mutable."), *InScriptStruct->GetName());
+		RemoveNode(StructNode, false);
+		if (bUndo)
+		{
+			ActionStack->CancelAction(Action);
+		}
+		return nullptr;
+	}
+
+	URigVMPin* InputPin = StructNode->FindPin(InInputPinName.ToString());
+	check(InputPin);
+	URigVMPin* OutputPin = StructNode->FindPin(InOutputPinName.ToString());
+	check(OutputPin);
+
+	if (InputPin->GetCPPType() != OutputPin->GetCPPType() ||
+		InputPin->IsArray() != OutputPin->IsArray())
+	{
+		ReportErrorf(TEXT("Injected node %s is using incompatible input and output pins."), *InScriptStruct->GetName());
+		RemoveNode(StructNode, false);
+		if (bUndo)
+		{
+			ActionStack->CancelAction(Action);
+		}
+		return nullptr;
+	}
+
+	if (InputPin->GetCPPType() != Pin->GetCPPType() ||
+		InputPin->IsArray() != Pin->IsArray())
+	{
+		ReportErrorf(TEXT("Injected node %s is using incompatible pin."), *InScriptStruct->GetName());
+		RemoveNode(StructNode, false);
+		if (bUndo)
+		{
+			ActionStack->CancelAction(Action);
+		}
+		return nullptr;
+	}
+
+	URigVMInjectionInfo* InjectionInfo = NewObject<URigVMInjectionInfo>(Pin);
+
+	// re-parent the struct node to be under the injection info
+	StructNode->Rename(nullptr, InjectionInfo);
+
+	InjectionInfo->StructNode = StructNode;
+	InjectionInfo->bInjectedAsInput = bAsInput;
+	InjectionInfo->InputPin = InputPin;
+	InjectionInfo->OutputPin = OutputPin;
+
+	if (bUndo)
+	{
+		ActionStack->AddAction(FRigVMAddInjectedNodeAction(InjectionInfo));
+	}
+
+	URigVMPin* PreviousInputPin = Pin;
+	URigVMPin* PreviousOutputPin = Pin;
+	if (Pin->InjectionInfos.Num() > 0)
+	{
+		PreviousInputPin = Pin->InjectionInfos.Last()->InputPin;
+		PreviousOutputPin = Pin->InjectionInfos.Last()->OutputPin;
+	}
+	Pin->InjectionInfos.Add(InjectionInfo);
+
+	Notify(ERigVMGraphNotifType::NodeAdded, StructNode);
+
+	// now update all of the links
+	if (bAsInput)
+	{
+		FString PinDefaultValue = PreviousInputPin->GetDefaultValue();
+		if (!PinDefaultValue.IsEmpty())
+		{
+			SetPinDefaultValue(InjectionInfo->InputPin, PinDefaultValue, true, false, false);
+		}
+		TArray<URigVMLink*> Links = PreviousInputPin->GetSourceLinks(true /* recursive */);
+		BreakAllLinks(PreviousInputPin, true, false);
+		AddLink(InjectionInfo->OutputPin, PreviousInputPin, false);
+		if (Links.Num() > 0)
+		{
+			RewireLinks(PreviousInputPin, InjectionInfo->InputPin, true, false, Links);
+		}
+	}
+	else
+	{
+		TArray<URigVMLink*> Links = PreviousOutputPin->GetTargetLinks(true /* recursive */);
+		BreakAllLinks(PreviousOutputPin, false, false);
+		AddLink(PreviousOutputPin, InjectionInfo->InputPin, false);
+		if (Links.Num() > 0)
+		{
+			RewireLinks(PreviousOutputPin, InjectionInfo->OutputPin, false, false, Links);
+		}
+	}
+
+	if (bUndo)
+	{
+		ActionStack->EndAction(Action);
+	}
+
+	return InjectionInfo;
+
+}
+
+URigVMInjectionInfo* URigVMController::AddInjectedNodeFromStructPath(const FString& InPinPath, bool bAsInput, const FString& InScriptStructPath, const FName& InMethodName, const FName& InInputPinName, const FName& InOutputPinName, const FString& InNodeName, bool bUndo)
+{
+	if (!IsValidGraph())
+	{
+		return nullptr;
+	}
+
+	UScriptStruct* ScriptStruct = URigVMPin::FindObjectFromCPPTypeObjectPath<UScriptStruct>(InScriptStructPath);
+	if (ScriptStruct == nullptr)
+	{
+		ReportErrorf(TEXT("Cannot find struct for path '%s'."), *InScriptStructPath);
+		return nullptr;
+	}
+
+	return AddInjectedNode(InPinPath, bAsInput, ScriptStruct, InMethodName, InInputPinName, InOutputPinName, InNodeName, bUndo);
+}
+
+URigVMNode* URigVMController::EjectNodeFromPin(const FString& InPinPath, bool bUndo)
+{
+	if (!IsValidGraph())
+	{
+		return nullptr;
+	}
+
+	URigVMPin* Pin = Graph->FindPin(InPinPath);
+	if (Pin == nullptr)
+	{
+		ReportErrorf(TEXT("Cannot find pin '%s'."), *InPinPath);
+		return nullptr;
+	}
+
+	if (!Pin->HasInjectedNodes())
+	{
+		ReportErrorf(TEXT("Pin '%s' has no injected nodes."), *InPinPath);
+		return nullptr;
+	}
+
+	URigVMInjectionInfo* Injection = Pin->InjectionInfos.Last();
+
+	UScriptStruct* ScriptStruct = Injection->StructNode->GetScriptStruct();
+	FName StructNodeName = Injection->StructNode->GetFName();
+	FName MethodName = Injection->StructNode->GetMethodName();
+	FName InputPinName = Injection->InputPin->GetFName();
+	FName OutputPinName = Injection->OutputPin->GetFName();
+
+	TMap<FName, FString> DefaultValues;
+	for (URigVMPin* PinOnNode : Injection->StructNode->Pins)
+	{
+		if (PinOnNode->GetDirection() == ERigVMPinDirection::Input ||
+			PinOnNode->GetDirection() == ERigVMPinDirection::Visible ||
+			PinOnNode->GetDirection() == ERigVMPinDirection::IO)
+		{
+			FString DefaultValue = PinOnNode->GetDefaultValue();
+			PostProcessDefaultValue(PinOnNode, DefaultValue);
+			DefaultValues.Add(PinOnNode->GetFName(), DefaultValue);
+		}
+	}
+
+	FRigVMBaseAction Action;
+	if (bUndo)
+	{
+		Action.Title = FString::Printf(TEXT("Eject Node"));
+		ActionStack->BeginAction(Action);
+	}
+
+	FVector2D Position = Pin->GetNode()->GetPosition() + FVector2D(0.f, 12.f) * float(Pin->GetPinIndex());
+	if (Pin->GetDirection() == ERigVMPinDirection::Output)
+	{
+		Position += FVector2D(250.f, 0.f);
+	}
+	else
+	{
+		Position -= FVector2D(250.f, 0.f);
+	}
+
+	URigVMNode* EjectedNode = AddStructNode(ScriptStruct, MethodName, Position, FString(), bUndo);
+
+	for (const TPair<FName, FString>& Pair : DefaultValues)
+	{
+		if (Pair.Value.IsEmpty())
+		{
+			continue;
+		}
+		if (URigVMPin* PinOnNode = EjectedNode->FindPin(Pair.Key.ToString()))
+		{
+			SetPinDefaultValue(PinOnNode, Pair.Value, true, bUndo, false);
+		}
+	}
+
+	TArray<URigVMLink*> PreviousLinks = Injection->InputPin->GetSourceLinks(true);
+	PreviousLinks.Append(Injection->OutputPin->GetTargetLinks(true));
+	for (URigVMLink* PreviousLink : PreviousLinks)
+	{
+		PreviousLink->PrepareForCopy();
+		PreviousLink->SourcePin = PreviousLink->TargetPin = nullptr;
+	}
+
+	RemoveNode(Injection->StructNode, bUndo);
+
+	FString OldNodeNamePrefix = StructNodeName.ToString() + TEXT(".");
+	FString NewNodeNamePrefix = EjectedNode->GetName() + TEXT(".");
+
+	for (URigVMLink* PreviousLink : PreviousLinks)
+	{
+		FString SourcePinPath = PreviousLink->SourcePinPath;
+		if (SourcePinPath.StartsWith(OldNodeNamePrefix))
+		{
+			SourcePinPath = NewNodeNamePrefix + SourcePinPath.RightChop(OldNodeNamePrefix.Len());
+		}
+		FString TargetPinPath = PreviousLink->TargetPinPath;
+		if (TargetPinPath.StartsWith(OldNodeNamePrefix))
+		{
+			TargetPinPath = NewNodeNamePrefix + TargetPinPath.RightChop(OldNodeNamePrefix.Len());
+		}
+
+		URigVMPin* SourcePin = Graph->FindPin(SourcePinPath);
+		URigVMPin* TargetPin = Graph->FindPin(TargetPinPath);
+		AddLink(SourcePin, TargetPin, bUndo);
+	}
+
+	TArray<FName> NodeNamesToSelect;
+	NodeNamesToSelect.Add(EjectedNode->GetFName());
+	SetNodeSelection(NodeNamesToSelect, bUndo);
+
+	if (bUndo)
+	{
+		ActionStack->EndAction(Action);
+	}
+
+	return EjectedNode;
+}
+
+
+bool URigVMController::Undo()
+{
+	if (!IsValidGraph())
+	{
+		return false;
+	}
+	return ActionStack->Undo(this);
+}
+
+bool URigVMController::Redo()
+{
+	if (!IsValidGraph())
+	{
+		return false;
+	}
+	return ActionStack->Redo(this);
+}
+
+bool URigVMController::OpenUndoBracket(const FString& InTitle)
+{
+	if (!IsValidGraph())
+	{
+		return false;
+	}
+	return ActionStack->OpenUndoBracket(InTitle);
+}
+
+bool URigVMController::CloseUndoBracket()
+{
+	if (!IsValidGraph())
+	{
+		return false;
+	}
+	return ActionStack->CloseUndoBracket();
+}
+
+bool URigVMController::CancelUndoBracket()
+{
+	if (!IsValidGraph())
+	{
+		return false;
+	}
+	return ActionStack->CancelUndoBracket();
+}
+
+FString URigVMController::ExportNodesToText(const TArray<FName>& InNodeNames)
+{
+	if (!IsValidGraph())
+	{
+		return FString();
+	}
+
+	UnMarkAllObjects(EObjectMark(OBJECTMARK_TagExp | OBJECTMARK_TagImp));
+
+	FStringOutputDevice Archive;
+	const FExportObjectInnerContext Context;
+
+	TArray<FName> AllNodeNames = InNodeNames;
+	for (const FName& NodeName : InNodeNames)
+	{
+		if (URigVMNode* Node = Graph->FindNodeByName(NodeName))
+		{
+			for (URigVMPin* Pin : Node->GetPins())
+			{
+				for (URigVMInjectionInfo* Injection : Pin->GetInjectedNodes())
+				{
+					AllNodeNames.AddUnique(Injection->StructNode->GetFName());
+				}
+			}
+		}
+	}
+
+	// Export each of the selected nodes
+	for (const FName& NodeName : InNodeNames)
+	{
+		if (URigVMNode* Node = Graph->FindNodeByName(NodeName))
+		{
+			UExporter::ExportToOutputDevice(&Context, Node, NULL, Archive, TEXT("copy"), 0, PPF_ExportsNotFullyQualified | PPF_Copy | PPF_Delimited, false, Node->GetOuter());
+		}
+	}
+
+	for (URigVMLink* Link : Graph->Links)
+	{
+		URigVMPin* SourcePin = Link->GetSourcePin();
+		URigVMPin* TargetPin = Link->GetTargetPin();
+		if (SourcePin && TargetPin)
+		{
+			if (!AllNodeNames.Contains(SourcePin->GetNode()->GetFName()))
+			{
+				continue;
+			}
+			if (!AllNodeNames.Contains(TargetPin->GetNode()->GetFName()))
+			{
+				continue;
+			}
+			Link->PrepareForCopy();
+			UExporter::ExportToOutputDevice(&Context, Link, NULL, Archive, TEXT("copy"), 0, PPF_ExportsNotFullyQualified | PPF_Copy | PPF_Delimited, false, Link->GetOuter());
+		}
+	}
+
+	return MoveTemp(Archive);
+}
+
+FString URigVMController::ExportSelectedNodesToText()
+{
+	if (!IsValidGraph())
+	{
+		return FString();
+	}
+	return ExportNodesToText(Graph->GetSelectNodes());
+}
+
+struct FRigVMControllerObjectFactory : public FCustomizableTextObjectFactory
+{
+public:
+	URigVMController* Controller;
+	TArray<URigVMNode*> CreatedNodes;
+	TMap<FName, FName> NodeNameMap;
+	TArray<URigVMLink*> CreatedLinks;
+public:
+	FRigVMControllerObjectFactory(URigVMController* InController)
+		: FCustomizableTextObjectFactory(GWarn)
+		, Controller(InController)
+	{
+	}
+
+protected:
+	virtual bool CanCreateClass(UClass* ObjectClass, bool& bOmitSubObjs) const override
+	{
+		if (URigVMNode* DefaultNode = Cast<URigVMNode>(ObjectClass->GetDefaultObject()))
+		{
+			// bOmitSubObjs = true;
+			return true;
+		}
+		if (URigVMLink* DefaultLink = Cast<URigVMLink>(ObjectClass->GetDefaultObject()))
+		{
+			return true;
+		}
+
+		return false;
+	}
+
+	virtual void UpdateObjectName(UClass* ObjectClass, FName& InOutObjName) override
+	{
+		if (URigVMNode* DefaultNode = Cast<URigVMNode>(ObjectClass->GetDefaultObject()))
+		{
+			FName ValidName = *Controller->GetValidNodeName(InOutObjName.ToString());
+			NodeNameMap.Add(InOutObjName, ValidName);
+			InOutObjName = ValidName;
+		}
+	}
+
+	virtual void ProcessConstructedObject(UObject* CreatedObject) override
+	{
+		if (URigVMNode* CreatedNode = Cast<URigVMNode>(CreatedObject))
+		{
+			CreatedNodes.AddUnique(CreatedNode);
+
+			for (URigVMPin* Pin : CreatedNode->GetPins())
+			{
+				for (URigVMInjectionInfo* Injection : Pin->GetInjectedNodes())
+				{
+					ProcessConstructedObject(Injection->StructNode);
+
+					FName NewName = Injection->StructNode->GetFName();
+					UpdateObjectName(URigVMNode::StaticClass(), NewName);
+					Injection->StructNode->Rename(*NewName.ToString(), nullptr);
+					Injection->InputPin = Injection->StructNode->FindPin(Injection->InputPin->GetName());
+					Injection->OutputPin = Injection->StructNode->FindPin(Injection->OutputPin->GetName());
+				}
+			}
+		}
+		else if (URigVMLink* CreatedLink = Cast<URigVMLink>(CreatedObject))
+		{
+			CreatedLinks.Add(CreatedLink);
+		}
+	}
+};
+
+bool URigVMController::CanImportNodesFromText(const FString& InText)
+{
+	if (!IsValidGraph())
+	{
+		return false;
+	}
+	FRigVMControllerObjectFactory Factory(nullptr);
+	return Factory.CanCreateObjectsFromText(InText);
+}
+
+TArray<FName> URigVMController::ImportNodesFromText(const FString& InText, bool bUndo)
+{
+	TArray<FName> NodeNames;
+	if (!IsValidGraph())
+	{
+		return NodeNames;
+	}
+
+	FRigVMControllerObjectFactory Factory(this);
+	Factory.ProcessBuffer(Graph, RF_Transactional, InText);
+
+	if (Factory.CreatedNodes.Num() == 0)
+	{
+		return NodeNames;
+	}
+
+	if (bUndo)
+	{
+		OpenUndoBracket(TEXT("Importing Nodes from Text"));
+	}
+
+	FRigVMInverseAction AddNodesAction;
+	if (bUndo)
+	{
+		ActionStack->BeginAction(AddNodesAction);
+	}
+
+	for (URigVMNode* CreatedNode : Factory.CreatedNodes)
+	{
+		Graph->Nodes.Add(CreatedNode);
+
+		if (bUndo)
+		{
+			ActionStack->AddAction(FRigVMRemoveNodeAction(CreatedNode));
+		}
+		Notify(ERigVMGraphNotifType::NodeAdded, CreatedNode);
+
+		NodeNames.Add(CreatedNode->GetFName());
+	}
+
+	if (bUndo)
+	{
+		ActionStack->EndAction(AddNodesAction);
+	}
+
+	if (Factory.CreatedLinks.Num() > 0)
+	{
+		FRigVMBaseAction AddLinksAction;
+		if (bUndo)
+		{
+			ActionStack->BeginAction(AddLinksAction);
+		}
+
+		for (URigVMLink* CreatedLink : Factory.CreatedLinks)
+		{
+			FString SourceLeft, SourceRight, TargetLeft, TargetRight;
+			if (URigVMPin::SplitPinPathAtStart(CreatedLink->SourcePinPath, SourceLeft, SourceRight) &&
+				URigVMPin::SplitPinPathAtStart(CreatedLink->TargetPinPath, TargetLeft, TargetRight))
+			{
+				const FName* NewSourceNodeName = Factory.NodeNameMap.Find(*SourceLeft);
+				const FName* NewTargetNodeName = Factory.NodeNameMap.Find(*TargetLeft);
+				if (NewSourceNodeName && NewTargetNodeName)
+				{
+					CreatedLink->SourcePinPath = URigVMPin::JoinPinPath(NewSourceNodeName->ToString(), SourceRight);
+					CreatedLink->TargetPinPath = URigVMPin::JoinPinPath(NewTargetNodeName->ToString(), TargetRight);
+					URigVMPin* SourcePin = CreatedLink->GetSourcePin();
+					URigVMPin* TargetPin = CreatedLink->GetTargetPin();
+					if (SourcePin && TargetPin)
+					{
+						Graph->Links.Add(CreatedLink);
+						SourcePin->Links.Add(CreatedLink);
+						TargetPin->Links.Add(CreatedLink);
+
+						if (bUndo)
+						{
+							ActionStack->AddAction(FRigVMAddLinkAction(SourcePin, TargetPin));
+						}
+						Notify(ERigVMGraphNotifType::LinkAdded, CreatedLink);
+						continue;
+					}
+				}
+			}
+
+			ReportErrorf(TEXT("Cannot import link '%s -> %s'."), *CreatedLink->SourcePinPath, *CreatedLink->TargetPinPath);
+			CreatedLink->Rename(nullptr, GetTransientPackage(), REN_ForceNoResetLoaders);
+			CreatedLink->MarkPendingKill();
+		}
+
+		if (bUndo)
+		{
+			ActionStack->EndAction(AddLinksAction);
+		}
+	}
+
+	if (bUndo)
+	{
+		CloseUndoBracket();
+	}
+
+	return NodeNames;
+}
+
+FName URigVMController::GetUniqueName(const FName& InName, TFunction<bool(const FName&)> IsNameAvailableFunction)
+{
+	FString NamePrefix = InName.ToString();
+	int32 NameSuffix = 0;
+	FString Name = NamePrefix;
+	while (!IsNameAvailableFunction(*Name))
+	{
+		NameSuffix++;
+		Name = FString::Printf(TEXT("%s_%d"), *NamePrefix, NameSuffix);
+	}
+	return *Name;
+}
+
+#endif
+
+bool URigVMController::RemoveNode(URigVMNode* InNode, bool bUndo)
+{
+	if (!IsValidNodeForGraph(InNode))
+	{
+		return false;
+	}
+
+	FRigVMRemoveNodeAction Action;
+	if (bUndo)
+	{
+		Action = FRigVMRemoveNodeAction(InNode);
+		Action.Title = FString::Printf(TEXT("Remove %s Node"), *InNode->GetNodeTitle());
+		ActionStack->BeginAction(Action);
+	}
+
+	if (URigVMInjectionInfo* InjectionInfo = InNode->GetInjectionInfo())
+	{
+		URigVMPin* Pin = InjectionInfo->GetPin();
+		check(Pin);
+
+		Pin->InjectionInfos.Remove(InjectionInfo);
+
+		if (InjectionInfo->bInjectedAsInput)
+		{
+			URigVMPin* LastInputPin = Pin;
+			RewireLinks(InjectionInfo->InputPin, LastInputPin, true, false);
+		}
+		else
+		{
+			URigVMPin* LastOutputPin = Pin;
+			RewireLinks(InjectionInfo->OutputPin, LastOutputPin, false, false);
+		}
+	}
+
+	if (bUndo)
+	{
+		SelectNode(InNode, false, bUndo);
+
+		for (URigVMPin* Pin : InNode->Pins)
+		{
+			TArray<URigVMInjectionInfo*> InjectedNodes = Pin->GetInjectedNodes();
+			for (URigVMInjectionInfo* InjectedNode : InjectedNodes)
+			{
+				RemoveNode(InjectedNode->StructNode, bUndo);
+			}
+
+			BreakAllLinks(Pin, true, bUndo);
+			BreakAllLinks(Pin, false, bUndo);
+			BreakAllLinksRecursive(Pin, true, false, bUndo);
+			BreakAllLinksRecursive(Pin, false, false, bUndo);
+		}
+	}
+
+	Graph->Nodes.Remove(InNode);
+	Graph->MarkPackageDirty();
+
+	Notify(ERigVMGraphNotifType::NodeRemoved, InNode);
+
+	if (URigVMVariableNode* VariableNode = Cast<URigVMVariableNode>(InNode))
+	{
+		Notify(ERigVMGraphNotifType::VariableRemoved, VariableNode);
+	}
+	if (URigVMParameterNode* ParameterNode = Cast<URigVMParameterNode>(InNode))
+	{
+		Notify(ERigVMGraphNotifType::ParameterRemoved, ParameterNode);
+	}
+
+	if (URigVMInjectionInfo* InjectionInfo = InNode->GetInjectionInfo())
+	{
+		InjectionInfo->Rename(nullptr, GetTransientPackage(), REN_ForceNoResetLoaders);
+		InjectionInfo->MarkPendingKill();
+	}
+
+	InNode->Rename(nullptr, GetTransientPackage(), REN_ForceNoResetLoaders);
+	InNode->MarkPendingKill();
+
+	if (bUndo)
+	{
+		ActionStack->EndAction(Action);
+	}
+
+	return true;
+}
+
+bool URigVMController::RemoveNodeByName(const FName& InNodeName, bool bUndo)
+{
+	if (!IsValidGraph())
+	{
+		return false;
+	}
+	return RemoveNode(Graph->FindNodeByName(InNodeName), bUndo);
+}
+
+bool URigVMController::SelectNode(URigVMNode* InNode, bool bSelect, bool bUndo)
+{
+	if (!IsValidNodeForGraph(InNode))
+	{
+		return false;
+	}
+
+	if (InNode->IsSelected() == bSelect)
+	{
+		return false;
+	}
+
+	TArray<FName> NewSelection = Graph->GetSelectNodes();
+	if (bSelect)
+	{
+		NewSelection.AddUnique(InNode->GetFName());
+	}
+	else
+	{
+		NewSelection.Remove(InNode->GetFName());
+	}
+
+	return SetNodeSelection(NewSelection, bUndo);
+}
+
+bool URigVMController::SelectNodeByName(const FName& InNodeName, bool bSelect, bool bUndo)
+{
+	if (!IsValidGraph())
+	{
+		return false;
+	}
+	return SelectNode(Graph->FindNodeByName(InNodeName), bSelect, bUndo);
+}
+
+bool URigVMController::ClearNodeSelection(bool bUndo)
+{
+	if (!IsValidGraph())
+	{
+		return false;
+	}
+
+	return SetNodeSelection(TArray<FName>(), bUndo);
+}
+
+bool URigVMController::SetNodeSelection(const TArray<FName>& InNodeNames, bool bUndo)
+{
+	if (!IsValidGraph())
+	{
+		return false;
+	}
+
+	FRigVMSetNodeSelectionAction Action;
+	if (bUndo)
+	{
+		Action = FRigVMSetNodeSelectionAction(Graph, InNodeNames);
+		ActionStack->BeginAction(Action);
+	}
+
+	bool bSelectedSomething = false;
+
+	TArray<FName> PreviousSelection = Graph->GetSelectNodes();
+	for (const FName& PreviouslySelectedNode : PreviousSelection)
+	{
+		if (!InNodeNames.Contains(PreviouslySelectedNode))
+		{
+			if(Graph->SelectedNodes.Remove(PreviouslySelectedNode) > 0)
+			{
+				Notify(ERigVMGraphNotifType::NodeDeselected, Graph->FindNodeByName(PreviouslySelectedNode));
+				bSelectedSomething = true;
+			}
+		}
+	}
+
+	for (const FName& InNodeName : InNodeNames)
+	{
+		if (URigVMNode* NodeToSelect = Graph->FindNodeByName(InNodeName))
+		{
+			int32 PreviousNum = Graph->SelectedNodes.Num();
+			Graph->SelectedNodes.AddUnique(InNodeName);
+			if (PreviousNum != Graph->SelectedNodes.Num())
+			{
+				Notify(ERigVMGraphNotifType::NodeSelected, NodeToSelect);
+				bSelectedSomething = true;
+			}
+		}
+	}
+
+	if (bUndo)
+	{
+		if (bSelectedSomething)
+		{
+			const TArray<FName>& SelectedNodes = Graph->GetSelectNodes();
+			if (SelectedNodes.Num() == 0)
+			{
+				Action.Title = TEXT("Deselect all nodes.");
+			}
+			else
+			{
+				if (SelectedNodes.Num() == 1)
+				{
+					Action.Title = FString::Printf(TEXT("Selected node '%s'."), *SelectedNodes[0].ToString());
+				}
+				else
+				{
+					Action.Title = TEXT("Selected multiple nodes.");
+				}
+			}
+			ActionStack->EndAction(Action);
+		}
+		else
+		{
+			ActionStack->CancelAction(Action);
+		}
+	}
+
+	if (bSelectedSomething)
+	{
+		Notify(ERigVMGraphNotifType::NodeSelectionChanged, nullptr);
+	}
+
+	return bSelectedSomething;
+}
+
+bool URigVMController::SetNodePosition(URigVMNode* InNode, const FVector2D& InPosition, bool bUndo, bool bMergeUndoAction)
+{
+	if (!IsValidNodeForGraph(InNode))
+	{
+		return false;
+	}
+
+	if ((InNode->Position - InPosition).IsNearlyZero())
+	{
+		return false;
+	}
+
+	FRigVMSetNodePositionAction Action;
+	if (bUndo)
+	{
+		Action = FRigVMSetNodePositionAction(InNode, InPosition);
+		Action.Title = FString::Printf(TEXT("Set Node Position"));
+		ActionStack->BeginAction(Action);
+	}
+
+	InNode->Position = InPosition;
+	Notify(ERigVMGraphNotifType::NodePositionChanged, InNode);
+
+	if (bUndo)
+	{
+		ActionStack->EndAction(Action, bMergeUndoAction);
+	}
+
+	return true;
+}
+
+bool URigVMController::SetNodePositionByName(const FName& InNodeName, const FVector2D& InPosition, bool bUndo, bool bMergeUndoAction)
+{
+	if(!IsValidGraph())
+	{
+		return false;
+	}
+
+	URigVMNode* Node = Graph->FindNodeByName(InNodeName);
+	return SetNodePosition(Node, InPosition, bUndo, bMergeUndoAction);
+}
+
+bool URigVMController::SetNodeSize(URigVMNode* InNode, const FVector2D& InSize, bool bUndo, bool bMergeUndoAction)
+{
+	if (!IsValidNodeForGraph(InNode))
+	{
+		return false;
+	}
+
+	if ((InNode->Size - InSize).IsNearlyZero())
+	{
+		return false;
+	}
+
+	FRigVMSetNodeSizeAction Action;
+	if (bUndo)
+	{
+		Action = FRigVMSetNodeSizeAction(InNode, InSize);
+		Action.Title = FString::Printf(TEXT("Set Node Size"));
+		ActionStack->BeginAction(Action);
+	}
+
+	InNode->Size = InSize;
+	Notify(ERigVMGraphNotifType::NodeSizeChanged, InNode);
+
+	if (bUndo)
+	{
+		ActionStack->EndAction(Action, bMergeUndoAction);
+	}
+
+	return true;
+}
+
+bool URigVMController::SetNodeSizeByName(const FName& InNodeName, const FVector2D& InSize, bool bUndo, bool bMergeUndoAction)
+{
+	if(!IsValidGraph())
+	{
+		return false;
+	}
+
+	URigVMNode* Node = Graph->FindNodeByName(InNodeName);
+	return SetNodeSize(Node, InSize, bUndo, bMergeUndoAction);
+}
+
+bool URigVMController::SetNodeColor(URigVMNode* InNode, const FLinearColor& InColor, bool bUndo, bool bMergeUndoAction)
+{
+	if (!IsValidNodeForGraph(InNode))
+	{
+		return false;
+	}
+
+	if ((FVector4(InNode->NodeColor) - FVector4(InColor)).IsNearlyZero3())
+	{
+		return false;
+	}
+
+	FRigVMSetNodeColorAction Action;
+	if (bUndo)
+	{
+		Action = FRigVMSetNodeColorAction(InNode, InColor);
+		Action.Title = FString::Printf(TEXT("Set Node Color"));
+		ActionStack->BeginAction(Action);
+	}
+
+	InNode->NodeColor = InColor;
+	Notify(ERigVMGraphNotifType::NodeColorChanged, InNode);
+
+	if (bUndo)
+	{
+		ActionStack->EndAction(Action, bMergeUndoAction);
+	}
+
+	return true;
+}
+
+bool URigVMController::SetNodeColorByName(const FName& InNodeName, const FLinearColor& InColor, bool bUndo, bool bMergeUndoAction)
+{
+	if(!IsValidGraph())
+	{
+		return false;
+	}
+
+	URigVMNode* Node = Graph->FindNodeByName(InNodeName);
+	return SetNodeColor(Node, InColor, bUndo, bMergeUndoAction);
+}
+
+bool URigVMController::SetCommentText(URigVMNode* InNode, const FString& InCommentText, bool bUndo)
+{
+	if (!IsValidNodeForGraph(InNode))
+	{
+		return false;
+	}
+
+	if (URigVMCommentNode* CommentNode = Cast<URigVMCommentNode>(InNode))
+	{
+		if(CommentNode->CommentText == InCommentText)
+		{
+			return false;
+		}
+
+		FRigVMSetCommentTextAction Action;
+		if (bUndo)
+		{
+			Action = FRigVMSetCommentTextAction(CommentNode, InCommentText);
+			Action.Title = FString::Printf(TEXT("Set Comment Text"));
+			ActionStack->BeginAction(Action);
+		}
+
+		CommentNode->CommentText = InCommentText;
+		Notify(ERigVMGraphNotifType::CommentTextChanged, InNode);
+
+		if (bUndo)
+		{
+			ActionStack->EndAction(Action);
+		}
+
+		return true;
+	}
+
+	return false;
+}
+
+bool URigVMController::SetCommentTextByName(const FName& InNodeName, const FString& InCommentText, bool bUndo)
+{
+	if(!IsValidGraph())
+	{
+		return false;
+	}
+
+	URigVMNode* Node = Graph->FindNodeByName(InNodeName);
+	return SetCommentText(Node, InCommentText, bUndo);
+}
+
+bool URigVMController::SetRerouteCompactness(URigVMNode* InNode, bool bShowAsFullNode, bool bUndo)
+{
+	if (!IsValidNodeForGraph(InNode))
+	{
+		return false;
+	}
+
+	if (URigVMRerouteNode* RerouteNode = Cast<URigVMRerouteNode>(InNode))
+	{
+		if (RerouteNode->bShowAsFullNode == bShowAsFullNode)
+		{
+			return false;
+		}
+
+		FRigVMSetRerouteCompactnessAction Action;
+		if (bUndo)
+		{
+			Action = FRigVMSetRerouteCompactnessAction(RerouteNode, bShowAsFullNode);
+			Action.Title = FString::Printf(TEXT("Set Reroute Size"));
+			ActionStack->BeginAction(Action);
+		}
+
+		RerouteNode->bShowAsFullNode = bShowAsFullNode;
+		Notify(ERigVMGraphNotifType::RerouteCompactnessChanged, InNode);
+
+		if (bUndo)
+		{
+			ActionStack->EndAction(Action);
+		}
+
+		return true;
+	}
+
+	return false;
+}
+
+bool URigVMController::SetRerouteCompactnessByName(const FName& InNodeName, bool bShowAsFullNode, bool bUndo)
+{
+	if(!IsValidGraph())
+	{
+		return false;
+	}
+
+	URigVMNode* Node = Graph->FindNodeByName(InNodeName);
+	return SetRerouteCompactness(Node, bShowAsFullNode, bUndo);
+}
+
+bool URigVMController::RenameVariable(const FName& InOldName, const FName& InNewName, bool bUndo)
+{
+	if(!IsValidGraph())
+	{
+		return false;
+	}
+
+	if (InOldName == InNewName)
+	{
+		ReportWarning(TEXT("RenameVariable: InOldName and InNewName are equal."));
+		return false;
+	}
+
+	TArray<FRigVMGraphVariableDescription> ExistingVariables = Graph->GetVariableDescriptions();
+	for (const FRigVMGraphVariableDescription& ExistingVariable : ExistingVariables)
+	{
+		if (ExistingVariable.Name == InNewName)
+		{
+			ReportErrorf(TEXT("Cannot rename variable to '%s' - variable already exists."), *InNewName.ToString());
+			return false;
+		}
+	}
+
+	FRigVMRenameVariableAction Action;
+	if (bUndo)
+	{
+		Action = FRigVMRenameVariableAction(InOldName, InNewName);
+		Action.Title = FString::Printf(TEXT("Rename Variable"));
+		ActionStack->BeginAction(Action);
+	}
+
+	TArray<URigVMNode*> RenamedNodes;
+	for (URigVMNode* Node : Graph->Nodes)
+	{
+		if(URigVMVariableNode* VariableNode = Cast<URigVMVariableNode>(Node))
+		{
+			if (VariableNode->GetVariableName() == InOldName)
+			{
+				VariableNode->FindPin(URigVMVariableNode::VariableName)->DefaultValue = InNewName.ToString();
+				RenamedNodes.Add(Node);
+			}
+		}
+	}
+
+	for (URigVMNode* RenamedNode : RenamedNodes)
+	{
+		Notify(ERigVMGraphNotifType::VariableRenamed, RenamedNode);
+		Graph->MarkPackageDirty();
+	}
+
+	if (bUndo)
+	{
+		if (RenamedNodes.Num() > 0)
+		{
+			ActionStack->EndAction(Action);
+		}
+		else
+		{
+			ActionStack->CancelAction(Action);
+		}
+	}
+
+	return RenamedNodes.Num() > 0;
+}
+
+bool URigVMController::RenameParameter(const FName& InOldName, const FName& InNewName, bool bUndo)
+{
+	if(!IsValidGraph())
+	{
+		return false;
+	}
+
+	if (InOldName == InNewName)
+	{
+		ReportWarning(TEXT("RenameParameter: InOldName and InNewName are equal."));
+		return false;
+	}
+
+	TArray<FRigVMGraphParameterDescription> ExistingParameters = Graph->GetParameterDescriptions();
+	for (const FRigVMGraphParameterDescription& ExistingParameter : ExistingParameters)
+	{
+		if (ExistingParameter.Name == InNewName)
+		{
+			ReportErrorf(TEXT("Cannot rename parameter to '%s' - parameter already exists."), *InNewName.ToString());
+			return false;
+		}
+	}
+
+	FRigVMRenameParameterAction Action;
+	if (bUndo)
+	{
+		Action = FRigVMRenameParameterAction(InOldName, InNewName);
+		Action.Title = FString::Printf(TEXT("Rename Parameter"));
+		ActionStack->BeginAction(Action);
+	}
+
+	TArray<URigVMNode*> RenamedNodes;
+	for (URigVMNode* Node : Graph->Nodes)
+	{
+		if(URigVMParameterNode* ParameterNode = Cast<URigVMParameterNode>(Node))
+		{
+			if (ParameterNode->GetParameterName() == InOldName)
+			{
+				ParameterNode->FindPin(URigVMParameterNode::ParameterName)->DefaultValue = InNewName.ToString();
+				RenamedNodes.Add(Node);
+			}
+		}
+	}
+
+	for (URigVMNode* RenamedNode : RenamedNodes)
+	{
+		Notify(ERigVMGraphNotifType::ParameterRenamed, RenamedNode);
+		Graph->MarkPackageDirty();
+	}
+
+	if (bUndo)
+	{
+		if (RenamedNodes.Num() > 0)
+		{
+			ActionStack->EndAction(Action);
+		}
+		else
+		{
+			ActionStack->CancelAction(Action);
+		}
+	}
+
+	return RenamedNodes.Num() > 0;
+}
+
+bool URigVMController::SetPinExpansion(const FString& InPinPath, bool bIsExpanded, bool bUndo)
+{
+	if (!IsValidGraph())
+	{
+		return false;
+	}
+
+	URigVMPin* Pin = Graph->FindPin(InPinPath);
+	if (Pin == nullptr)
+	{
+		ReportErrorf(TEXT("Cannot find pin '%s'."), *InPinPath);
+		return false;
+	}
+
+	return SetPinExpansion(Pin, bIsExpanded, bUndo);
+}
+
+bool URigVMController::SetPinExpansion(URigVMPin* InPin, bool bIsExpanded, bool bUndo)
+{
+	if (!IsValidPinForGraph(InPin))
+	{
+		return false;
+	}
+
+	if (InPin->GetSubPins().Num() == 0)
+	{
+		return false;
+	}
+
+	if (InPin->IsExpanded() == bIsExpanded)
+	{
+		return false;
+	}
+
+	FRigVMSetPinExpansionAction Action;
+	if (bUndo)
+	{
+		Action = FRigVMSetPinExpansionAction(InPin, bIsExpanded);
+		Action.Title = bIsExpanded ? TEXT("Expand Pin") : TEXT("Collapse Pin");
+		ActionStack->BeginAction(Action);
+	}
+
+	InPin->bIsExpanded = bIsExpanded;
+
+	Notify(ERigVMGraphNotifType::PinExpansionChanged, InPin);
+	Graph->MarkPackageDirty();
+
+	if (bUndo)
+	{
+		ActionStack->EndAction(Action);
+	}
+
+	return true;
+}
+
+bool URigVMController::SetPinIsWatched(const FString& InPinPath, bool bIsWatched, bool bUndo)
+{
+	if (!IsValidGraph())
+	{
+		return false;
+	}
+
+	URigVMPin* Pin = Graph->FindPin(InPinPath);
+	if (Pin == nullptr)
+	{
+		ReportErrorf(TEXT("Cannot find pin '%s'."), *InPinPath);
+		return false;
+	}
+
+	return SetPinIsWatched(Pin, bIsWatched, bUndo);
+}
+
+bool URigVMController::SetPinIsWatched(URigVMPin* InPin, bool bIsWatched, bool bUndo)
+{
+	if (!IsValidPinForGraph(InPin))
+	{
+		return false;
+	}
+
+	if (InPin->GetParentPin() != nullptr)
+	{
+		return false;
+	}
+
+	if (InPin->RequiresWatch() == bIsWatched)
+	{
+		return false;
+	}
+
+	FRigVMSetPinWatchAction Action;
+	if (bUndo)
+	{
+		Action = FRigVMSetPinWatchAction(InPin, bIsWatched);
+		Action.Title = bIsWatched ? TEXT("Watch Pin") : TEXT("Unwatch Pin");
+		ActionStack->BeginAction(Action);
+	}
+
+	InPin->bRequiresWatch = bIsWatched;
+
+	Notify(ERigVMGraphNotifType::PinWatchedChanged, InPin);
+	Graph->MarkPackageDirty();
+
+	if (bUndo)
+	{
+		ActionStack->EndAction(Action);
+	}
+
+	return true;
+}
+
+FString URigVMController::GetPinDefaultValue(const FString& InPinPath)
+{
+	if (!IsValidGraph())
+	{
+		return FString();
+	}
+
+	URigVMPin* Pin = Graph->FindPin(InPinPath);
+	if (Pin == nullptr)
+	{
+		ReportErrorf(TEXT("Cannot find pin '%s'."), *InPinPath);
+		return FString();
+	}
+	Pin = Pin->GetPinForLink();
+
+	return Pin->GetDefaultValue();
+}
+
+bool URigVMController::SetPinDefaultValue(const FString& InPinPath, const FString& InDefaultValue, bool bResizeArrays, bool bUndo, bool bMergeUndoAction)
+{
+	if (!IsValidGraph())
+	{
+		return false;
+	}
+
+	URigVMPin* Pin = Graph->FindPin(InPinPath);
+	if (Pin == nullptr)
+	{
+		ReportErrorf(TEXT("Cannot find pin '%s'."), *InPinPath);
+		return false;
+	}
+
+	if (URigVMVariableNode* VariableNode = Cast<URigVMVariableNode>(Pin->GetNode()))
+	{
+		if (Pin->GetName() == URigVMVariableNode::VariableName)
+		{
+			return SetVariableName(VariableNode, *InDefaultValue, bUndo);
+		}
+	}
+	if (URigVMParameterNode* ParameterNode = Cast<URigVMParameterNode>(Pin->GetNode()))
+	{
+		if (Pin->GetName() == URigVMParameterNode::ParameterName)
+		{
+			return SetParameterName(ParameterNode, *InDefaultValue, bUndo);
+		}
+	}
+
+	SetPinDefaultValue(Pin, InDefaultValue, bResizeArrays, bUndo, bMergeUndoAction);
+	URigVMPin* PinForLink = Pin->GetPinForLink();
+	if (PinForLink != Pin)
+	{
+		SetPinDefaultValue(PinForLink, InDefaultValue, bResizeArrays, false, bMergeUndoAction);
+	}
+
+	return true;
+}
+
+void URigVMController::SetPinDefaultValue(URigVMPin* InPin, const FString& InDefaultValue, bool bResizeArrays, bool bUndo, bool bMergeUndoAction)
+{
+	check(InPin);
+	ensure(!InDefaultValue.IsEmpty());
+
+	FRigVMSetPinDefaultValueAction Action;
+	if (bUndo)
+	{
+		Action = FRigVMSetPinDefaultValueAction(InPin, InDefaultValue);
+		Action.Title = FString::Printf(TEXT("Set Pin Default Value"));
+		ActionStack->BeginAction(Action);
+	}
+
+	bool bSetPinDefaultValueSucceeded = false;
+	if (InPin->IsArray())
+	{
+		if (ShouldPinBeUnfolded(InPin))
+		{
+			TArray<FString> Elements = SplitDefaultValue(InDefaultValue);
+
+			if (bResizeArrays)
+			{
+				while (Elements.Num() > InPin->SubPins.Num())
+				{
+					InsertArrayPin(InPin, INDEX_NONE, FString(), bUndo);
+				}
+				while (Elements.Num() < InPin->SubPins.Num())
+				{
+					RemoveArrayPin(InPin->SubPins.Last()->GetPinPath(), bUndo);
+				}
+			}
+			else
+			{
+				ensure(Elements.Num() == InPin->SubPins.Num());
+			}
+
+			for (int32 ElementIndex = 0; ElementIndex < Elements.Num(); ElementIndex++)
+			{
+				URigVMPin* SubPin = InPin->SubPins[ElementIndex];
+				PostProcessDefaultValue(SubPin, Elements[ElementIndex]);
+				if (!Elements[ElementIndex].IsEmpty())
+				{
+					SetPinDefaultValue(SubPin, Elements[ElementIndex], bResizeArrays, false, false);
+					bSetPinDefaultValueSucceeded = true;
+				}
+			}
+		}
+	}
+	else if (InPin->IsStruct())
+	{
+		TArray<FString> MemberValuePairs = SplitDefaultValue(InDefaultValue);
+
+		for (const FString& MemberValuePair : MemberValuePairs)
+		{
+			FString MemberName, MemberValue;
+			if (MemberValuePair.Split(TEXT("="), &MemberName, &MemberValue))
+			{
+				URigVMPin* SubPin = InPin->FindSubPin(MemberName);
+				if (SubPin && !MemberValue.IsEmpty())
+				{
+					PostProcessDefaultValue(SubPin, MemberValue);
+					if (!MemberValue.IsEmpty())
+					{
+						SetPinDefaultValue(SubPin, MemberValue, bResizeArrays, false, false);
+						bSetPinDefaultValueSucceeded = true;
+					}
+				}
+			}
+		}
+	}
+	
+	if(!bSetPinDefaultValueSucceeded)
+	{
+		if (InPin->GetSubPins().Num() == 0)
+		{
+			InPin->DefaultValue = InDefaultValue;
+			Notify(ERigVMGraphNotifType::PinDefaultValueChanged, InPin);
+			Graph->MarkPackageDirty();
+		}
+	}
+
+	if (bUndo)
+	{
+		ActionStack->EndAction(Action, bMergeUndoAction);
+	}
+}
+
+bool URigVMController::ResetPinDefaultValue(const FString& InPinPath, bool bUndo)
+{
+	if (!IsValidGraph())
+	{
+		return false;
+	}
+
+	URigVMPin* Pin = Graph->FindPin(InPinPath);
+	if (Pin == nullptr)
+	{
+		ReportErrorf(TEXT("Cannot find pin '%s'."), *InPinPath);
+		return false;
+	}
+
+	if (Cast<URigVMStructNode>(Pin->GetNode()) == nullptr)
+	{
+		ReportErrorf(TEXT("Pin '%s' is not on a struct node."), *InPinPath);
+		return false;
+	}
+
+	return ResetPinDefaultValue(Pin, bUndo);
+}
+
+bool URigVMController::ResetPinDefaultValue(URigVMPin* InPin, bool bUndo)
+{
+	check(InPin);
+
+	if (URigVMStructNode* StructNode = Cast<URigVMStructNode>(InPin->GetNode()))
+	{
+		TSharedPtr<FStructOnScope> StructOnScope = StructNode->ConstructStructInstance(true /* use default */);
+
+		TArray<FString> Parts;
+		if (!URigVMPin::SplitPinPath(InPin->GetPinPath(), Parts))
+		{
+			return false;
+		}
+
+		int32 PartIndex = 1; // cut off the first one since it's the node
+
+		UStruct* Struct = StructNode->ScriptStruct;
+		FProperty* Property = Struct->FindPropertyByName(*Parts[PartIndex++]);
+		if (!ensure(Property != nullptr))
+		{
+			return false;
+		}
+
+		const uint8* Memory = StructOnScope->GetStructMemory();
+		Memory = Property->ContainerPtrToValuePtr<uint8>(Memory);
+
+		while (PartIndex < Parts.Num() && Property != nullptr)
+		{
+			if (FArrayProperty* ArrayProperty = CastField<FArrayProperty>(Property))
+			{
+				Property = ArrayProperty->Inner;
+				int32 ArrayIndex = FCString::Atoi(*Parts[PartIndex++]);
+				TArray<uint8>* ArrayPtr = (TArray<uint8>*)Memory;
+				Memory = ArrayPtr->GetData() + Property->GetSize() * ArrayIndex;
+				continue;
+			}
+
+			if (FStructProperty* StructProperty = CastField<FStructProperty>(Property))
+			{
+				Struct = StructProperty->Struct;
+				Property = Struct->FindPropertyByName(*Parts[PartIndex++]);
+				Memory = Property->ContainerPtrToValuePtr<uint8>(Memory);
+				continue;
+			}
+
+			break;
+		}
+
+		if (Memory)
+		{
+			FString DefaultValue;
+			Property->ExportTextItem(DefaultValue, Memory, nullptr, nullptr, PPF_None);
+
+			if (!DefaultValue.IsEmpty())
+			{
+				SetPinDefaultValue(InPin, DefaultValue, true, bUndo, false);
+				return true;
+			}
+		}
+	}
+
+	return false;
+}
+
+TArray<FString> URigVMController::SplitDefaultValue(const FString& InDefaultValue)
+{
+	TArray<FString> Parts;
+	if (InDefaultValue.IsEmpty())
+	{
+		return Parts;
+	}
+
+	ensure(InDefaultValue[0] == TCHAR('('));
+	ensure(InDefaultValue[InDefaultValue.Len() - 1] == TCHAR(')'));
+
+	FString Content = InDefaultValue.Mid(1, InDefaultValue.Len() - 2);
+	int32 BraceCount = 0;
+	int32 QuoteCount = 0;
+
+	int32 LastPartStartIndex = 0;
+	for (int32 CharIndex = 0; CharIndex < Content.Len(); CharIndex++)
+	{
+		TCHAR Char = Content[CharIndex];
+		if (QuoteCount > 0)
+		{
+			if (Char == TCHAR('"'))
+			{
+				QuoteCount = 0;
+			}
+		}
+		else if (Char == TCHAR('"'))
+		{
+			QuoteCount = 1;
+		}
+		
+		if (Char == TCHAR('('))
+		{
+			if (QuoteCount == 0)
+			{
+				BraceCount++;
+			}
+		}
+		else if (Char == TCHAR(')'))
+		{
+			if (QuoteCount == 0)
+			{
+				BraceCount--;
+				BraceCount = FMath::Max<int32>(BraceCount, 0);
+			}
+		}
+		else if (Char == TCHAR(',') && BraceCount == 0 && QuoteCount == 0)
+		{
+			Parts.Add(Content.Mid(LastPartStartIndex, CharIndex - LastPartStartIndex));
+			LastPartStartIndex = CharIndex + 1;
+		}
+	}
+
+	if (!Content.IsEmpty())
+	{
+		Parts.Add(Content.Mid(LastPartStartIndex));
+	}
+	return Parts;
+}
+
+FString URigVMController::AddArrayPin(const FString& InArrayPinPath, const FString& InDefaultValue, bool bUndo)
+{
+	return InsertArrayPin(InArrayPinPath, INDEX_NONE, InDefaultValue, bUndo);
+}
+
+FString URigVMController::DuplicateArrayPin(const FString& InArrayElementPinPath, bool bUndo)
+{
+	if (!IsValidGraph())
+	{
+		return FString();
+	}
+
+	URigVMPin* ElementPin = Graph->FindPin(InArrayElementPinPath);
+	if (ElementPin == nullptr)
+	{
+		ReportErrorf(TEXT("Cannot find pin '%s'."), *InArrayElementPinPath);
+		return FString();
+	}
+
+	if (!ElementPin->IsArrayElement())
+	{
+		ReportErrorf(TEXT("Pin '%s' is not an array element."), *InArrayElementPinPath);
+		return FString();
+	}
+
+	URigVMPin* ArrayPin = ElementPin->GetParentPin();
+	check(ArrayPin);
+	ensure(ArrayPin->IsArray());
+
+	FString DefaultValue = ElementPin->GetDefaultValue();
+	return InsertArrayPin(ArrayPin->GetPinPath(), ElementPin->GetPinIndex() + 1, DefaultValue, bUndo);
+}
+
+FString URigVMController::InsertArrayPin(const FString& InArrayPinPath, int32 InIndex, const FString& InDefaultValue, bool bUndo)
+{
+	if (!IsValidGraph())
+	{
+		return FString();
+	}
+
+	URigVMPin* ArrayPin = Graph->FindPin(InArrayPinPath);
+	if (ArrayPin == nullptr)
+	{
+		ReportErrorf(TEXT("Cannot find pin '%s'."), *InArrayPinPath);
+		return FString();
+	}
+
+	URigVMPin* ElementPin = InsertArrayPin(ArrayPin, InIndex, InDefaultValue, bUndo);
+	if (ElementPin)
+	{
+		return ElementPin->GetPinPath();
+	}
+
+	return FString();
+}
+
+URigVMPin* URigVMController::InsertArrayPin(URigVMPin* ArrayPin, int32 InIndex, const FString& InDefaultValue, bool bUndo)
+{
+	if (!ArrayPin->IsArray())
+	{
+		ReportErrorf(TEXT("Pin '%s' is not an array."), *ArrayPin->GetPinPath());
+		return nullptr;
+	}
+
+	if (!ShouldPinBeUnfolded(ArrayPin))
+	{
+		ReportErrorf(TEXT("Cannot insert array pin under '%s'."), *ArrayPin->GetPinPath());
+		return nullptr;
+	}
+
+	if (InIndex == INDEX_NONE)
+	{
+		InIndex = ArrayPin->GetSubPins().Num();
+	}
+
+	FRigVMInsertArrayPinAction Action;
+	if (bUndo)
+	{
+		Action = FRigVMInsertArrayPinAction(ArrayPin, InIndex, InDefaultValue);
+		Action.Title = FString::Printf(TEXT("Insert Array Pin"));
+		ActionStack->BeginAction(Action);
+	}
+
+	for (int32 ExistingIndex = ArrayPin->GetSubPins().Num() - 1; ExistingIndex >= InIndex; ExistingIndex--)
+	{
+		URigVMPin* ExistingPin = ArrayPin->GetSubPins()[ExistingIndex];
+		ExistingPin->Rename(*FString::FormatAsNumber(ExistingIndex + 1));
+	}
+
+	URigVMPin* Pin = NewObject<URigVMPin>(ArrayPin, *FString::FormatAsNumber(InIndex));
+	ConfigurePinFromPin(Pin, ArrayPin);
+	Pin->CPPType = ArrayPin->GetArrayElementCppType();
+	ArrayPin->SubPins.Insert(Pin, InIndex);
+
+	if (Pin->IsStruct())
+	{
+		UScriptStruct* ScriptStruct = Pin->GetScriptStruct();
+		if (ScriptStruct)
+		{
+			FString DefaultValue = InDefaultValue;
+			CreateDefaultValueForStructIfRequired(ScriptStruct, DefaultValue);
+			AddPinsForStruct(ScriptStruct, Pin->GetNode(), Pin, Pin->Direction, DefaultValue, false);
+		}
+	}
+	else if (Pin->IsArray())
+	{
+		FArrayProperty * ArrayProperty = CastField<FArrayProperty>(FindPropertyForPin(Pin->GetPinPath()));
+		if (ArrayProperty)
+		{
+			TArray<FString> ElementDefaultValues = SplitDefaultValue(InDefaultValue);
+			AddPinsForArray(ArrayProperty, Pin->GetNode(), Pin, Pin->Direction, ElementDefaultValues, false);
+		}
+	}
+	else
+	{
+		FString DefaultValue = InDefaultValue;
+		PostProcessDefaultValue(Pin, DefaultValue);
+		Pin->DefaultValue = DefaultValue;
+	}
+
+	Notify(ERigVMGraphNotifType::PinArraySizeChanged, ArrayPin);
+	Graph->MarkPackageDirty();
+
+	if (bUndo)
+	{
+		ActionStack->EndAction(Action);
+	}
+
+	return Pin;
+}
+
+bool URigVMController::RemoveArrayPin(const FString& InArrayElementPinPath, bool bUndo)
+{
+	if(!IsValidGraph())
+	{
+		return false;
+	}
+
+	URigVMPin* ArrayElementPin = Graph->FindPin(InArrayElementPinPath);
+	if (ArrayElementPin == nullptr)
+	{
+		ReportErrorf(TEXT("Cannot find pin '%s'."), *InArrayElementPinPath);
+		return false;
+	}
+
+	if (!ArrayElementPin->IsArrayElement())
+	{
+		ReportErrorf(TEXT("Pin '%s' is not an array element."), *InArrayElementPinPath);
+		return false;
+	}
+
+	URigVMPin* ArrayPin = ArrayElementPin->GetParentPin();
+	check(ArrayPin);
+	ensure(ArrayPin->IsArray());
+
+	FRigVMRemoveArrayPinAction Action;
+	if (bUndo)
+	{
+		Action = FRigVMRemoveArrayPinAction(ArrayElementPin);
+		Action.Title = FString::Printf(TEXT("Remove Array Pin"));
+		ActionStack->BeginAction(Action);
+	}
+
+	int32 IndexToRemove = ArrayElementPin->GetPinIndex();
+	if (!RemovePin(ArrayElementPin, bUndo))
+	{
+		return false;
+	}
+
+	for (int32 ExistingIndex = ArrayPin->GetSubPins().Num() - 1; ExistingIndex >= IndexToRemove; ExistingIndex--)
+	{
+		URigVMPin* ExistingPin = ArrayPin->GetSubPins()[ExistingIndex];
+		ExistingPin->SetNameFromIndex();
+	}
+
+	Graph->MarkPackageDirty();
+	Notify(ERigVMGraphNotifType::PinArraySizeChanged, ArrayPin);
+
+	if (bUndo)
+	{
+		ActionStack->EndAction(Action);
+	}
+
+	return true;
+}
+
+bool URigVMController::RemovePin(URigVMPin* InPinToRemove, bool bUndo)
+{
+	if (bUndo)
+	{
+		BreakAllLinks(InPinToRemove, true, bUndo);
+		BreakAllLinks(InPinToRemove, false, bUndo);
+		BreakAllLinksRecursive(InPinToRemove, true, false, bUndo);
+		BreakAllLinksRecursive(InPinToRemove, false, false, bUndo);
+	}
+
+	if (URigVMPin* ParentPin = InPinToRemove->GetParentPin())
+	{
+		ParentPin->SubPins.Remove(InPinToRemove);
+	}
+	else if(URigVMNode* Node = InPinToRemove->GetNode())
+	{
+		Node->Pins.Remove(InPinToRemove);
+	}
+
+	InPinToRemove->Rename(nullptr, GetTransientPackage(), REN_ForceNoResetLoaders);
+	InPinToRemove->MarkPendingKill();
+
+	TArray<URigVMPin*> SubPins = InPinToRemove->GetSubPins();
+	for (URigVMPin* SubPin : SubPins)
+	{
+		if (!RemovePin(SubPin, bUndo))
+		{
+			return false;
+		}
+	}
+
+	return true;
+}
+
+bool URigVMController::ClearArrayPin(const FString& InArrayPinPath, bool bUndo)
+{
+	return SetArrayPinSize(InArrayPinPath, 0, FString(), bUndo);
+}
+
+bool URigVMController::SetArrayPinSize(const FString& InArrayPinPath, int32 InSize, const FString& InDefaultValue, bool bUndo)
+{
+	if (!IsValidGraph())
+	{
+		return false;
+	}
+
+	URigVMPin* Pin = Graph->FindPin(InArrayPinPath);
+	if (Pin == nullptr)
+	{
+		ReportErrorf(TEXT("Cannot find pin '%s'."), *InArrayPinPath);
+		return false;
+	}
+
+	if (!Pin->IsArray())
+	{
+		ReportErrorf(TEXT("Pin '%s' is not an array."), *InArrayPinPath);
+		return false;
+	}
+
+	FRigVMBaseAction Action;
+	if (bUndo)
+	{
+		Action.Title = FString::Printf(TEXT("Set Array Pin Size (%d)"), InSize);
+		ActionStack->BeginAction(Action);
+	}
+
+	InSize = FMath::Max<int32>(InSize, 0);
+	int32 AddedPins = 0;
+	int32 RemovedPins = 0;
+
+	FString DefaultValue = InDefaultValue;
+	if (DefaultValue.IsEmpty())
+	{
+		if (Pin->GetSubPins().Num() > 0)
+		{
+			DefaultValue = Pin->GetSubPins().Last()->GetDefaultValue();
+		}
+		CreateDefaultValueForStructIfRequired(Pin->GetScriptStruct(), DefaultValue);
+	}
+
+	while (Pin->GetSubPins().Num() > InSize)
+	{
+		if (!RemoveArrayPin(Pin->GetSubPins()[Pin->GetSubPins().Num()-1]->GetPinPath(), bUndo))
+		{
+			if (bUndo)
+			{
+				ActionStack->CancelAction(Action);
+			}
+			return false;
+		}
+		RemovedPins++;
+	}
+
+	while (Pin->GetSubPins().Num() < InSize)
+	{
+		if (AddArrayPin(Pin->GetPinPath(), DefaultValue, bUndo).IsEmpty())
+		{
+			if (bUndo)
+			{
+				ActionStack->CancelAction(Action);
+			}
+			return false;
+		}
+		AddedPins++;
+	}
+
+	if (bUndo)
+	{
+		if (RemovedPins > 0 || AddedPins > 0)
+		{
+			ActionStack->EndAction(Action);
+		}
+		else
+		{
+			ActionStack->CancelAction(Action);
+		}
+	}
+
+	return RemovedPins > 0 || AddedPins > 0;
+}
+
+bool URigVMController::AddLink(const FString& InOutputPinPath, const FString& InInputPinPath, bool bUndo)
+{
+	if(!IsValidGraph())
+	{
+		return false;
+	}
+
+	URigVMPin* OutputPin = Graph->FindPin(InOutputPinPath);
+	if (OutputPin == nullptr)
+	{
+		ReportErrorf(TEXT("Cannot find pin '%s'."), *InOutputPinPath);
+		return false;
+	}
+	OutputPin = OutputPin->GetPinForLink();
+
+	URigVMPin* InputPin = Graph->FindPin(InInputPinPath);
+	if (InputPin == nullptr)
+	{
+		ReportErrorf(TEXT("Cannot find pin '%s'."), *InInputPinPath);
+		return false;
+	}
+	InputPin = InputPin->GetPinForLink();
+
+	return AddLink(OutputPin, InputPin, bUndo);
+}
+
+bool URigVMController::AddLink(URigVMPin* OutputPin, URigVMPin* InputPin, bool bUndo)
+{
+	if(OutputPin == nullptr)
+	{
+		ReportError(TEXT("OutputPin is nullptr."));
+		return false;
+	}
+
+	if(InputPin == nullptr)
+	{
+		ReportError(TEXT("InputPin is nullptr."));
+		return false;
+	}
+
+	if(!IsValidPinForGraph(OutputPin) || !IsValidPinForGraph(InputPin))
+	{
+		return false;
+	}
+
+	FString FailureReason;
+	if (!Graph->CanLink(OutputPin, InputPin, &FailureReason))
+	{
+		ReportErrorf(TEXT("Cannot link '%s' to '%s': %s."), *OutputPin->GetPinPath(), *InputPin->GetPinPath(), *FailureReason);
+		return false;
+	}
+
+	ensure(!OutputPin->IsLinkedTo(InputPin));
+	ensure(!InputPin->IsLinkedTo(OutputPin));
+
+	FRigVMAddLinkAction Action;
+	if (bUndo)
+	{
+		Action = FRigVMAddLinkAction(OutputPin, InputPin);
+		Action.Title = FString::Printf(TEXT("Add Link"));
+		ActionStack->BeginAction(Action);
+	}
+
+	BreakAllLinks(InputPin, true, bUndo);
+	if (bUndo)
+	{
+		BreakAllLinksRecursive(InputPin, true, true, bUndo);
+		BreakAllLinksRecursive(InputPin, true, false, bUndo);
+	}
+
+	if (bUndo)
+	{
+		ExpandPinRecursively(OutputPin->GetParentPin(), true);
+		ExpandPinRecursively(InputPin->GetParentPin(), true);
+	}
+
+	URigVMLink* Link = NewObject<URigVMLink>(Graph);
+	Link->SourcePin = OutputPin;
+	Link->TargetPin = InputPin;
+	Link->SourcePinPath = OutputPin->GetPinPath();
+	Link->TargetPinPath = InputPin->GetPinPath();
+	Graph->Links.Add(Link);
+	OutputPin->Links.Add(Link);
+	InputPin->Links.Add(Link);
+
+	Graph->MarkPackageDirty();
+	Notify(ERigVMGraphNotifType::LinkAdded, Link);
+
+	if (bUndo)
+	{
+		ActionStack->EndAction(Action);
+	}
+
+	return true;
+}
+
+bool URigVMController::BreakLink(const FString& InOutputPinPath, const FString& InInputPinPath, bool bUndo)
+{
+	if(!IsValidGraph())
+	{
+		return false;
+	}
+
+	URigVMPin* OutputPin = Graph->FindPin(InOutputPinPath);
+	if (OutputPin == nullptr)
+	{
+		ReportErrorf(TEXT("Cannot find pin '%s'."), *InOutputPinPath);
+		return false;
+	}
+	OutputPin = OutputPin->GetPinForLink();
+
+	URigVMPin* InputPin = Graph->FindPin(InInputPinPath);
+	if (InputPin == nullptr)
+	{
+		ReportErrorf(TEXT("Cannot find pin '%s'."), *InInputPinPath);
+		return false;
+	}
+	InputPin = InputPin->GetPinForLink();
+
+	return BreakLink(OutputPin, InputPin, bUndo);
+}
+
+bool URigVMController::BreakLink(URigVMPin* OutputPin, URigVMPin* InputPin, bool bUndo)
+{
+	if(!IsValidPinForGraph(OutputPin) || !IsValidPinForGraph(InputPin))
+	{
+		return false;
+	}
+
+	if (!OutputPin->IsLinkedTo(InputPin))
+	{
+		return false;
+	}
+	ensure(InputPin->IsLinkedTo(OutputPin));
+
+	for (URigVMLink* Link : InputPin->Links)
+	{
+		if (Link->SourcePin == OutputPin && Link->TargetPin == InputPin)
+		{
+			FRigVMBreakLinkAction Action;
+			if (bUndo)
+			{
+				Action = FRigVMBreakLinkAction(OutputPin, InputPin);
+				Action.Title = FString::Printf(TEXT("Break Link"));
+				ActionStack->BeginAction(Action);
+			}
+
+			OutputPin->Links.Remove(Link);
+			InputPin->Links.Remove(Link);
+			Graph->Links.Remove(Link);
+			
+			Graph->MarkPackageDirty();
+			Notify(ERigVMGraphNotifType::LinkRemoved, Link);
+
+			Link->Rename(nullptr, GetTransientPackage());
+			Link->MarkPendingKill();
+
+			if (bUndo)
+			{
+				ActionStack->EndAction(Action);
+			}
+
+			return true;
+		}
+	}
+
+	return false;
+}
+
+bool URigVMController::BreakAllLinks(const FString& InPinPath, bool bAsInput, bool bUndo)
+{
+	if(!IsValidGraph())
+	{
+		return false;
+	}
+
+	URigVMPin* Pin = Graph->FindPin(InPinPath);
+	if (Pin == nullptr)
+	{
+		ReportErrorf(TEXT("Cannot find pin '%s'."), *InPinPath);
+		return false;
+	}
+	Pin = Pin->GetPinForLink();
+
+	return BreakAllLinks(Pin, bAsInput, bUndo);
+}
+
+bool URigVMController::BreakAllLinks(URigVMPin* Pin, bool bAsInput, bool bUndo)
+{
+	if(!IsValidPinForGraph(Pin))
+	{
+		return false;
+	}
+
+	FRigVMBaseAction Action;
+	if (bUndo)
+	{
+		Action.Title = FString::Printf(TEXT("Break All Links"));
+		ActionStack->BeginAction(Action);
+	}
+
+	int32 LinksBroken = 0;
+	TArray<URigVMLink*> Links = Pin->GetLinks();
+	for (int32 LinkIndex = Links.Num() - 1; LinkIndex >= 0; LinkIndex--)
+	{
+		URigVMLink* Link = Links[LinkIndex];
+		if (bAsInput && Link->GetTargetPin() == Pin)
+		{
+			LinksBroken += BreakLink(Link->GetSourcePin(), Pin, bUndo) ? 1 : 0;
+		}
+		else if (!bAsInput && Link->GetSourcePin() == Pin)
+		{
+			LinksBroken += BreakLink(Pin, Link->GetTargetPin(), bUndo) ? 1 : 0;
+		}
+	}
+
+	if (bUndo)
+	{
+		if (LinksBroken > 0)
+		{
+			ActionStack->EndAction(Action);
+		}
+		else
+		{
+			ActionStack->CancelAction(Action);
+		}
+	}
+
+	return LinksBroken > 0;
+}
+
+void URigVMController::BreakAllLinksRecursive(URigVMPin* Pin, bool bAsInput, bool bTowardsParent, bool bUndo)
+{
+	if (bTowardsParent)
+	{
+		URigVMPin* ParentPin = Pin->GetParentPin();
+		if (ParentPin)
+		{
+			BreakAllLinks(ParentPin, bAsInput, bUndo);
+			BreakAllLinksRecursive(ParentPin, bAsInput, bTowardsParent, bUndo);
+		}
+	}
+	else
+	{
+		for (URigVMPin* SubPin : Pin->SubPins)
+		{
+			BreakAllLinks(SubPin, bAsInput, bUndo);
+			BreakAllLinksRecursive(SubPin, bAsInput, bTowardsParent, bUndo);
+		}
+	}
+}
+
+void URigVMController::ExpandPinRecursively(URigVMPin* InPin, bool bUndo)
+{
+	if (InPin == nullptr)
+	{
+		return;
+	}
+
+	if (bUndo)
+	{
+		OpenUndoBracket(TEXT("Expand Pin Recursively"));
+	}
+
+	bool bExpandedSomething = false;
+	while (InPin)
+	{
+		if (SetPinExpansion(InPin, true, bUndo))
+		{
+			bExpandedSomething = true;
+		}
+		InPin = InPin->GetParentPin();
+	}
+
+	if (bUndo)
+	{
+		if (bExpandedSomething)
+		{
+			CloseUndoBracket();
+		}
+		else
+		{
+			CancelUndoBracket();
+		}
+	}
+}
+
+bool URigVMController::SetVariableName(URigVMVariableNode* InVariableNode, const FName& InVariableName, bool bUndo)
+{
+	if (!IsValidNodeForGraph(InVariableNode))
+	{
+		return false;
+	}
+
+	if (InVariableNode->GetVariableName() == InVariableName)
+	{
+		return false;
+	}
+
+	if (InVariableName == NAME_None)
+	{
+		return false;
+	}
+
+	TArray<FRigVMGraphVariableDescription> Descriptions = Graph->GetVariableDescriptions();
+	TMap<FName, int32> NameToIndex;
+	for (int32 VariableIndex = 0; VariableIndex < Descriptions.Num(); VariableIndex++)
+	{
+		NameToIndex.Add(Descriptions[VariableIndex].Name, VariableIndex);
+	}
+
+	FName VariableName = GetUniqueName(InVariableName, [Descriptions, NameToIndex, InVariableNode](const FName& InName) {
+		const int32* FoundIndex = NameToIndex.Find(InName);
+		if (FoundIndex == nullptr)
+		{
+			return true;
+		}
+		return InVariableNode->GetCPPType() == Descriptions[*FoundIndex].CPPType;
+	});
+
+	int32 NodesSharingName = 0;
+	for (URigVMNode* Node : Graph->Nodes)
+	{
+		if (URigVMVariableNode* OtherVariableNode = Cast<URigVMVariableNode>(Node))
+		{
+			if (OtherVariableNode->GetVariableName() == InVariableNode->GetVariableName())
+			{
+				NodesSharingName++;
+			}
+		}
+	}
+
+	if (NodesSharingName == 1)
+	{
+		Notify(ERigVMGraphNotifType::VariableRemoved, InVariableNode);
+	}
+
+	SetPinDefaultValue(InVariableNode->FindPin(URigVMVariableNode::VariableName), VariableName.ToString(), false, bUndo, false);
+
+	Notify(ERigVMGraphNotifType::VariableAdded, InVariableNode);
+
+	return true;
+}
+
+bool URigVMController::SetParameterName(URigVMParameterNode* InParameterNode, const FName& InParameterName, bool bUndo)
+{
+	if (!IsValidNodeForGraph(InParameterNode))
+	{
+		return false;
+	}
+
+	if (InParameterNode->GetParameterName() == InParameterName)
+	{
+		return false;
+	}
+
+	if (InParameterName == NAME_None)
+	{
+		return false;
+	}
+
+	TArray<FRigVMGraphParameterDescription> Descriptions = Graph->GetParameterDescriptions();
+	TMap<FName, int32> NameToIndex;
+	for (int32 ParameterIndex = 0; ParameterIndex < Descriptions.Num(); ParameterIndex++)
+	{
+		NameToIndex.Add(Descriptions[ParameterIndex].Name, ParameterIndex);
+	}
+
+	FName ParameterName = GetUniqueName(InParameterName, [Descriptions, NameToIndex, InParameterNode](const FName& InName) {
+		const int32* FoundIndex = NameToIndex.Find(InName);
+		if (FoundIndex == nullptr)
+		{
+			return true;
+		}
+		return InParameterNode->GetCPPType() == Descriptions[*FoundIndex].CPPType && InParameterNode->IsInput() == Descriptions[*FoundIndex].bIsInput;
+	});
+
+	int32 NodesSharingName = 0;
+	for (URigVMNode* Node : Graph->Nodes)
+	{
+		if (URigVMParameterNode* OtherParameterNode = Cast<URigVMParameterNode>(Node))
+		{
+			if (OtherParameterNode->GetParameterName() == InParameterNode->GetParameterName())
+			{
+				NodesSharingName++;
+			}
+		}
+	}
+
+	if (NodesSharingName == 1)
+	{
+		Notify(ERigVMGraphNotifType::ParameterRemoved, InParameterNode);
+	}
+
+	SetPinDefaultValue(InParameterNode->FindPin(URigVMParameterNode::ParameterName), ParameterName.ToString(), false, bUndo, false);
+
+	Notify(ERigVMGraphNotifType::ParameterAdded, InParameterNode);
+
+	return true;
+}
+
+URigVMRerouteNode* URigVMController::AddFreeRerouteNode(bool bShowAsFullNode, const FString& InCPPType, const FName& InCPPTypeObjectPath, bool bIsConstant, const FName& InCustomWidgetName, const FString& InDefaultValue, const FVector2D& InPosition, const FString& InNodeName, bool bUndo)
+{
+	if (!IsValidGraph())
+	{
+		return nullptr;
+	}
+
+	FRigVMBaseAction Action;
+	if (bUndo)
+	{
+		Action.Title = FString::Printf(TEXT("Add Reroute"));
+		ActionStack->BeginAction(Action);
+	}
+
+	FString Name = GetValidNodeName(InNodeName.IsEmpty() ? FString(TEXT("RerouteNode")) : InNodeName);
+	URigVMRerouteNode* Node = NewObject<URigVMRerouteNode>(Graph, *Name);
+	Node->Position = InPosition;
+	Node->bShowAsFullNode = bShowAsFullNode;
+
+	URigVMPin* ValuePin = NewObject<URigVMPin>(Node, *URigVMRerouteNode::ValueName);
+	ValuePin->CPPType = InCPPType;
+	ValuePin->CPPTypeObjectPath = InCPPTypeObjectPath;
+	ValuePin->bIsConstant = bIsConstant;
+	ValuePin->CustomWidgetName = InCustomWidgetName;
+	ValuePin->Direction = ERigVMPinDirection::IO;
+	Node->Pins.Add(ValuePin);
+	Graph->Nodes.Add(Node);
+
+	if (ValuePin->IsStruct())
+	{
+		FString DefaultValue = InDefaultValue;
+		CreateDefaultValueForStructIfRequired(ValuePin->GetScriptStruct(), DefaultValue);
+		AddPinsForStruct(ValuePin->GetScriptStruct(), Node, ValuePin, ValuePin->Direction, DefaultValue, false);
+	}
+	else if (!InDefaultValue.IsEmpty() && InDefaultValue != TEXT("()"))
+	{
+		SetPinDefaultValue(ValuePin, InDefaultValue, true, false, false);
+	}
+
+	ForEveryPinRecursively(ValuePin, [](URigVMPin* Pin) {
+		if (Pin->GetParentPin() != nullptr)
+		{
+			Pin->Direction = ERigVMPinDirection::Visible;
+		}
+	});
+
+	if (bUndo)
+	{
+		ActionStack->AddAction(FRigVMAddRerouteNodeAction(Node));
+	}
+
+	Notify(ERigVMGraphNotifType::NodeAdded, Node);
+
+	if (bUndo)
+	{
+		ActionStack->EndAction(Action);
+	}
+
+	return Node;
+}
+
+void URigVMController::ForEveryPinRecursively(URigVMPin* InPin, TFunction<void(URigVMPin*)> OnEachPinFunction)
+{
+	OnEachPinFunction(InPin);
+	for (URigVMPin* SubPin : InPin->SubPins)
+	{
+		ForEveryPinRecursively(SubPin, OnEachPinFunction);
+	}
+}
+
+void URigVMController::ForEveryPinRecursively(URigVMNode* InNode, TFunction<void(URigVMPin*)> OnEachPinFunction)
+{
+	for (URigVMPin* Pin : InNode->Pins)
+	{
+		ForEveryPinRecursively(Pin, OnEachPinFunction);
+	}
+}
+
+void URigVMController::SetExecuteContextStruct(UStruct* InExecuteContextStruct)
+{
+	check(InExecuteContextStruct);
+	ensure(InExecuteContextStruct->IsChildOf(FRigVMExecuteContext::StaticStruct()));
+	ExecuteContextStruct = InExecuteContextStruct;
+}
+
+FString URigVMController::GetValidNodeName(const FString& InPrefix)
+{
+	return GetUniqueName(*InPrefix, [&](const FName& InName) {
+		return Graph->IsNameAvailable(InName.ToString());
+	}).ToString();
+}
+
+bool URigVMController::IsValidGraph()
+{
+	if (Graph == nullptr)
+	{
+		ReportError(TEXT("Controller does not have a graph associated - use SetGraph / set_graph."));
+		return false;
+	}
+	return true;
+}
+
+bool URigVMController::IsValidNodeForGraph(URigVMNode* InNode)
+{
+	if(!IsValidGraph())
+	{
+		return false;
+	}
+
+	if (InNode == nullptr)
+	{
+		ReportError(TEXT("InNode is nullptr."));
+		return false;
+	}
+
+	if (InNode->GetGraph() != Graph)
+	{
+		ReportErrorf(TEXT("InNode '%s' is on a different graph."), *InNode->GetNodePath());
+		return false;
+	}
+
+	if (InNode->GetNodeIndex() == INDEX_NONE)
+	{
+		ReportErrorf(TEXT("InNode '%s' is transient (not yet nested to a graph)."), *InNode->GetName());
+	}
+
+	return true;
+}
+
+bool URigVMController::IsValidPinForGraph(URigVMPin* InPin)
+{
+	if(!IsValidGraph())
+	{
+		return false;
+	}
+
+	if (InPin == nullptr)
+	{
+		ReportError(TEXT("InPin is nullptr."));
+		return false;
+	}
+
+	if (!IsValidNodeForGraph(InPin->GetNode()))
+	{
+		return false;
+	}
+
+	if (InPin->GetPinIndex() == INDEX_NONE)
+	{
+		ReportErrorf(TEXT("InPin '%s' is transient (not yet nested properly)."), *InPin->GetName());
+	}
+
+	return true;
+}
+
+bool URigVMController::IsValidLinkForGraph(URigVMLink* InLink)
+{
+	if(!IsValidGraph())
+	{
+		return false;
+	}
+
+	if (InLink == nullptr)
+	{
+		ReportError(TEXT("InLink is nullptr."));
+		return false;
+	}
+
+	if (InLink->GetGraph() != Graph)
+	{
+		ReportError(TEXT("InLink is on a different graph."));
+		return false;
+	}
+
+	if(InLink->GetSourcePin() == nullptr)
+	{
+		ReportError(TEXT("InLink has no source pin."));
+		return false;
+	}
+
+	if(InLink->GetTargetPin() == nullptr)
+	{
+		ReportError(TEXT("InLink has no target pin."));
+		return false;
+	}
+
+	if (InLink->GetLinkIndex() == INDEX_NONE)
+	{
+		ReportError(TEXT("InLink is transient (not yet nested properly)."));
+	}
+
+	if(!IsValidPinForGraph(InLink->GetSourcePin()))
+	{
+		return false;
+	}
+
+	if(!IsValidPinForGraph(InLink->GetTargetPin()))
+	{
+		return false;
+	}
+
+	return true;
+}
+
+void URigVMController::AddPinsForStruct(UStruct* InStruct, URigVMNode* InNode, URigVMPin* InParentPin, ERigVMPinDirection InPinDirection, const FString& InDefaultValue, bool bAutoExpandArrays)
+{
+	TArray<FString> MemberNameValuePairs = SplitDefaultValue(InDefaultValue);
+	TMap<FName, FString> MemberValues;
+	for (const FString& MemberNameValuePair : MemberNameValuePairs)
+	{
+		FString MemberName, MemberValue;
+		if (MemberNameValuePair.Split(TEXT("="), &MemberName, &MemberValue))
+		{
+			MemberValues.Add(*MemberName, MemberValue);
+		}
+	}
+
+	for (TFieldIterator<FProperty> It(InStruct); It; ++It)
+	{
+		FName PropertyName = It->GetFName();
+
+		URigVMPin* Pin = NewObject<URigVMPin>(InParentPin == nullptr ? Cast<UObject>(InNode) : Cast<UObject>(InParentPin), PropertyName);
+		ConfigurePinFromProperty(*It, Pin, InPinDirection);
+
+		if (InParentPin)
+		{
+			InParentPin->SubPins.Add(Pin);
+		}
+		else
+		{
+			InNode->Pins.Add(Pin);
+		}
+
+		FString* DefaultValuePtr = MemberValues.Find(Pin->GetFName());
+
+		FStructProperty* StructProperty = CastField<FStructProperty>(*It);
+		if (StructProperty)
+		{
+			if (ShouldStructBeUnfolded(StructProperty->Struct))
+			{
+				FString DefaultValue;
+				if (DefaultValuePtr != nullptr)
+				{
+					DefaultValue = *DefaultValuePtr;
+				}
+				CreateDefaultValueForStructIfRequired(StructProperty->Struct, DefaultValue);
+
+				AddPinsForStruct(StructProperty->Struct, InNode, Pin, Pin->GetDirection(), DefaultValue, bAutoExpandArrays);
+			}
+			else if(DefaultValuePtr != nullptr)
+			{
+				Pin->DefaultValue = *DefaultValuePtr;
+			}
+		}
+
+		FArrayProperty* ArrayProperty = CastField<FArrayProperty>(*It);
+		if (ArrayProperty)
+		{
+			ensure(Pin->IsArray());
+
+			if (DefaultValuePtr)
+			{
+				if (ShouldPinBeUnfolded(Pin))
+				{
+					TArray<FString> ElementDefaultValues = SplitDefaultValue(*DefaultValuePtr);
+					AddPinsForArray(ArrayProperty, InNode, Pin, Pin->Direction, ElementDefaultValues, bAutoExpandArrays);
+				}
+				else
+				{
+					FString DefaultValue = *DefaultValuePtr;
+					PostProcessDefaultValue(Pin, DefaultValue);
+					Pin->DefaultValue = *DefaultValuePtr;
+				}
+			}
+		}
+		
+		if (!Pin->IsArray() && !Pin->IsStruct() && DefaultValuePtr != nullptr)
+		{
+			FString DefaultValue = *DefaultValuePtr;
+			PostProcessDefaultValue(Pin, DefaultValue);
+			Pin->DefaultValue = DefaultValue;
+		}
+	}
+}
+
+void URigVMController::AddPinsForArray(FArrayProperty* InArrayProperty, URigVMNode* InNode, URigVMPin* InParentPin, ERigVMPinDirection InPinDirection, const TArray<FString>& InDefaultValues, bool bAutoExpandArrays)
+{
+	check(InParentPin);
+	if (!ShouldPinBeUnfolded(InParentPin))
+	{
+		return;
+	}
+
+	for (int32 ElementIndex = 0; ElementIndex < InDefaultValues.Num(); ElementIndex++)
+	{
+		FString ElementName = FString::FormatAsNumber(InParentPin->SubPins.Num());
+		URigVMPin* Pin = NewObject<URigVMPin>(InParentPin, *ElementName);
+
+		ConfigurePinFromProperty(InArrayProperty->Inner, Pin, InPinDirection);
+		FString DefaultValue = InDefaultValues[ElementIndex];
+
+		InParentPin->SubPins.Add(Pin);
+
+		if (bAutoExpandArrays)
+		{
+			TGuardValue<bool> ErrorGuard(bReportWarningsAndErrors, false);
+			ExpandPinRecursively(Pin, false);
+		}
+
+		FStructProperty* StructProperty = CastField<FStructProperty>(InArrayProperty->Inner);
+		if (StructProperty)
+		{
+			if (ShouldPinBeUnfolded(Pin))
+			{
+				AddPinsForStruct(StructProperty->Struct, InNode, Pin, Pin->Direction, DefaultValue, bAutoExpandArrays);
+			}
+			else if(!DefaultValue.IsEmpty())
+			{
+				PostProcessDefaultValue(Pin, DefaultValue);
+				Pin->DefaultValue = DefaultValue;
+			}
+		}
+
+		FArrayProperty* ArrayProperty = CastField<FArrayProperty>(InArrayProperty->Inner);
+		if (ArrayProperty)
+		{
+			if (ShouldPinBeUnfolded(Pin))
+			{
+				TArray<FString> ElementDefaultValues = SplitDefaultValue(DefaultValue);
+				AddPinsForArray(ArrayProperty, InNode, Pin, Pin->Direction, ElementDefaultValues, bAutoExpandArrays);
+			}
+			else if (!DefaultValue.IsEmpty())
+			{
+				PostProcessDefaultValue(Pin, DefaultValue);
+				Pin->DefaultValue = DefaultValue;
+			}
+		}
+
+		if (!Pin->IsArray() && !Pin->IsStruct())
+		{
+			PostProcessDefaultValue(Pin, DefaultValue);
+			Pin->DefaultValue = DefaultValue;
+		}
+	}
+}
+
+void URigVMController::ConfigurePinFromProperty(FProperty* InProperty, URigVMPin* InOutPin, ERigVMPinDirection InPinDirection)
+{
+	if (InPinDirection == ERigVMPinDirection::Invalid)
+	{
+#if WITH_EDITOR
+		bool bIsInput = InProperty->HasMetaData(TEXT("Input"));
+		bool bIsOutput = InProperty->HasMetaData(TEXT("Output"));
+
+		if (bIsInput)
+		{
+			InOutPin->Direction = bIsOutput ? ERigVMPinDirection::IO : ERigVMPinDirection::Input;
+		}
+		else
+		{
+			InOutPin->Direction = bIsOutput ? ERigVMPinDirection::Output : ERigVMPinDirection::Hidden;
+		}
+#endif
+	}
+	else
+	{
+		InOutPin->Direction = InPinDirection;
+	}
+
+#if WITH_EDITOR
+
+	if (!InOutPin->IsArrayElement())
+	{
+		FString DisplayNameText = InProperty->GetDisplayNameText().ToString();
+		if (!DisplayNameText.IsEmpty())
+		{
+			InOutPin->DisplayName = *DisplayNameText;
+		}
+		else
+		{
+			InOutPin->DisplayName = NAME_None;
+		}
+	}
+	InOutPin->bIsConstant = InProperty->HasMetaData(TEXT("Constant"));
+	FString CustomWidgetName = InProperty->GetMetaData(TEXT("CustomWidget"));
+	InOutPin->CustomWidgetName = CustomWidgetName.IsEmpty() ? FName(NAME_None) : FName(*CustomWidgetName);
+#endif
+
+	FString ExtendedCppType;
+	InOutPin->CPPType = InProperty->GetCPPType(&ExtendedCppType);
+	InOutPin->CPPType += ExtendedCppType;
+
+	FProperty* PropertyForType = InProperty;
+	FArrayProperty* ArrayProperty = CastField<FArrayProperty>(PropertyForType);
+	if (ArrayProperty)
+	{
+		PropertyForType = ArrayProperty->Inner;
+	}
+
+	if (FStructProperty* StructProperty = CastField<FStructProperty>(PropertyForType))
+	{
+		InOutPin->CPPTypeObject = StructProperty->Struct;
+	}
+	else if (FEnumProperty* EnumProperty = CastField<FEnumProperty>(PropertyForType))
+	{
+		InOutPin->CPPTypeObject = EnumProperty->GetEnum();
+	}
+	else if (FByteProperty* ByteProperty = CastField<FByteProperty>(PropertyForType))
+	{
+		InOutPin->CPPTypeObject = ByteProperty->Enum;
+	}
+
+	if (InOutPin->CPPTypeObject)
+	{
+		InOutPin->CPPTypeObjectPath = *InOutPin->CPPTypeObject->GetPathName();
+	}
+}
+
+void URigVMController::ConfigurePinFromPin(URigVMPin* InOutPin, URigVMPin* InPin)
+{
+	InOutPin->bIsConstant = InPin->bIsConstant;
+	InOutPin->Direction = InPin->Direction;
+	InOutPin->CPPType = InPin->CPPType;
+	InOutPin->CPPTypeObjectPath = InPin->CPPTypeObjectPath;
+	InOutPin->CPPTypeObject = InPin->CPPTypeObject;
+	InOutPin->DefaultValue = InPin->DefaultValue;
+}
+
+bool URigVMController::ShouldStructBeUnfolded(const UStruct* Struct)
+{
+	if (Struct == nullptr)
+	{
+		return false;
+	}
+	if (Struct->IsChildOf(UClass::StaticClass()))
+	{
+		return false;
+	}
+	if(Struct->IsChildOf(FRigVMExecuteContext::StaticStruct()))
+	{
+		return false;
+	}
+	if (UnfoldStructDelegate.IsBound())
+	{
+		if (!UnfoldStructDelegate.Execute(Struct))
+		{
+			return false;
+		}
+	}
+	return true;
+}
+
+bool URigVMController::ShouldPinBeUnfolded(URigVMPin* InPin)
+{
+	if (InPin->IsStruct())
+	{
+		return ShouldStructBeUnfolded(InPin->GetScriptStruct());
+	}
+	else if (InPin->IsArray())
+	{
+		return InPin->GetDirection() == ERigVMPinDirection::Input ||
+			InPin->GetDirection() == ERigVMPinDirection::IO;
+	}
+	return false;
+}
+
+FProperty* URigVMController::FindPropertyForPin(const FString& InPinPath)
+{
+	if(!IsValidGraph())
+	{
+		return nullptr;
+	}
+
+	TArray<FString> Parts;
+	if (!URigVMPin::SplitPinPath(InPinPath, Parts))
+	{
+		return nullptr;
+	}
+
+	URigVMPin* Pin = Graph->FindPin(InPinPath);
+	if (Pin == nullptr)
+	{
+		ReportErrorf(TEXT("Cannot find pin '%s'."), *InPinPath);
+		return nullptr;
+	}
+
+	URigVMNode* Node = Pin->GetNode();
+
+	URigVMStructNode* StructNode = Cast<URigVMStructNode>(Node);
+	if (StructNode)
+	{
+		int32 PartIndex = 1; // cut off the first one since it's the node
+
+		UStruct* Struct = StructNode->ScriptStruct;
+		FProperty* Property = Struct->FindPropertyByName(*Parts[PartIndex++]);
+
+		while (PartIndex < Parts.Num() && Property != nullptr)
+		{
+			if (FArrayProperty* ArrayProperty = CastField<FArrayProperty>(Property))
+			{
+				Property = ArrayProperty->Inner;
+				PartIndex++;
+				continue;
+			}
+
+			if (FStructProperty* StructProperty = CastField<FStructProperty>(Property))
+			{
+				Struct = StructProperty->Struct;
+				Property = Struct->FindPropertyByName(*Parts[PartIndex++]);
+				continue;
+			}
+
+			break;
+		}
+
+		if (PartIndex == Parts.Num())
+		{
+			return Property;
+		}
+	}
+
+	return nullptr;
+}
+
+int32 URigVMController::DetachLinksFromPinObjects()
+{
+	check(Graph);
+	TGuardValue<bool> SuspendNotifs(bSuspendNotifications, true);
+
+	for (URigVMLink* Link : Graph->Links)
+	{
+		URigVMPin* SourcePin = Link->GetSourcePin();
+		URigVMPin* TargetPin = Link->GetTargetPin();
+
+		Link->SourcePinPath = SourcePin->GetPinPath();
+		Link->TargetPinPath = TargetPin->GetPinPath();
+
+		SourcePin->Links.Remove(Link);
+		TargetPin->Links.Remove(Link);
+
+		Link->SourcePin = nullptr;
+		Link->TargetPin = nullptr;
+	}
+
+	return Graph->Links.Num();
+}
+
+int32 URigVMController::ReattachLinksToPinObjects()
+{
+	check(Graph);
+	TGuardValue<bool> SuspendNotifs(bSuspendNotifications, true);
+
+	// fix up the pin links based on the persisted data
+	TArray<URigVMLink*> NewLinks;
+	for (URigVMLink* Link : Graph->Links)
+	{
+		URigVMPin* SourcePin = Link->GetSourcePin();
+		URigVMPin* TargetPin = Link->GetTargetPin();
+		if (SourcePin == nullptr)
+		{
+			ReportWarningf(TEXT("Unable to re-create link %s -> %s"), *Link->SourcePinPath, *Link->TargetPinPath);
+			if (TargetPin != nullptr)
+			{
+				TargetPin->Links.Remove(Link);
+			}
+			continue;
+		}
+		if (TargetPin == nullptr)
+		{
+			ReportWarningf(TEXT("Unable to re-create link %s -> %s"), *Link->SourcePinPath, *Link->TargetPinPath);
+			if (SourcePin != nullptr)
+			{
+				SourcePin->Links.Remove(Link);
+			}
+			continue;
+		}
+
+		SourcePin->Links.AddUnique(Link);
+		TargetPin->Links.AddUnique(Link);
+		NewLinks.Add(Link);
+	}
+	Graph->Links = NewLinks;
+	return NewLinks.Num();
+}
+
+#if WITH_EDITOR
+
+void URigVMController::RepopulatePinsOnNode(URigVMNode* InNode)
+{
+	if (InNode == nullptr)
+	{
+		ReportError(TEXT("InNode is nullptr."));
+		return;
+	}
+
+	URigVMStructNode* StructNode = Cast<URigVMStructNode>(InNode);
+	if (StructNode == nullptr)
+	{
+		return;
+	}
+
+	TGuardValue<bool> SuspendNotifs(bSuspendNotifications, true);
+	check(Graph);
+
+	// todo: ensure to keep all links
+
+	TMap<FName, FString> DefaultValues;
+	for (URigVMPin* Pin : InNode->Pins)
+	{
+		FString DefaultValue = Pin->GetDefaultValue();
+		if (!DefaultValue.IsEmpty())
+		{
+			DefaultValues.Add(Pin->GetFName(), DefaultValue);
+		}
+	}
+
+	TMap<FString, bool> ExpansionStates;
+	TArray<URigVMPin*> AllPins = InNode->GetAllPinsRecursively();
+	for (URigVMPin* Pin : AllPins)
+	{
+		if (Pin->IsExpanded())
+		{
+			FString PinPath = Pin->GetPinPath();
+			PinPath = PinPath.RightChop(InNode->GetNodePath().Len() + 1);
+			ExpansionStates.Add(PinPath, true);
+		}
+	}
+
+	TMap<FString, TArray<URigVMInjectionInfo*>> InjectionInfos;
+	for (URigVMPin* Pin : InNode->Pins)
+	{
+		if (Pin->HasInjectedNodes())
+		{
+			InjectionInfos.Add(Pin->GetName(), Pin->GetInjectedNodes());
+		}
+	}
+
+	// also in case this node is part of an injection
+	FName InjectionInputPinName = NAME_None;
+	FName InjectionOutputPinName = NAME_None;
+	if (URigVMInjectionInfo* InjectionInfo = InNode->GetInjectionInfo())
+	{
+		InjectionInputPinName = InjectionInfo->InputPin->GetFName();
+		InjectionOutputPinName = InjectionInfo->OutputPin->GetFName();
+	}
+
+	TArray<URigVMPin*> Pins = InNode->Pins;
+	for (URigVMPin* Pin : Pins)
+	{
+		RemovePin(Pin, false);
+	}
+	InNode->Pins.Reset();
+	Pins.Reset();
+
+	UScriptStruct* ScriptStruct = StructNode->GetScriptStruct();
+
+	FString NodeColorMetadata;
+	ScriptStruct->GetStringMetaDataHierarchical(*URigVMNode::NodeColorName, &NodeColorMetadata);
+	if (!NodeColorMetadata.IsEmpty())
+	{
+		StructNode->NodeColor = GetColorFromMetadata(NodeColorMetadata);
+	}
+
+	FString ExportedDefaultValue;
+	CreateDefaultValueForStructIfRequired(ScriptStruct, ExportedDefaultValue);
+	AddPinsForStruct(ScriptStruct, StructNode, nullptr, ERigVMPinDirection::Invalid, ExportedDefaultValue, false);
+
+	for (TPair<FString, TArray<URigVMInjectionInfo*>> InjectionPair : InjectionInfos)
+	{
+		if (URigVMPin* Pin = InNode->FindPin(InjectionPair.Key))
+		{
+			for (URigVMInjectionInfo* InjectionInfo : InjectionPair.Value)
+			{
+				InjectionInfo->Rename(nullptr, Pin);
+				InjectionInfo->InputPin = InjectionInfo->StructNode->FindPin(InjectionInfo->InputPin->GetName());
+				InjectionInfo->OutputPin = InjectionInfo->StructNode->FindPin(InjectionInfo->OutputPin->GetName());
+				Pin->InjectionInfos.Add(InjectionInfo);
+			}
+		}
+		else
+		{
+			for (URigVMInjectionInfo* InjectionInfo : InjectionPair.Value)
+			{
+				InjectionInfo->Rename(nullptr, GetTransientPackage());
+				InjectionInfo->MarkPendingKill();
+				InjectionInfo->StructNode->Rename(nullptr, InNode->GetGraph());
+			}
+		}
+	}
+
+	for (TPair<FName, FString> DefaultValuePair : DefaultValues)
+	{
+		if (URigVMPin* Pin = InNode->FindPin(DefaultValuePair.Key.ToString()))
+		{
+			SetPinDefaultValue(Pin, DefaultValuePair.Value, true, false, false);
+		}
+	}
+
+	for (TPair<FString, bool> ExpansionStatePair : ExpansionStates)
+	{
+		if (URigVMPin* Pin = InNode->FindPin(ExpansionStatePair.Key))
+		{
+			SetPinExpansion(Pin->GetPinPath(), ExpansionStatePair.Value, false);
+		}
+	}
+
+	if (URigVMInjectionInfo* InjectionInfo = InNode->GetInjectionInfo())
+	{
+		InjectionInfo->InputPin = InNode->FindPin(InjectionInputPinName.ToString());
+		InjectionInfo->OutputPin = InNode->FindPin(InjectionOutputPinName.ToString());
+	}
+
+}
+
+#endif
+
+FLinearColor URigVMController::GetColorFromMetadata(const FString& InMetadata)
+{
+	FLinearColor Color = FLinearColor::Black;
+
+	FString Metadata = InMetadata;
+	Metadata.TrimStartAndEnd();
+	FString SplitString(TEXT(" "));
+	FString Red, Green, Blue, GreenAndBlue;
+	if (Metadata.Split(SplitString, &Red, &GreenAndBlue))
+	{
+		Red.TrimEnd();
+		GreenAndBlue.TrimStart();
+		if (GreenAndBlue.Split(SplitString, &Green, &Blue))
+		{
+			Green.TrimEnd();
+			Blue.TrimStart();
+
+			float RedValue = FCString::Atof(*Red);
+			float GreenValue = FCString::Atof(*Green);
+			float BlueValue = FCString::Atof(*Blue);
+			Color = FLinearColor(RedValue, GreenValue, BlueValue);
+		}
+	}
+
+	return Color;
+}
+
+void URigVMController::ReportWarning(const FString& InMessage)
+{
+	if(!bReportWarningsAndErrors)
+	{
+		return;
+	}
+
+	FString Message = InMessage;
+	if (Graph != nullptr)
+	{
+		if (UPackage* Package = Cast<UPackage>(Graph->GetOutermost()))
+		{
+			Message = FString::Printf(TEXT("%s : %s"), *Package->GetPathName(), *InMessage);
+		}
+	}
+
+	FScriptExceptionHandler::Get().HandleException(ELogVerbosity::Warning, *Message, *FString());
+}
+
+void URigVMController::ReportError(const FString& InMessage)
+{
+	if(!bReportWarningsAndErrors)
+	{
+		return;
+	}
+
+	FString Message = InMessage;
+	if (Graph != nullptr)
+	{
+		if (UPackage* Package = Cast<UPackage>(Graph->GetOutermost()))
+		{
+			Message = FString::Printf(TEXT("%s : %s"), *Package->GetPathName(), *InMessage);
+		}
+	}
+
+	FScriptExceptionHandler::Get().HandleException(ELogVerbosity::Error, *Message, *FString());
+}
+
+void URigVMController::CreateDefaultValueForStructIfRequired(UScriptStruct* InStruct, FString& OutDefaultValue)
+{
+	if ((OutDefaultValue.IsEmpty() || OutDefaultValue == TEXT("()")) && InStruct != nullptr)
+	{
+		TArray<uint8, TAlignedHeapAllocator<16>> TempBuffer;
+		TempBuffer.AddUninitialized(InStruct->GetStructureSize());
+		InStruct->InitializeDefaultValue(TempBuffer.GetData());
+		InStruct->ExportText(OutDefaultValue, TempBuffer.GetData(), nullptr, nullptr, PPF_None, nullptr);
+		InStruct->DestroyStruct(TempBuffer.GetData());
+	}
+}
+
+void URigVMController::PostProcessDefaultValue(URigVMPin* Pin, FString& OutDefaultValue)
+{
+	if (Pin->IsArray() && OutDefaultValue.IsEmpty())
+	{
+		OutDefaultValue = TEXT("()");
+	}
+	else if (Pin->IsStruct() && (OutDefaultValue.IsEmpty() || OutDefaultValue == TEXT("()")))
+	{
+		CreateDefaultValueForStructIfRequired(Pin->GetScriptStruct(), OutDefaultValue);
+	}
+	else if (Pin->IsStringType())
+	{
+		while (OutDefaultValue.StartsWith(TEXT("\"")))
+		{
+			OutDefaultValue = OutDefaultValue.RightChop(1);
+		}
+		while (OutDefaultValue.EndsWith(TEXT("\"")))
+		{
+			OutDefaultValue = OutDefaultValue.LeftChop(1);
+		}
+	}
+}
+
+#if WITH_EDITOR
+
+void URigVMController::RewireLinks(URigVMPin* InOldPin, URigVMPin* InNewPin, bool bAsInput, bool bUndo, TArray<URigVMLink*> InLinks)
+{
+	ensure(InOldPin->GetRootPin() == InOldPin);
+	ensure(InNewPin->GetRootPin() == InNewPin);
+
+	if (bAsInput)
+	{
+		TArray<URigVMLink*> Links = InLinks;
+		if (Links.Num() == 0)
+		{
+			Links = InOldPin->GetSourceLinks(true /* recursive */);
+		}
+
+		for (URigVMLink* Link : Links)
+		{
+			FString SegmentPath = Link->GetTargetPin()->GetSegmentPath();
+			URigVMPin* NewPin = SegmentPath.IsEmpty() ? InNewPin : InNewPin->FindSubPin(SegmentPath);
+			check(NewPin);
+
+			BreakLink(Link->GetSourcePin(), Link->GetTargetPin(), false);
+			ensure(AddLink(Link->GetSourcePin(), NewPin, false));
+		}
+	}
+	else
+	{
+		TArray<URigVMLink*> Links = InLinks;
+		if (Links.Num() == 0)
+		{
+			Links = InOldPin->GetTargetLinks(true /* recursive */);
+		}
+
+		for (URigVMLink* Link : Links)
+		{
+			FString SegmentPath = Link->GetSourcePin()->GetSegmentPath();
+			URigVMPin* NewPin = SegmentPath.IsEmpty() ? InNewPin : InNewPin->FindSubPin(SegmentPath);
+			check(NewPin);
+
+			BreakLink(Link->GetSourcePin(), Link->GetTargetPin(), false);
+			ensure(AddLink(NewPin, Link->GetTargetPin(), false));
+		}
+	}
+}
+
+#endif
\ No newline at end of file
