COPY from //Dev-Anim to //Dev-Main

diff --git a/Engine/Source/Developer/RigVMDeveloper/Private/RigVMModel/RigVMControllerActions.cpp b/Engine/Source/Developer/RigVMDeveloper/Private/RigVMModel/RigVMControllerActions.cpp
new file mode 100644
--- /dev/null
+++ b/Engine/Source/Developer/RigVMDeveloper/Private/RigVMModel/RigVMControllerActions.cpp
@@ -0,0 +1,1001 @@
+// Copyright Epic Games, Inc. All Rights Reserved.
+
+#include "RigVMModel/RigVMControllerActions.h"
+#include "UObject/Package.h"
+
+#if WITH_EDITOR
+#include "Misc/ITransaction.h"
+#endif
+
+FRigVMActionWrapper::FRigVMActionWrapper(const FRigVMActionKey& Key)
+{
+	ScriptStruct = FindObjectChecked<UScriptStruct>(ANY_PACKAGE, *Key.ScriptStructPath);
+	Data.SetNumUninitialized(ScriptStruct->GetStructureSize());
+	ScriptStruct->InitializeStruct(Data.GetData(), 1);
+	ScriptStruct->ImportText(*Key.ExportedText, Data.GetData(), nullptr, PPF_None, nullptr, ScriptStruct->GetName());
+}
+FRigVMActionWrapper::~FRigVMActionWrapper()
+{
+	if (Data.Num() > 0 && ScriptStruct != nullptr)
+	{
+		ScriptStruct->DestroyStruct(Data.GetData(), 1);
+	}
+}
+
+FRigVMBaseAction* FRigVMActionWrapper::GetAction()
+{
+	return (FRigVMBaseAction*)Data.GetData();
+}
+
+FString FRigVMActionWrapper::ExportText()
+{
+	FString ExportedText;
+	if (Data.Num() > 0 && ScriptStruct != nullptr)
+	{
+		ScriptStruct->ExportText(ExportedText, GetAction(), nullptr, nullptr, PPF_None, nullptr);
+	}
+	return ExportedText;
+}
+
+bool URigVMActionStack::OpenUndoBracket(const FString& InTitle)
+{
+	FRigVMBaseAction* Action = new FRigVMBaseAction;
+	Action->Title = InTitle;
+	BracketActions.Add(Action);
+	BeginAction(*Action);
+	return true;
+}
+
+bool URigVMActionStack::CloseUndoBracket()
+{
+	ensure(BracketActions.Num() > 0);
+	FRigVMBaseAction* Action = BracketActions.Pop();
+	EndAction(*Action);
+	delete(Action);
+	return true;
+}
+
+bool URigVMActionStack::CancelUndoBracket()
+{
+	ensure(BracketActions.Num() > 0);
+	FRigVMBaseAction* Action = BracketActions.Pop();
+	CancelAction(*Action);
+	delete(Action);
+	return true;
+}
+
+bool URigVMActionStack::Undo(URigVMController* InController)
+{
+	check(InController)
+
+	if (UndoActions.Num() == 0)
+	{
+		InController->ReportWarning(TEXT("Nothing to undo."));
+		return false;
+	}
+
+	FRigVMActionKey KeyToUndo = UndoActions.Pop();
+	ActionIndex--;
+	FRigVMActionWrapper Wrapper(KeyToUndo);
+	if (Wrapper.GetAction()->Undo(InController))
+	{
+		RedoActions.Add(KeyToUndo);
+		return true;
+	}
+	return false;
+}
+
+bool URigVMActionStack::Redo(URigVMController* InController)
+{
+	check(InController)
+
+	if (RedoActions.Num() == 0)
+	{
+		InController->ReportWarning(TEXT("Nothing to redo."));
+		return false;
+	}
+
+	FRigVMActionKey KeyToRedo = RedoActions.Pop();
+	ActionIndex++;
+	FRigVMActionWrapper Wrapper(KeyToRedo);
+	if (Wrapper.GetAction()->Redo(InController))
+	{
+		UndoActions.Add(KeyToRedo);
+		return true;
+	}
+	return false;
+}
+
+#if WITH_EDITOR
+
+void URigVMActionStack::PostTransacted(const FTransactionObjectEvent& TransactionEvent)
+{
+	if (TransactionEvent.GetEventType() == ETransactionObjectEventType::UndoRedo)
+	{
+		URigVMController* Controller = Cast< URigVMController>(GetOuter());
+		if (Controller == nullptr)
+		{
+			return;
+		}
+
+		int32 DesiredActionIndex = ActionIndex;
+		ActionIndex = UndoActions.Num();
+
+		if (DesiredActionIndex == ActionIndex)
+		{
+			return;
+		}
+
+		ModifiedEvent.Broadcast(ERigVMGraphNotifType::InteractionBracketOpened, nullptr, nullptr);
+
+		while (DesiredActionIndex < ActionIndex)
+		{
+			if (UndoActions.Num() == 0)
+			{
+				break;
+			}
+			if (!Undo(Controller))
+			{
+				return;
+			}
+		}
+		while (DesiredActionIndex > ActionIndex)
+		{
+			if (RedoActions.Num() == 0)
+			{
+				break;
+			}
+			if (!Redo(Controller))
+			{
+				return;
+			}
+		}
+
+		ModifiedEvent.Broadcast(ERigVMGraphNotifType::InteractionBracketClosed, nullptr, nullptr);
+	}
+}
+
+#endif
+
+
+bool FRigVMBaseAction::Merge(const FRigVMBaseAction* Other)
+{
+	return SubActions.Num() == 0 && Other->SubActions.Num() == 0;
+}
+
+bool FRigVMBaseAction::Undo(URigVMController* InController)
+{
+	bool Result = true;
+	for (int32 KeyIndex = SubActions.Num() - 1; KeyIndex >= 0; KeyIndex--)
+	{
+		FRigVMActionWrapper Wrapper(SubActions[KeyIndex]);
+		if(!Wrapper.GetAction()->Undo(InController))
+		{
+			Result = false;
+		}
+	}
+	return Result;
+}
+
+bool FRigVMBaseAction::Redo(URigVMController* InController)
+{
+	bool Result = true;
+	for (int32 KeyIndex = 0; KeyIndex < SubActions.Num(); KeyIndex++)
+	{
+		FRigVMActionWrapper Wrapper(SubActions[KeyIndex]);
+		if (!Wrapper.GetAction()->Redo(InController))
+		{
+			Result = false;
+		}
+	}
+	return Result;
+}
+
+bool FRigVMInverseAction::Undo(URigVMController* InController)
+{
+	return FRigVMBaseAction::Redo(InController);
+}
+
+bool FRigVMInverseAction::Redo(URigVMController* InController)
+{
+	return FRigVMBaseAction::Undo(InController);
+}
+
+FRigVMAddStructNodeAction::FRigVMAddStructNodeAction()
+	: ScriptStructPath()
+	, MethodName(NAME_None)
+	, Position(FVector2D::ZeroVector)
+	, NodePath()
+{
+}
+
+FRigVMAddStructNodeAction::FRigVMAddStructNodeAction(URigVMStructNode* InNode)
+	: ScriptStructPath(InNode->GetScriptStruct()->GetPathName())
+	, MethodName(InNode->GetMethodName())
+	, Position(InNode->GetPosition())
+	, NodePath(InNode->GetNodePath())
+{
+}
+
+bool FRigVMAddStructNodeAction::Undo(URigVMController* InController)
+{
+	if (!FRigVMBaseAction::Undo(InController))
+	{
+		return false;
+	}
+	return InController->RemoveNodeByName(*NodePath, false);
+}
+
+bool FRigVMAddStructNodeAction::Redo(URigVMController* InController)
+{
+#if WITH_EDITOR
+	if (URigVMStructNode* Node = InController->AddStructNodeFromStructPath(ScriptStructPath, MethodName, Position, NodePath, false))
+	{
+		return FRigVMBaseAction::Redo(InController);
+	}
+#endif
+	return false;
+}
+
+FRigVMAddVariableNodeAction::FRigVMAddVariableNodeAction()
+	: VariableName(NAME_None)
+	, CPPType()
+	, CPPTypeObjectPath()
+	, bIsGetter(false)
+	, DefaultValue()
+	, Position(FVector2D::ZeroVector)
+	, NodePath()
+{
+}
+
+FRigVMAddVariableNodeAction::FRigVMAddVariableNodeAction(URigVMVariableNode* InNode)
+	: VariableName(InNode->GetVariableName())
+	, CPPType(InNode->GetCPPType())
+	, CPPTypeObjectPath()
+	, bIsGetter(InNode->IsGetter())
+	, DefaultValue(InNode->GetDefaultValue())
+	, Position(InNode->GetPosition())
+	, NodePath(InNode->GetNodePath())
+{
+	if (InNode->GetCPPTypeObject())
+	{
+		CPPTypeObjectPath = InNode->GetCPPTypeObject()->GetPathName();
+	}
+}
+
+bool FRigVMAddVariableNodeAction::Undo(URigVMController* InController)
+{
+	if (!FRigVMBaseAction::Undo(InController))
+	{
+		return false;
+	}
+	return InController->RemoveNodeByName(*NodePath, false);
+}
+
+bool FRigVMAddVariableNodeAction::Redo(URigVMController* InController)
+{
+#if WITH_EDITOR
+	if (URigVMVariableNode* Node = InController->AddVariableNodeFromObjectPath(VariableName, CPPType, CPPTypeObjectPath, bIsGetter, DefaultValue, Position, NodePath, false))
+	{
+		return FRigVMBaseAction::Redo(InController);
+	}
+#endif
+	return false;
+}
+
+FRigVMAddParameterNodeAction::FRigVMAddParameterNodeAction()
+	: ParameterName(NAME_None)
+	, CPPType()
+	, CPPTypeObjectPath()
+	, bIsInput(false)
+	, DefaultValue()
+	, Position(FVector2D::ZeroVector)
+	, NodePath()
+{
+}
+
+FRigVMAddParameterNodeAction::FRigVMAddParameterNodeAction(URigVMParameterNode* InNode)
+	: ParameterName(InNode->GetParameterName())
+	, CPPType(InNode->GetCPPType())
+	, CPPTypeObjectPath()
+	, bIsInput(InNode->IsInput())
+	, DefaultValue(InNode->GetDefaultValue())
+	, Position(InNode->GetPosition())
+	, NodePath(InNode->GetNodePath())
+{
+	if (InNode->GetCPPTypeObject())
+	{
+		CPPTypeObjectPath = InNode->GetCPPTypeObject()->GetPathName();
+	}
+}
+
+bool FRigVMAddParameterNodeAction::Undo(URigVMController* InController)
+{
+	if (!FRigVMBaseAction::Undo(InController))
+	{
+		return false;
+	}
+	return InController->RemoveNodeByName(*NodePath, false);
+}
+
+bool FRigVMAddParameterNodeAction::Redo(URigVMController* InController)
+{
+#if WITH_EDITOR
+	if (URigVMParameterNode* Node = InController->AddParameterNodeFromObjectPath(ParameterName, CPPType, CPPTypeObjectPath, bIsInput, DefaultValue, Position, NodePath, false))
+	{
+		return FRigVMBaseAction::Redo(InController);
+	}
+#endif
+	return false;
+}
+
+FRigVMAddCommentNodeAction::FRigVMAddCommentNodeAction()
+	: CommentText()
+	, Position(FVector2D::ZeroVector)
+	, Size(FVector2D::ZeroVector)
+	, Color(FLinearColor::Black)
+	, NodePath()
+{
+}
+
+FRigVMAddCommentNodeAction::FRigVMAddCommentNodeAction(URigVMCommentNode* InNode)
+	: CommentText(InNode->GetCommentText())
+	, Position(InNode->GetPosition())
+	, Size(InNode->GetSize())
+	, Color(InNode->GetNodeColor())
+	, NodePath(InNode->GetNodePath())
+{
+}
+
+bool FRigVMAddCommentNodeAction::Undo(URigVMController* InController)
+{
+	if (!FRigVMBaseAction::Undo(InController))
+	{
+		return false;
+	}
+	return InController->RemoveNodeByName(*NodePath, false);
+}
+
+bool FRigVMAddCommentNodeAction::Redo(URigVMController* InController)
+{
+#if WITH_EDITOR
+	if (URigVMCommentNode* Node = InController->AddCommentNode(CommentText, Position, Size, Color, NodePath, false))
+	{
+		return FRigVMBaseAction::Redo(InController);
+	}
+#endif
+	return false;
+}
+
+FRigVMAddRerouteNodeAction::FRigVMAddRerouteNodeAction()
+	: bShowAsFullNode(false)
+	, CPPType()
+	, CPPTypeObjectPath(NAME_None)
+	, DefaultValue()
+	, bIsConstant(false)
+	, CustomWidgetName(NAME_None)
+	, Position(FVector2D::ZeroVector)
+	, NodePath()
+{
+}
+
+FRigVMAddRerouteNodeAction::FRigVMAddRerouteNodeAction(URigVMRerouteNode* InNode)
+	: bShowAsFullNode(InNode->GetShowsAsFullNode())
+	, CPPType(InNode->GetPins()[0]->GetCPPType())
+	, CPPTypeObjectPath(NAME_None)
+	, DefaultValue(InNode->GetPins()[0]->GetDefaultValue())
+	, bIsConstant(InNode->GetPins()[0]->IsConstant())
+	, CustomWidgetName(InNode->GetPins()[0]->GetCustomWidgetName())
+	, Position(InNode->GetPosition())
+	, NodePath(InNode->GetNodePath())
+{
+	if (InNode->GetPins()[0]->GetCPPTypeObject())
+	{
+		CPPTypeObjectPath = *InNode->GetPins()[0]->GetCPPTypeObject()->GetPathName();
+	}
+}
+
+bool FRigVMAddRerouteNodeAction::Undo(URigVMController* InController)
+{
+	if (!FRigVMBaseAction::Undo(InController))
+	{
+		return false;
+	}
+	return InController->RemoveNodeByName(*NodePath, false);
+}
+
+bool FRigVMAddRerouteNodeAction::Redo(URigVMController* InController)
+{
+#if WITH_EDITOR
+	if (URigVMRerouteNode* Node = InController->AddFreeRerouteNode(bShowAsFullNode, CPPType, CPPTypeObjectPath, bIsConstant, CustomWidgetName, DefaultValue, Position, NodePath, false))
+	{
+		return FRigVMBaseAction::Redo(InController);
+	}
+#endif
+	return false;
+}
+
+FRigVMAddInjectedNodeAction::FRigVMAddInjectedNodeAction()
+	: PinPath()
+	, bAsInput(false)
+	, ScriptStructPath()
+	, MethodName(NAME_None)
+	, InputPinName(NAME_None)
+	, OutputPinName(NAME_None)
+	, NodePath()
+{
+}
+
+FRigVMAddInjectedNodeAction::FRigVMAddInjectedNodeAction(URigVMInjectionInfo* InInjectionInfo)
+	: PinPath(InInjectionInfo->GetPin()->GetPinPath())
+	, bAsInput(InInjectionInfo->bInjectedAsInput)
+	, ScriptStructPath(InInjectionInfo->StructNode->GetScriptStruct()->GetPathName())
+	, MethodName(InInjectionInfo->StructNode->GetMethodName())
+	, InputPinName(InInjectionInfo->InputPin->GetFName())
+	, OutputPinName(InInjectionInfo->OutputPin->GetFName())
+	, NodePath(InInjectionInfo->StructNode->GetName())
+{
+}
+
+bool FRigVMAddInjectedNodeAction::Undo(URigVMController* InController)
+{
+	if (!FRigVMBaseAction::Undo(InController))
+	{
+		return false;
+	}
+	return InController->RemoveNodeByName(*NodePath, false);
+}
+
+bool FRigVMAddInjectedNodeAction::Redo(URigVMController* InController)
+{
+#if WITH_EDITOR
+	if (URigVMInjectionInfo* InjectionInfo = InController->AddInjectedNodeFromStructPath(PinPath, bAsInput, ScriptStructPath, MethodName, InputPinName, OutputPinName, NodePath, false))
+	{
+		return FRigVMBaseAction::Redo(InController);
+	}
+#endif
+	return false;
+}
+
+FRigVMRemoveNodeAction::FRigVMRemoveNodeAction(URigVMNode* InNode)
+{
+	FRigVMInverseAction InverseAction;
+
+	if (URigVMInjectionInfo* InjectionInfo = InNode->GetInjectionInfo())
+	{
+		InverseAction.AddAction(FRigVMAddInjectedNodeAction(InjectionInfo));
+		for (URigVMPin* Pin : InNode->GetPins())
+		{
+			if (Pin->GetDirection() == ERigVMPinDirection::Input ||
+				Pin->GetDirection() == ERigVMPinDirection::Visible)
+			{
+				InverseAction.AddAction(FRigVMSetPinDefaultValueAction(Pin, Pin->GetDefaultValue()));
+			}
+		}
+	}
+	else if (URigVMStructNode* StructNode = Cast<URigVMStructNode>(InNode))
+	{
+		InverseAction.AddAction(FRigVMAddStructNodeAction(StructNode));
+		for (URigVMPin* Pin : StructNode->GetPins())
+		{
+			if (Pin->GetDirection() == ERigVMPinDirection::Input ||
+				Pin->GetDirection() == ERigVMPinDirection::Visible)
+			{
+				InverseAction.AddAction(FRigVMSetPinDefaultValueAction(Pin, Pin->GetDefaultValue()));
+			}
+		}
+	}
+	else if (URigVMVariableNode* VariableNode = Cast<URigVMVariableNode>(InNode))
+	{
+		InverseAction.AddAction(FRigVMAddVariableNodeAction(VariableNode));
+		URigVMPin* ValuePin = VariableNode->FindPin(TEXT("Value"));
+		InverseAction.AddAction(FRigVMSetPinDefaultValueAction(ValuePin, ValuePin->GetDefaultValue()));
+	}
+	else if (URigVMParameterNode* ParameterNode = Cast<URigVMParameterNode>(InNode))
+	{
+		InverseAction.AddAction(FRigVMAddParameterNodeAction(ParameterNode));
+		URigVMPin* ValuePin = ParameterNode->FindPin(TEXT("Value"));
+		InverseAction.AddAction(FRigVMSetPinDefaultValueAction(ValuePin, ValuePin->GetDefaultValue()));
+	}
+	else if (URigVMCommentNode* CommentNode = Cast<URigVMCommentNode>(InNode))
+	{
+		InverseAction.AddAction(FRigVMAddCommentNodeAction(CommentNode));
+	}
+	else if (URigVMRerouteNode* RerouteNode = Cast<URigVMRerouteNode>(InNode))
+	{
+		InverseAction.AddAction(FRigVMAddRerouteNodeAction(RerouteNode));
+	}
+	else
+	{
+		ensure(false);
+	}
+
+	for (URigVMPin* Pin : InNode->GetPins())
+	{
+		if(Pin->IsExpanded())
+		{
+			FRigVMSetPinExpansionAction ExpansionAction(Pin, true);
+			ExpansionAction.OldIsExpanded = false;
+			InverseAction.AddAction(ExpansionAction);
+		}
+	}
+
+	InverseActionKey.Set<FRigVMInverseAction>(InverseAction);
+}
+
+bool FRigVMRemoveNodeAction::Undo(URigVMController* InController)
+{
+	FRigVMActionWrapper InverseWrapper(InverseActionKey);
+	if (!InverseWrapper.GetAction()->Undo(InController))
+	{
+		return false;
+	}
+	return FRigVMBaseAction::Undo(InController);
+}
+
+bool FRigVMRemoveNodeAction::Redo(URigVMController* InController)
+{
+	if (!FRigVMBaseAction::Redo(InController))
+	{
+		return false;
+	}
+	FRigVMActionWrapper InverseWrapper(InverseActionKey);
+	return InverseWrapper.GetAction()->Redo(InController);
+}
+
+FRigVMSetNodeSelectionAction::FRigVMSetNodeSelectionAction()
+{
+
+}
+
+FRigVMSetNodeSelectionAction::FRigVMSetNodeSelectionAction(URigVMGraph* InGraph, TArray<FName> InNewSelection)
+{
+	OldSelection = InGraph->GetSelectNodes();
+	NewSelection = InNewSelection;
+}
+
+bool FRigVMSetNodeSelectionAction::Undo(URigVMController* InController)
+{
+	if(!FRigVMBaseAction::Undo(InController))
+	{
+		return false;
+	}
+	return InController->SetNodeSelection(OldSelection, false);
+}
+
+bool FRigVMSetNodeSelectionAction::Redo(URigVMController* InController)
+{
+	if(!InController->SetNodeSelection(NewSelection, false))
+	{
+		return false;
+	}
+	return FRigVMBaseAction::Redo(InController);
+}
+
+FRigVMSetNodePositionAction::FRigVMSetNodePositionAction(URigVMNode* InNode, const FVector2D& InNewPosition)
+: NodePath(InNode->GetNodePath())
+, OldPosition(InNode->GetPosition())
+, NewPosition(InNewPosition)
+{
+}
+
+bool FRigVMSetNodePositionAction::Merge(const FRigVMBaseAction* Other)
+{
+	if (!FRigVMBaseAction::Merge(Other))
+	{
+		return false;
+	}
+
+	const FRigVMSetNodePositionAction* Action = (const FRigVMSetNodePositionAction*)Other;
+	if (NodePath != Action->NodePath)
+	{
+		return false;
+	}
+
+	NewPosition = Action->NewPosition;
+	return true;
+}
+
+bool FRigVMSetNodePositionAction::Undo(URigVMController* InController)
+{
+	if(!FRigVMBaseAction::Undo(InController))
+	{
+		return false;
+	}
+	return InController->SetNodePositionByName(*NodePath, OldPosition, false);
+}
+
+bool FRigVMSetNodePositionAction::Redo(URigVMController* InController)
+{
+	if(!InController->SetNodePositionByName(*NodePath, NewPosition, false))
+	{
+		return false;
+	}
+	return FRigVMBaseAction::Redo(InController);
+}
+
+FRigVMSetNodeSizeAction::FRigVMSetNodeSizeAction(URigVMNode* InNode, const FVector2D& InNewSize)
+: NodePath(InNode->GetNodePath())
+, OldSize(InNode->GetSize())
+, NewSize(InNewSize)
+{
+}
+
+bool FRigVMSetNodeSizeAction::Merge(const FRigVMBaseAction* Other)
+{
+	if (!FRigVMBaseAction::Merge(Other))
+	{
+		return false;
+	}
+
+	const FRigVMSetNodeSizeAction* Action = (const FRigVMSetNodeSizeAction*)Other;
+	if (NodePath != Action->NodePath)
+	{
+		return false;
+	}
+
+	NewSize = Action->NewSize;
+	return true;
+}
+
+bool FRigVMSetNodeSizeAction::Undo(URigVMController* InController)
+{
+	if(!FRigVMBaseAction::Undo(InController))
+	{
+		return false;
+	}
+	return InController->SetNodeSizeByName(*NodePath, OldSize, false);
+}
+
+bool FRigVMSetNodeSizeAction::Redo(URigVMController* InController)
+{
+	if(!InController->SetNodeSizeByName(*NodePath, NewSize, false))
+	{
+		return false;
+	}
+	return FRigVMBaseAction::Redo(InController);
+}
+
+FRigVMSetNodeColorAction::FRigVMSetNodeColorAction(URigVMNode* InNode, const FLinearColor& InNewColor)
+: NodePath(InNode->GetNodePath())
+, OldColor(InNode->GetNodeColor())
+, NewColor(InNewColor)
+{
+}
+
+bool FRigVMSetNodeColorAction::Merge(const FRigVMBaseAction* Other)
+{
+	if (!FRigVMBaseAction::Merge(Other))
+	{
+		return false;
+	}
+
+	const FRigVMSetNodeColorAction* Action = (const FRigVMSetNodeColorAction*)Other;
+	if (NodePath != Action->NodePath)
+	{
+		return false;
+	}
+
+	NewColor = Action->NewColor;
+	return true;
+}
+
+bool FRigVMSetNodeColorAction::Undo(URigVMController* InController)
+{
+	if(!FRigVMBaseAction::Undo(InController))
+	{
+		return false;
+	}
+	return InController->SetNodeColorByName(*NodePath, OldColor, false);
+}
+
+bool FRigVMSetNodeColorAction::Redo(URigVMController* InController)
+{
+	if(!InController->SetNodeColorByName(*NodePath, NewColor, false))
+	{
+		return false;
+	}
+	return FRigVMBaseAction::Redo(InController);
+}
+
+FRigVMSetCommentTextAction::FRigVMSetCommentTextAction(URigVMCommentNode* InNode, const FString& InNewText)
+: NodePath(InNode->GetNodePath())
+, OldText(InNode->GetCommentText())
+, NewText(InNewText)
+{
+}
+
+bool FRigVMSetCommentTextAction::Undo(URigVMController* InController)
+{
+	if(!FRigVMBaseAction::Undo(InController))
+	{
+		return false;
+	}
+	return InController->SetCommentTextByName(*NodePath, OldText, false);
+}
+
+bool FRigVMSetCommentTextAction::Redo(URigVMController* InController)
+{
+	if(!InController->SetCommentTextByName(*NodePath, NewText, false))
+	{
+		return false;
+	}
+	return FRigVMBaseAction::Redo(InController);
+}
+
+FRigVMSetRerouteCompactnessAction::FRigVMSetRerouteCompactnessAction()
+	: NodePath()
+	, OldShowAsFullNode(false)
+	, NewShowAsFullNode(false)
+{
+
+}
+FRigVMSetRerouteCompactnessAction::FRigVMSetRerouteCompactnessAction(URigVMRerouteNode* InNode, bool InShowAsFullNode)
+	: NodePath(InNode->GetNodePath())
+	, OldShowAsFullNode(InNode->GetShowsAsFullNode())
+	, NewShowAsFullNode(InShowAsFullNode)
+{
+}
+
+bool FRigVMSetRerouteCompactnessAction::Undo(URigVMController* InController)
+{
+	if(!FRigVMBaseAction::Undo(InController))
+	{
+		return false;
+	}
+	return InController->SetRerouteCompactnessByName(*NodePath, OldShowAsFullNode, false);
+}
+
+bool FRigVMSetRerouteCompactnessAction::Redo(URigVMController* InController)
+{
+	if(!InController->SetRerouteCompactnessByName(*NodePath, NewShowAsFullNode, false))
+	{
+		return false;
+	}
+	return FRigVMBaseAction::Redo(InController);
+}
+
+FRigVMRenameVariableAction::FRigVMRenameVariableAction(const FName& InOldVariableName, const FName& InNewVariableName)
+: OldVariableName(InOldVariableName.ToString())
+, NewVariableName(InNewVariableName.ToString())
+{
+}
+
+bool FRigVMRenameVariableAction::Undo(URigVMController* InController)
+{
+	if(!FRigVMBaseAction::Undo(InController))
+	{
+		return false;
+	}
+	return InController->RenameVariable(*NewVariableName, *OldVariableName, false);
+}
+
+bool FRigVMRenameVariableAction::Redo(URigVMController* InController)
+{
+	if(!InController->RenameVariable(*OldVariableName, *NewVariableName, false))
+	{
+		return false;
+	}
+	return FRigVMBaseAction::Redo(InController);
+}
+
+FRigVMRenameParameterAction::FRigVMRenameParameterAction(const FName& InOldParameterName, const FName& InNewParameterName)
+: OldParameterName(InOldParameterName.ToString())
+, NewParameterName(InNewParameterName.ToString())
+{
+}
+
+bool FRigVMRenameParameterAction::Undo(URigVMController* InController)
+{
+	if(!FRigVMBaseAction::Undo(InController))
+	{
+		return false;
+	}
+	return InController->RenameParameter(*NewParameterName, *OldParameterName, false);
+}
+
+bool FRigVMRenameParameterAction::Redo(URigVMController* InController)
+{
+	if(!InController->RenameParameter(*OldParameterName, *NewParameterName, false))
+	{
+		return false;
+	}
+	return FRigVMBaseAction::Redo(InController);
+}
+
+FRigVMSetPinExpansionAction::FRigVMSetPinExpansionAction(URigVMPin* InPin, bool bNewIsExpanded)
+: PinPath(InPin->GetPinPath())
+, OldIsExpanded(InPin->IsExpanded())
+, NewIsExpanded(bNewIsExpanded)
+{
+}
+
+bool FRigVMSetPinExpansionAction::Undo(URigVMController* InController)
+{
+	if(!FRigVMBaseAction::Undo(InController))
+	{
+		return false;
+	}
+	return InController->SetPinExpansion(PinPath, OldIsExpanded, false);
+}
+
+bool FRigVMSetPinExpansionAction::Redo(URigVMController* InController)
+{
+	if(!InController->SetPinExpansion(PinPath, NewIsExpanded, false))
+	{
+		return false;
+	}
+	return FRigVMBaseAction::Redo(InController);
+}
+
+FRigVMSetPinWatchAction::FRigVMSetPinWatchAction(URigVMPin* InPin, bool bNewIsWatched)
+: PinPath(InPin->GetPinPath())
+, OldIsWatched(InPin->RequiresWatch())
+, NewIsWatched(bNewIsWatched)
+{
+}
+
+bool FRigVMSetPinWatchAction::Undo(URigVMController* InController)
+{
+	if(!FRigVMBaseAction::Undo(InController))
+	{
+		return false;
+	}
+	return InController->SetPinIsWatched(PinPath, OldIsWatched, false);
+}
+
+bool FRigVMSetPinWatchAction::Redo(URigVMController* InController)
+{
+	if(!InController->SetPinIsWatched(PinPath, NewIsWatched, false))
+	{
+		return false;
+	}
+	return FRigVMBaseAction::Redo(InController);
+}
+
+FRigVMSetPinDefaultValueAction::FRigVMSetPinDefaultValueAction(URigVMPin* InPin, const FString& InNewDefaultValue)
+: PinPath(InPin->GetPinPath())
+, OldDefaultValue(InPin->GetDefaultValue())
+, NewDefaultValue(InNewDefaultValue)
+{
+}
+
+bool FRigVMSetPinDefaultValueAction::Merge(const FRigVMBaseAction* Other)
+{
+	if (!FRigVMBaseAction::Merge(Other))
+	{
+		return false;
+	}
+
+	const FRigVMSetPinDefaultValueAction* Action = (const FRigVMSetPinDefaultValueAction*)Other;
+	if (PinPath != Action->PinPath)
+	{
+		return false;
+	}
+
+	NewDefaultValue = Action->NewDefaultValue;
+	return true;
+}
+
+bool FRigVMSetPinDefaultValueAction::Undo(URigVMController* InController)
+{
+	if(!FRigVMBaseAction::Undo(InController))
+	{
+		return false;
+	}
+	if (OldDefaultValue.IsEmpty())
+	{
+		return true;
+	}
+	return InController->SetPinDefaultValue(PinPath, OldDefaultValue, true, false);
+}
+
+bool FRigVMSetPinDefaultValueAction::Redo(URigVMController* InController)
+{
+	if (!NewDefaultValue.IsEmpty())
+	{
+		if (!InController->SetPinDefaultValue(PinPath, NewDefaultValue, true, false))
+		{
+			return false;
+		}
+	}
+	return FRigVMBaseAction::Redo(InController);
+}
+
+FRigVMInsertArrayPinAction::FRigVMInsertArrayPinAction(URigVMPin* InArrayPin, int32 InIndex, const FString& InNewDefaultValue)
+: ArrayPinPath(InArrayPin->GetPinPath())
+, Index(InIndex)
+, NewDefaultValue(InNewDefaultValue)
+{
+}
+
+bool FRigVMInsertArrayPinAction::Undo(URigVMController* InController)
+{
+	if(!FRigVMBaseAction::Undo(InController))
+	{
+		return false;
+	}
+	return InController->RemoveArrayPin(FString::Printf(TEXT("%s.%d"), *ArrayPinPath, Index), false);
+}
+
+bool FRigVMInsertArrayPinAction::Redo(URigVMController* InController)
+{
+	if(InController->InsertArrayPin(ArrayPinPath, Index, NewDefaultValue, false).IsEmpty())
+	{
+		return false;
+	}
+	return FRigVMBaseAction::Redo(InController);
+}
+
+FRigVMRemoveArrayPinAction::FRigVMRemoveArrayPinAction(URigVMPin* InArrayElementPin)
+: ArrayPinPath(InArrayElementPin->GetParentPin()->GetPinPath())
+, Index(InArrayElementPin->GetPinIndex())
+, DefaultValue(InArrayElementPin->GetDefaultValue())
+{
+}
+
+bool FRigVMRemoveArrayPinAction::Undo(URigVMController* InController)
+{
+	if (InController->InsertArrayPin(*ArrayPinPath, Index, DefaultValue, false).IsEmpty())
+	{
+		return false;
+	}
+	return FRigVMBaseAction::Undo(InController);
+}
+
+bool FRigVMRemoveArrayPinAction::Redo(URigVMController* InController)
+{
+	if (!FRigVMBaseAction::Redo(InController))
+	{
+		return false;
+	}
+	return InController->RemoveArrayPin(FString::Printf(TEXT("%s.%d"), *ArrayPinPath, Index), false);
+}
+
+FRigVMAddLinkAction::FRigVMAddLinkAction(URigVMPin* InOutputPin, URigVMPin* InInputPin)
+: OutputPinPath(InOutputPin->GetPinPath())
+, InputPinPath(InInputPin->GetPinPath())
+{
+}
+
+bool FRigVMAddLinkAction::Undo(URigVMController* InController)
+{
+	if(!FRigVMBaseAction::Undo(InController))
+	{
+		return false;
+	}
+	return InController->BreakLink(OutputPinPath, InputPinPath, false);
+}
+
+bool FRigVMAddLinkAction::Redo(URigVMController* InController)
+{
+	if(!InController->AddLink(OutputPinPath, InputPinPath, false))
+	{
+		return false;
+	}
+	return FRigVMBaseAction::Redo(InController);
+}
+
+FRigVMBreakLinkAction::FRigVMBreakLinkAction(URigVMPin* InOutputPin, URigVMPin* InInputPin)
+: OutputPinPath(InOutputPin->GetPinPath())
+, InputPinPath(InInputPin->GetPinPath())
+{
+}
+
+bool FRigVMBreakLinkAction::Undo(URigVMController* InController)
+{
+	if(!FRigVMBaseAction::Undo(InController))
+	{
+		return false;
+	}
+	return InController->AddLink(OutputPinPath, InputPinPath, false);
+}
+
+bool FRigVMBreakLinkAction::Redo(URigVMController* InController)
+{
+	if(!InController->BreakLink(OutputPinPath, InputPinPath, false))
+	{
+		return false;
+	}
+	return FRigVMBaseAction::Redo(InController);
+}
\ No newline at end of file
diff --git a/Engine/Source/Developer/RigVMDeveloper/Private/RigVMModel/RigVMLink.cpp b/Engine/Source/Developer/RigVMDeveloper/Private/RigVMModel/RigVMLink.cpp
new file mode 100644
--- /dev/null
+++ b/Engine/Source/Developer/RigVMDeveloper/Private/RigVMModel/RigVMLink.cpp
@@ -0,0 +1,86 @@
+// Copyright Epic Games, Inc. All Rights Reserved.
+
+#include "RigVMModel/RigVMLink.h"
+#include "RigVMModel/RigVMGraph.h"
+
+URigVMGraph* URigVMLink::GetGraph() const
+{
+	return Cast<URigVMGraph>(GetOuter());
+}
+
+void URigVMLink::Serialize(FArchive& Ar)
+{
+	if (Ar.IsLoading())
+	{
+		Ar << SourcePinPath;
+		Ar << TargetPinPath;
+	}
+	else
+	{
+		FString SourcePinPathTemp;
+		if (SourcePin)
+		{
+			if (SourcePin->GetNode())
+			{
+				SourcePinPathTemp = SourcePin->GetPinPath();
+			}
+		}
+		FString TargetPinPathTemp;
+		if (TargetPin)
+		{
+			if (TargetPin->GetNode())
+			{
+				TargetPinPathTemp = TargetPin->GetPinPath();
+			}
+		}
+		Ar << SourcePinPathTemp;
+		Ar << TargetPinPathTemp;
+	}
+}
+
+int32 URigVMLink::GetLinkIndex() const
+{
+	int32 Index = INDEX_NONE;
+	URigVMGraph* Graph = GetGraph();
+	if (Graph != nullptr)
+	{
+		Graph->GetLinks().Find((URigVMLink*)this, Index);
+	}
+	return Index;
+
+}
+
+URigVMPin* URigVMLink::GetSourcePin()
+{
+	if (SourcePin == nullptr)
+	{
+		SourcePin = GetGraph()->FindPin(SourcePinPath);
+	}
+	return SourcePin;
+}
+
+URigVMPin* URigVMLink::GetTargetPin()
+{
+	if (TargetPin == nullptr)
+	{
+		TargetPin = GetGraph()->FindPin(TargetPinPath);
+	}
+	return TargetPin;
+}
+
+FString URigVMLink::GetPinPathRepresentation()
+{
+	return FString::Printf(TEXT("%s -> %s"), *GetSourcePin()->GetPinPath(), *GetTargetPin()->GetPinPath());
+}
+
+void URigVMLink::PrepareForCopy()
+{
+	if (URigVMPin* CurrentSourcePin = GetSourcePin())
+	{
+		SourcePinPath = CurrentSourcePin->GetPinPath();
+	}
+	if (URigVMPin* CurrenTargetPin = GetTargetPin())
+	{
+		TargetPinPath = CurrenTargetPin->GetPinPath();
+	}
+}
diff --git a/Engine/Source/Developer/RigVMDeveloper/Private/RigVMModel/RigVMPin.cpp b/Engine/Source/Developer/RigVMDeveloper/Private/RigVMModel/RigVMPin.cpp
new file mode 100644
--- /dev/null
+++ b/Engine/Source/Developer/RigVMDeveloper/Private/RigVMModel/RigVMPin.cpp
@@ -0,0 +1,753 @@
+// Copyright Epic Games, Inc. All Rights Reserved.
+
+#include "RigVMModel/RigVMPin.h"
+#include "RigVMModel/RigVMNode.h"
+#include "RigVMModel/RigVMGraph.h"
+#include "RigVMModel/RigVMLink.h"
+#include "RigVMCompiler/RigVMCompiler.h"
+#include "RigVMCore/RigVMExecuteContext.h"
+#include "UObject/Package.h"
+#include "Misc/PackageName.h"
+
+URigVMGraph* URigVMInjectionInfo::GetGraph() const
+{
+	return GetPin()->GetGraph();
+}
+
+URigVMPin* URigVMInjectionInfo::GetPin() const
+{
+	return CastChecked<URigVMPin>(GetOuter());
+}
+
+bool URigVMPin::SplitPinPathAtStart(const FString& InPinPath, FString& LeftMost, FString& Right)
+{
+	return InPinPath.Split(TEXT("."), &LeftMost, &Right, ESearchCase::IgnoreCase, ESearchDir::FromStart);
+}
+
+bool URigVMPin::SplitPinPathAtEnd(const FString& InPinPath, FString& Left, FString& RightMost)
+{
+	return InPinPath.Split(TEXT("."), &Left, &RightMost, ESearchCase::IgnoreCase, ESearchDir::FromEnd);
+}
+
+bool URigVMPin::SplitPinPath(const FString& InPinPath, TArray<FString>& Parts)
+{
+	int32 OriginalPartsCount = Parts.Num();
+	FString PinPathRemaining = InPinPath;
+	FString Left, Right;
+	while(SplitPinPathAtStart(PinPathRemaining, Left, Right))
+	{
+		Parts.Add(Left);
+		Left.Empty();
+		PinPathRemaining = Right;
+	}
+
+	if (!Right.IsEmpty())
+	{
+		Parts.Add(Right);
+	}
+
+	return Parts.Num() > OriginalPartsCount;
+}
+
+FString URigVMPin::JoinPinPath(const FString& Left, const FString& Right)
+{
+	ensure(!Left.IsEmpty() && !Right.IsEmpty());
+	return Left + TEXT(".") + Right;
+}
+
+FString URigVMPin::JoinPinPath(const TArray<FString>& InParts)
+{
+	if (InParts.Num() == 0)
+	{
+		return FString();
+	}
+
+	FString Result = InParts[0];
+	for (int32 PartIndex = 1; PartIndex < InParts.Num(); PartIndex++)
+	{
+		Result += TEXT(".") + InParts[PartIndex];
+	}
+
+	return Result;
+}
+
+URigVMPin::URigVMPin()
+	: Direction(ERigVMPinDirection::Invalid)
+	, bIsExpanded(false)
+	, bIsConstant(false)
+	, bRequiresWatch(false)
+	, CPPType(FString())
+	, CPPTypeObject(nullptr)
+	, CPPTypeObjectPath(NAME_None)
+	, DefaultValue(FString())
+{
+}
+
+FString URigVMPin::GetPinPath() const
+{
+	URigVMPin* ParentPin = GetParentPin();
+	if(ParentPin)
+	{
+		return FString::Printf(TEXT("%s.%s"), *ParentPin->GetPinPath(), *GetName());
+	}
+
+	URigVMNode* Node = GetNode();
+	if (Node == nullptr)
+	{
+		return FString();
+	}
+	return FString::Printf(TEXT("%s.%s"), *Node->GetNodePath(), *GetName());
+}
+
+FString URigVMPin::GetSegmentPath() const
+{
+	URigVMPin* ParentPin = GetParentPin();
+	if (ParentPin)
+	{
+		FString ParentSegmentPath = ParentPin->GetSegmentPath();
+		if (ParentSegmentPath.IsEmpty())
+		{
+			return GetName();
+		}
+		return FString::Printf(TEXT("%s.%s"), *ParentSegmentPath, *GetName());
+	}
+	return FString();
+}
+
+FName URigVMPin::GetDisplayName() const
+{
+	if (DisplayName == NAME_None)
+	{
+		return GetFName();
+	}
+
+	if (InjectionInfos.Num() > 0)
+	{
+		FString ProcessedDisplayName = DisplayName.ToString();
+
+		for (URigVMInjectionInfo* Injection : InjectionInfos)
+		{
+			if (TSharedPtr<FStructOnScope> DefaultStructScope = Injection->StructNode->ConstructStructInstance())
+			{
+				FRigVMStruct* DefaultStruct = (FRigVMStruct*)DefaultStructScope->GetStructMemory();
+				ProcessedDisplayName = DefaultStruct->ProcessPinLabelForInjection(ProcessedDisplayName);
+			}
+		}
+
+		return *ProcessedDisplayName;
+	}
+
+	return DisplayName;
+}
+
+ERigVMPinDirection URigVMPin::GetDirection() const
+{
+	return Direction;
+}
+
+bool URigVMPin::IsExpanded() const
+{
+	return bIsExpanded;
+}
+
+bool URigVMPin::IsConstant() const
+{
+	if (IsArrayElement())
+	{
+		return GetParentPin()->IsConstant();
+	}
+	return bIsConstant;
+}
+
+bool URigVMPin::RequiresWatch() const
+{
+	return bRequiresWatch;
+}
+
+bool URigVMPin::IsStruct() const
+{
+	if (IsArray())
+	{
+		return false;
+	}
+	return GetScriptStruct() != nullptr;
+}
+
+bool URigVMPin::IsStructMember() const
+{
+	URigVMPin* ParentPin = GetParentPin();
+	if (ParentPin == nullptr)
+	{
+		return false;
+	}
+	return ParentPin->IsStruct();
+}
+
+bool URigVMPin::IsArray() const
+{
+	return CPPType.StartsWith(TEXT("TArray"));
+}
+
+bool URigVMPin::IsArrayElement() const
+{
+	URigVMPin* ParentPin = GetParentPin();
+	if (ParentPin == nullptr)
+	{
+		return false;
+	}
+	return ParentPin->IsArray();
+}
+
+int32 URigVMPin::GetPinIndex() const
+{
+	int32 Index = INDEX_NONE;
+
+	URigVMPin* ParentPin = GetParentPin();
+	if (ParentPin != nullptr)
+	{
+		ParentPin->GetSubPins().Find((URigVMPin*)this, Index);
+	}
+	else
+	{
+		URigVMNode* Node = GetNode();
+		if (Node != nullptr)
+		{
+			Node->GetPins().Find((URigVMPin*)this, Index);
+		}
+	}
+	return Index;
+}
+
+void URigVMPin::SetNameFromIndex()
+{
+	LowLevelRename(*FString::FormatAsNumber(GetPinIndex()));
+}
+
+int32 URigVMPin::GetArraySize() const
+{
+	return SubPins.Num();
+}
+
+FString URigVMPin::GetCPPType() const
+{
+	return CPPType;
+}
+
+FString URigVMPin::GetArrayElementCppType() const
+{
+	if (!IsArray())
+	{
+		return FString();
+	}
+	return CPPType.Mid(7, CPPType.Len() - 8);
+}
+
+bool URigVMPin::IsStringType() const
+{
+	return CPPType.Equals(TEXT("FString")) || CPPType.Equals(TEXT("FName"));
+}
+
+bool URigVMPin::IsExecuteContext() const
+{
+	if (UScriptStruct* ScriptStruct = GetScriptStruct())
+	{
+		if (ScriptStruct->IsChildOf(FRigVMExecuteContext::StaticStruct()))
+		{
+			return true;
+		}
+	}
+	return false;
+}
+
+
+FString URigVMPin::GetDefaultValue() const
+{
+	/*
+	URigVMPin* PinForDefaultValue = URigVMCompiler::FollowPinForDefaultValue((URigVMPin*)this);
+	if (PinForDefaultValue != this)
+	{
+		return PinForDefaultValue->GetDefaultValue();
+	}
+	*/
+
+	if (IsArray())
+	{
+		if (SubPins.Num() > 0)
+		{
+			TArray<FString> ElementDefaultValues;
+			for (URigVMPin* SubPin : SubPins)
+			{
+				FString ElementDefaultValue = SubPin->GetDefaultValue();
+				if (SubPin->IsStringType())
+				{
+					ElementDefaultValue = TEXT("\"") + ElementDefaultValue + TEXT("\"");
+				}
+				else if (ElementDefaultValue.IsEmpty() || ElementDefaultValue == TEXT("()"))
+				{
+					continue;
+				}
+				ElementDefaultValues.Add(ElementDefaultValue);
+			}
+			if (ElementDefaultValues.Num() == 0)
+			{
+				return TEXT("()");
+			}
+			return FString::Printf(TEXT("(%s)"), *FString::Join(ElementDefaultValues, TEXT(",")));
+		}
+
+		return DefaultValue.IsEmpty() ? TEXT("()") : DefaultValue;
+	}
+	else if (IsStruct())
+	{
+		if (SubPins.Num() > 0)
+		{
+			TArray<FString> MemberDefaultValues;
+			for (URigVMPin* SubPin : SubPins)
+			{
+				FString MemberDefaultValue = SubPin->GetDefaultValue();
+				if (SubPin->IsStringType())
+				{
+					MemberDefaultValue = TEXT("\"") + MemberDefaultValue + TEXT("\"");
+				}
+				else if (MemberDefaultValue.IsEmpty() || MemberDefaultValue == TEXT("()"))
+				{
+					continue;
+				}
+				MemberDefaultValues.Add(FString::Printf(TEXT("%s=%s"), *SubPin->GetName(), *MemberDefaultValue));
+			}
+			if (MemberDefaultValues.Num() == 0)
+			{
+				return TEXT("()");
+			}
+			return FString::Printf(TEXT("(%s)"), *FString::Join(MemberDefaultValues, TEXT(",")));
+		}
+
+		return DefaultValue.IsEmpty() ? TEXT("()") : DefaultValue;
+	}
+
+	return DefaultValue;
+}
+
+FName URigVMPin::GetCustomWidgetName() const
+{
+	if (IsArrayElement())
+	{
+		return GetParentPin()->GetCustomWidgetName();
+	}
+	return CustomWidgetName;
+}
+
+FText URigVMPin::GetToolTipText() const
+{
+	if(URigVMNode* Node = GetNode())
+	{
+		return Node->GetToolTipTextForPin(this);
+	}
+	return FText();
+}
+
+UObject* URigVMPin::FindObjectFromCPPTypeObjectPath(const FString& InObjectPath)
+{
+	if (InObjectPath.IsEmpty())
+	{
+		return nullptr;
+	}
+
+	if (InObjectPath == FName(NAME_None).ToString())
+	{
+		return nullptr;
+	}
+
+	// we do this to avoid ambiguous searches for 
+	// common names such as "transform" or "vector"
+	UPackage* Package = ANY_PACKAGE;
+	FString PackageName;
+	FString CPPTypeObjectName = InObjectPath;
+	if (InObjectPath.Split(TEXT("."), &PackageName, &CPPTypeObjectName))
+	{
+		Package = FindPackage(nullptr, *PackageName);
+	}
+	
+	if (UObject* ObjectWithinPackage = FindObject<UObject>(Package, *CPPTypeObjectName))
+	{
+		return ObjectWithinPackage;
+	}
+
+	return FindObject<UObject>(ANY_PACKAGE, *InObjectPath);
+}
+
+void URigVMPin::UpdateCPPTypeObjectIfRequired() const
+{
+	if (CPPTypeObject == nullptr)
+	{
+		if (CPPTypeObjectPath != NAME_None)
+		{
+			URigVMPin* MutableThis = (URigVMPin*)this;
+			MutableThis->CPPTypeObject = FindObjectFromCPPTypeObjectPath(CPPTypeObjectPath.ToString());
+		}
+	}
+}
+
+UObject* URigVMPin::GetCPPTypeObject() const
+{
+	UpdateCPPTypeObjectIfRequired();
+	return CPPTypeObject;
+}
+
+UScriptStruct* URigVMPin::GetScriptStruct() const
+{
+	return Cast<UScriptStruct>(GetCPPTypeObject());
+}
+
+UEnum* URigVMPin::GetEnum() const
+{
+	return Cast<UEnum>(GetCPPTypeObject());
+}
+
+URigVMPin* URigVMPin::GetParentPin() const
+{
+	return Cast<URigVMPin>(GetOuter());
+}
+
+URigVMPin* URigVMPin::GetRootPin() const
+{
+	URigVMPin* ParentPin = GetParentPin();
+	if (ParentPin == nullptr)
+	{
+		return const_cast<URigVMPin*>(this);
+	}
+	return ParentPin->GetRootPin();
+}
+
+URigVMPin* URigVMPin::GetPinForLink() const
+{
+	URigVMPin* RootPin = GetRootPin();
+
+	if (RootPin->InjectionInfos.Num() == 0)
+	{
+		return const_cast<URigVMPin*>(this);
+	}
+
+	URigVMPin* PinForLink =
+		((Direction == ERigVMPinDirection::Input) || (Direction == ERigVMPinDirection::IO)) ?
+		RootPin->InjectionInfos.Last()->InputPin :
+		RootPin->InjectionInfos.Last()->OutputPin;
+
+	if (RootPin != this)
+	{
+		FString SegmentPath = GetSegmentPath();
+		return PinForLink->FindSubPin(SegmentPath);
+	}
+
+	return PinForLink;
+}
+
+URigVMPin* URigVMPin::GetOriginalPinFromInjectedNode() const
+{
+	if (URigVMInjectionInfo* Injection = GetNode()->GetInjectionInfo())
+	{
+		URigVMPin* RootPin = GetRootPin();
+		URigVMPin* OriginalPin = nullptr;
+		if (Injection->bInjectedAsInput && Injection->InputPin == RootPin)
+		{
+			TArray<URigVMPin*> LinkedPins = Injection->OutputPin->GetLinkedTargetPins();
+			if (LinkedPins.Num() == 1)
+			{
+				OriginalPin = LinkedPins[0]->GetOriginalPinFromInjectedNode();
+			}
+		}
+		else if (!Injection->bInjectedAsInput && Injection->OutputPin == RootPin)
+		{
+			TArray<URigVMPin*> LinkedPins = Injection->InputPin->GetLinkedSourcePins();
+			if (LinkedPins.Num() == 1)
+			{
+				OriginalPin = LinkedPins[0]->GetOriginalPinFromInjectedNode();
+			}
+		}
+
+		if (OriginalPin)
+		{
+			if (this != RootPin)
+			{
+				OriginalPin = OriginalPin->FindSubPin(GetSegmentPath());
+			}
+			return OriginalPin;
+		}
+	}
+
+	return const_cast<URigVMPin*>(this);
+}
+
+const TArray<URigVMPin*>& URigVMPin::GetSubPins() const
+{
+	return SubPins;
+}
+
+URigVMPin* URigVMPin::FindSubPin(const FString& InPinPath) const
+{
+	FString Left, Right;
+	if (!URigVMPin::SplitPinPathAtStart(InPinPath, Left, Right))
+	{
+		Left = InPinPath;
+	}
+
+	for (URigVMPin* Pin : SubPins)
+	{
+		if (Pin->GetName() == Left)
+		{
+			if (Right.IsEmpty())
+			{
+				return Pin;
+			}
+			return Pin->FindSubPin(Right);
+		}
+	}
+	return nullptr;
+}
+
+bool URigVMPin::IsLinkedTo(URigVMPin* InPin) const
+{
+	for (URigVMLink* Link : Links)
+	{
+		if (Link->GetSourcePin() == InPin || Link->GetTargetPin() == InPin)
+		{
+			return true;
+		}
+	}
+	return false;
+}
+
+const TArray<URigVMLink*>& URigVMPin::GetLinks() const
+{
+	return Links;
+}
+
+TArray<URigVMPin*> URigVMPin::GetLinkedSourcePins(bool bRecursive) const
+{
+	TArray<URigVMPin*> Pins;
+	for (URigVMLink* Link : Links)
+	{
+		if (Link->GetTargetPin() == this)
+		{
+			Pins.AddUnique(Link->GetSourcePin());
+		}
+	}
+
+	if (bRecursive)
+	{
+		for (URigVMPin* SubPin : SubPins)
+		{
+			Pins.Append(SubPin->GetLinkedSourcePins(bRecursive));
+		}
+	}
+
+	return Pins;
+}
+
+TArray<URigVMPin*> URigVMPin::GetLinkedTargetPins(bool bRecursive) const
+{
+	TArray<URigVMPin*> Pins;
+	for (URigVMLink* Link : Links)
+	{
+		if (Link->GetSourcePin() == this)
+		{
+			Pins.AddUnique(Link->GetTargetPin());
+		}
+	}
+
+	if (bRecursive)
+	{
+		for (URigVMPin* SubPin : SubPins)
+		{
+			Pins.Append(SubPin->GetLinkedTargetPins(bRecursive));
+		}
+	}
+
+	return Pins;
+}
+
+TArray<URigVMLink*> URigVMPin::GetSourceLinks(bool bRecursive) const
+{
+	TArray<URigVMLink*> Results;
+	for (URigVMLink* Link : Links)
+	{
+		if (Link->GetTargetPin() == this)
+		{
+			Results.Add(Link);
+		}
+	}
+
+	if (bRecursive)
+	{
+		for (URigVMPin* SubPin : SubPins)
+		{
+			Results.Append(SubPin->GetSourceLinks(bRecursive));
+		}
+	}
+
+	return Results;
+}
+
+TArray<URigVMLink*> URigVMPin::GetTargetLinks(bool bRecursive) const
+{
+	TArray<URigVMLink*> Results;
+	for (URigVMLink* Link : Links)
+	{
+		if (Link->GetSourcePin() == this)
+		{
+			Results.Add(Link);
+		}
+	}
+
+	if (bRecursive)
+	{
+		for (URigVMPin* SubPin : SubPins)
+		{
+			Results.Append(SubPin->GetTargetLinks(bRecursive));
+		}
+	}
+
+	return Results;
+}
+
+URigVMNode* URigVMPin::GetNode() const
+{
+	URigVMPin* ParentPin = GetParentPin();
+	if (ParentPin)
+	{
+		return ParentPin->GetNode();
+	}
+
+	URigVMNode* Node = Cast<URigVMNode>(GetOuter());
+	if(Node)
+	{
+		return Node;
+	}
+
+	return nullptr;
+}
+
+URigVMGraph* URigVMPin::GetGraph() const
+{
+	URigVMNode* Node = GetNode();
+	if(Node)
+	{
+		return Node->GetGraph();
+	}
+
+	return nullptr;
+}
+
+bool URigVMPin::CanLink(URigVMPin* InSourcePin, URigVMPin* InTargetPin, FString* OutFailureReason)
+{
+	if (InSourcePin == nullptr || InTargetPin == nullptr)
+	{
+		if (OutFailureReason)
+		{
+			*OutFailureReason = TEXT("One of the pins is nullptr.");
+		}
+		return false;
+	}
+
+	if (InSourcePin == InTargetPin)
+	{
+		if (OutFailureReason)
+		{
+			*OutFailureReason = TEXT("Source and target pins are the same.");
+		}
+		return false;
+	}
+
+	if (InSourcePin->GetNode() == InTargetPin->GetNode())
+	{
+		if (OutFailureReason)
+		{
+			*OutFailureReason = TEXT("Source and target pins are on the same node.");
+		}
+		return false;
+	}
+
+	if (InSourcePin->GetGraph() != InTargetPin->GetGraph())
+	{
+		if (OutFailureReason)
+		{
+			*OutFailureReason = TEXT("Source and target pins are in different graphs.");
+		}
+		return false;
+	}
+
+	if (InSourcePin->Direction != ERigVMPinDirection::Output &&
+		InSourcePin->Direction != ERigVMPinDirection::IO)
+	{
+		if (OutFailureReason)
+		{
+			*OutFailureReason = TEXT("Source pin is not an output.");
+		}
+		return false;
+	}
+
+	if (InTargetPin->Direction != ERigVMPinDirection::Input &&
+		InTargetPin->Direction != ERigVMPinDirection::IO)
+	{
+		if (OutFailureReason)
+		{
+			*OutFailureReason = TEXT("Target pin is not an input.");
+		}
+		return false;
+	}
+
+	if (InTargetPin->IsConstant() && !InSourcePin->IsConstant())
+	{
+		if (OutFailureReason)
+		{
+			*OutFailureReason = TEXT("Cannot connect non-constants to constants.");
+		}
+		return false;
+	}
+
+	if (InSourcePin->CPPType != InTargetPin->CPPType)
+	{
+		if (OutFailureReason)
+		{
+			*OutFailureReason = TEXT("Source and target pin types are not compatible.");
+		}
+		return false;
+	}
+
+	// only allow to link to specified input / output pins on an injected node
+	if (URigVMInjectionInfo* SourceInjectionInfo = InSourcePin->GetNode()->GetInjectionInfo())
+	{
+		if (SourceInjectionInfo->OutputPin != InSourcePin->GetRootPin())
+		{
+			if (OutFailureReason)
+			{
+				*OutFailureReason = TEXT("Cannot link to a non-exposed pin on an injected node.");
+			}
+			return false;
+		}
+	}
+
+	// only allow to link to specified input / output pins on an injected node
+	if (URigVMInjectionInfo* TargetInjectionInfo = InTargetPin->GetNode()->GetInjectionInfo())
+	{
+		if (TargetInjectionInfo->InputPin != InTargetPin->GetRootPin())
+		{
+			if (OutFailureReason)
+			{
+				*OutFailureReason = TEXT("Cannot link to a non-exposed pin on an injected node.");
+			}
+			return false;
+		}
+	}
+
+	if (InSourcePin->IsLinkedTo(InTargetPin))
+	{
+		if (OutFailureReason)
+		{
+			*OutFailureReason = TEXT("Source and target pins are already connected.");
+		}
+		return false;
+	}
+
+	return true;
+}
diff --git a/Engine/Source/Developer/RigVMDeveloper/Public/RigVMCompiler/RigVMAST.h b/Engine/Source/Developer/RigVMDeveloper/Public/RigVMCompiler/RigVMAST.h
new file mode 100644
--- /dev/null
+++ b/Engine/Source/Developer/RigVMDeveloper/Public/RigVMCompiler/RigVMAST.h
@@ -0,0 +1,994 @@
+// Copyright Epic Games, Inc. All Rights Reserved.
+
+#pragma once
+
+#include "CoreMinimal.h"
+#include "RigVMModel/RigVMPin.h"
+#include "RigVMAST.generated.h"
+
+class FRigVMParserAST;
+class FRigVMBlockExprAST;
+class FRigVMEntryExprAST;
+class FRigVMCallExternExprAST;
+class FRigVMNoOpExprAST;
+class FRigVMVarExprAST;
+class FRigVMLiteralExprAST;
+class FRigVMAssignExprAST;
+class FRigVMCopyExprAST;
+class FRigVMCachedValueExprAST;
+class FRigVMExitExprAST;
+
+class URigVMPin;
+class URigVMNode;
+class URigVMLink;
+class URigVMGraph;
+
+/*
+ * Base class for an expression within an abstract syntax tree.
+ * The base implements parent / child relationships as well as a
+ * simple typing system.
+ * An expression is a multi child / multi parent element of a
+ * directed tree (there can be no cycles).
+ * Expressions can only be constructed by an AST parser, and
+ * are also memory-owned by the parser.
+ */
+class RIGVMDEVELOPER_API FRigVMExprAST
+{
+public:
+
+	// Simple enum for differentiating expression types.
+	enum EType
+	{
+		Block,
+		Entry,
+		CallExtern,
+		NoOp,
+		Var,
+		Literal,
+		Assign,
+		Copy,
+		CachedValue,
+		Exit,
+		Invalid
+	};
+
+	// virtual destructor
+	virtual ~FRigVMExprAST() {}
+
+	// disable copy constructor
+	FRigVMExprAST(const FRigVMExprAST&) = delete;
+
+	// returns the parser this expression is owned by
+	// @return the parser this expression is owned by 
+	const FRigVMParserAST* GetParser() const { return ParserPtr; }
+
+	// returns the name of the expression (can be NAME_None)
+	// @return the name of the expression
+	FName GetName() const { return Name; }
+
+	// returns the exact type of the expression
+	// @return the exact type of the expression
+	EType GetType() const { return Type; }
+
+	// returns the name of the expression's type
+	// @return the name of the expression's type
+	FName GetTypeName() const;
+
+	// provides type checking for inherited types
+	// @param InType the type to check against
+	// @return true if this expression is of a given type
+	virtual bool IsA(EType InType) const = 0;
+
+	// returns the index of this expression within the parser's storage
+	// @return the index of this expression within the parser's storage
+	int32 GetIndex() const { return Index; }
+
+	// returns the parent of this expression
+	// @return the parent of this expression
+	const FRigVMExprAST* GetParent() const;
+
+	// returns the block of this expression
+	// @return the block of this expression
+	const FRigVMBlockExprAST* GetBlock() const;
+
+	// returns the root / top level block of this expression
+	// @return the root / top level block of this expression
+	const FRigVMBlockExprAST* GetRootBlock() const;
+
+
+	// returns the number of children of this expression
+	// @return the number of children of this expression
+	int32 NumChildren() const { return Children.Num(); }
+
+	// accessor operator for a given child
+	// @param InIndex the index of the child to retrieve (bound = NumChildren() - 1)
+	// @return the child at the given index
+	FORCEINLINE const FRigVMExprAST* operator[](int32 InIndex) const { return Children[InIndex]; }
+
+	// begin iterator accessor for the children
+	FORCEINLINE TArray<FRigVMExprAST*>::RangedForConstIteratorType begin() const { return Children.begin(); }
+	// end iterator accessor for the children
+	FORCEINLINE TArray<FRigVMExprAST*>::RangedForConstIteratorType end() const { return Children.end(); }
+
+	// templated getter to retrieve a child with a given index
+	// type checking will occur within the ::To method and raise
+	// @param InIndex the index of the child to retrieve
+	// @return the child at a given index cast to the provided class
+	template<class ObjectType>
+	const ObjectType* ChildAt(int32 InIndex) const
+	{
+		return Children[InIndex]->To<ObjectType>();
+	}
+
+	// getter to retrieve a child with a given index
+	// @param InIndex the index of the child to retrieve
+	// @return the child at a given index
+	const FRigVMExprAST* ChildAt(int32 InIndex) const
+	{
+		return Children[InIndex];
+	}
+
+	// returns the number of parents of this expression
+	// @return the number of parents of this expression
+	int32 NumParents() const { return Parents.Num(); }
+
+	// templated getter to retrieve a parent with a given index
+	// type checking will occur within the ::To method and raise
+	// @param InIndex the index of the parent to retrieve
+	// @return the parent at a given index cast to the provided class
+	template<class ObjectType>
+	const ObjectType* ParentAt(int32 InIndex) const
+	{
+		return Parents[InIndex]->To<ObjectType>();
+	}
+
+	// getter to retrieve a parent with a given index
+	// @param InIndex the index of the parent to retrieve
+	// @return the parent at a given index
+	const FRigVMExprAST* ParentAt(int32 InIndex) const
+	{
+		return Parents[InIndex];
+	}
+
+	// const templated cast for casting between
+	// different expression types.
+	// specializations below are used for type checking
+	// @return this object cast to the provided class
+	template<class ObjectType>
+	const ObjectType* To() const
+	{
+		checkNoEntry();
+		return nullptr;
+	}
+
+	// templated cast for casting between
+	// different expression types.
+	// specializations below are used for type checking
+	// @return this object cast to the provided class
+	template<class ObjectType>
+	ObjectType* To()
+	{
+		return (ObjectType*)this;
+	}
+
+protected:
+
+	// default constructor (protected so that only parser can access it)
+	FRigVMExprAST(const FRigVMParserAST* InParser, EType InType = EType::Invalid);
+
+	// adds a parent to this expression
+	// this in consequence also adds this as a child to the parent
+	// @param InParent the parent to add
+	void AddParent(FRigVMExprAST* InParent);
+
+	// removes a parent from this expression
+	// this in consequence also removes this as a child from the parent
+	// @param InParent the parent to remove
+	void RemoveParent(FRigVMExprAST* InParent);
+
+	// removes a child from this expression
+	// this in consequence also removes this as a parent from the child
+	// @param InChild the child to remove
+	void RemoveChild(FRigVMExprAST* InChild);
+
+	// replaces a parent of this expression with a new one
+	// @param InCurrentParent the current parent to replace
+	// @param InNewParent the new parent to replace it with
+	void ReplaceParent(FRigVMExprAST* InCurrentParent, FRigVMExprAST* InNewParent);
+
+	// replaces a child of this expression with a new one
+	// @param InCurrentChild the current child to replace
+	// @param InNewChild the new child to replace it with
+	void ReplaceChild(FRigVMExprAST* InCurrentChild, FRigVMExprAST* InNewChild);
+
+private:
+
+	// returns a string containing an indented tree structure
+	// for debugging purposes. this is only used by the parser
+	// @param InPrefix the prefix to use for indentation
+	// @return the text representation of this part of the tree
+	virtual FString DumpText(const FString& InPrefix = FString()) const;
+
+	// returns a string containing a dot file notation
+	// for debugging purposes. this is only used by the parser
+	// @param OutExpressionDefined a bool map to keep track of which expressions have been processed yet
+	// @param InPrefix the prefix to use for indentation
+	// @return the text representation of this part of the tree
+	virtual FString DumpDot(TArray<bool>& OutExpressionDefined, const FString& InPrefix = FString()) const;
+
+	FName Name;
+	EType Type;
+	int32 Index;
+	const FRigVMParserAST* ParserPtr;
+	TArray<FRigVMExprAST*> Parents;
+	TArray<FRigVMExprAST*> Children;
+
+	friend class FRigVMParserAST;
+};
+
+// specialized cast for type checkiFRigVMAssignExprASTng
+	// for a Block / FRigVMBlockExprAST expression
+	// will raise if types are not compatible
+	// @return this expression cast to FRigVMBlockExprAST
+template<>
+FORCEINLINE const FRigVMBlockExprAST* FRigVMExprAST::To() const
+{
+	ensure(IsA(EType::Block));
+	return (const FRigVMBlockExprAST*)this;
+}
+
+// specialized cast for type checking
+// for a Entry / FRigVMEntryExprAST expression
+// will raise if types are not compatible
+// @return this expression cast to FRigVMEntryExprAST
+template<>
+FORCEINLINE const FRigVMEntryExprAST* FRigVMExprAST::To() const
+{
+	ensure(IsA(EType::Entry));
+	return (const FRigVMEntryExprAST*)this;
+}
+
+// specialized cast for type checking
+// for a CallExtern / FRigVMCallExternExprAST expression
+// will raise if types are not compatible
+// @return this expression cast to FRigVMCallExternExprAST
+template<>
+FORCEINLINE const FRigVMCallExternExprAST* FRigVMExprAST::To() const
+{
+	ensure(IsA(EType::CallExtern));
+	return (const FRigVMCallExternExprAST*)this;
+}
+
+// specialized cast for type checking
+// for a NoOp / FRigVMNoOpExprAST expression
+// will raise if types are not compatible
+// @return this expression cast to FRigVMNoOpExprASTo 
+template<>
+FORCEINLINE const FRigVMNoOpExprAST* FRigVMExprAST::To() const
+{
+	ensure(IsA(EType::NoOp));
+	return (const FRigVMNoOpExprAST*)this;
+}
+
+// specialized cast for type checking
+// for a Var / FRigVMVarExprAST expression
+// will raise if types are not compatible
+// @return this expression cast to FRigVMVarExprAST
+template<>
+FORCEINLINE const FRigVMVarExprAST* FRigVMExprAST::To() const
+{
+	ensure(IsA(EType::Var));
+	return (const FRigVMVarExprAST*)this;
+}
+
+// specialized cast for type checking
+// for a Literal / FRigVMLiteralExprAST expression
+// will raise if types are not compatible
+// @return this expression cast to FRigVMLiteralExprAST
+template<>
+FORCEINLINE const FRigVMLiteralExprAST* FRigVMExprAST::To() const
+{
+	ensure(IsA(EType::Literal));
+	return (const FRigVMLiteralExprAST*)this;
+}
+
+// specialized cast for type checking
+// for a Assign / FRigVMAssignExprAST expression
+// will raise if types are not compatible
+// @return this expression cast to FRigVMAssignExprAST
+template<>
+FORCEINLINE const FRigVMAssignExprAST* FRigVMExprAST::To() const
+{
+	ensure(IsA(EType::Assign));
+	return (const FRigVMAssignExprAST*)this;
+}
+
+// specialized cast for type checking
+// for a Copy / FRigVMCopyExprAST expression
+// will raise if types are not compatible
+// @return this expression cast to FRigVMCopyExprASTo 
+template<>
+FORCEINLINE const FRigVMCopyExprAST* FRigVMExprAST::To() const
+{
+	ensure(IsA(EType::Copy));
+	return (const FRigVMCopyExprAST*)this;
+}
+
+// specialized cast for type checking
+// for a CachedValue / FRigVMCachedValueExprAST expression
+// will raise if types are not compatible
+// @return this expression cast to FRigVMCachedValueExprAST
+template<>
+FORCEINLINE const FRigVMCachedValueExprAST* FRigVMExprAST::To() const
+{
+	ensure(IsA(EType::CachedValue));
+	return (const FRigVMCachedValueExprAST*)this;
+}
+
+// specialized cast for type checking
+// for a Exit / FRigVMExitExprAST expression
+// will raise if types are not compatible
+// @return this expression cast to FRigVMExitExprASTo 
+template<>
+FORCEINLINE const FRigVMExitExprAST* FRigVMExprAST::To() const
+{
+	ensure(IsA(EType::Exit));
+	return (const FRigVMExitExprAST*)this;
+}
+/*
+ * An abstract syntax tree block expression represents a sequence
+ * of child expressions to be executed in order.
+ * In C++ a block is represented by the curly braces { expr1, expr2, ...}.
+ */ 
+class RIGVMDEVELOPER_API FRigVMBlockExprAST : public FRigVMExprAST
+{
+public:
+
+	// virtual destructor
+	virtual ~FRigVMBlockExprAST() {}
+
+	// disable copy constructor
+	FRigVMBlockExprAST(const FRigVMBlockExprAST&) = delete;
+
+	// returns true if this block needs to execute
+	// this is determined by the block containing an entry expression
+	// @return true if this block needs to execute
+	bool ShouldExecute() const;
+
+	// returns true if this block contains an entry expression
+	// @return true if this block contains an entry expression
+	bool ContainsEntry() const;
+
+	// returns true if this block contains a given expression
+	// @param InExpression the expression to check
+	// @return true if this block contains a given expression
+	bool Contains(const FRigVMExprAST* InExpression) const;
+
+	// overload of the type checking mechanism
+	virtual bool IsA(EType InType) const override
+	{
+		return InType == EType::Block;
+	};
+
+protected:
+
+	// default constructor (protected so that only parser can access it)
+	FRigVMBlockExprAST(const FRigVMParserAST* InParser, EType InType = EType::Block)
+		: FRigVMExprAST(InParser, InType)
+	{}
+
+private:
+
+	friend class FRigVMParserAST;
+};
+
+/*
+ * An abstract syntax tree node expression represents any expression
+ * which references a node from the RigVM model.
+ */
+class RIGVMDEVELOPER_API FRigVMNodeExprAST : public FRigVMBlockExprAST
+{
+public:
+
+	// virtual destructor
+	virtual ~FRigVMNodeExprAST() {}
+
+	// disable copy constructor
+	FRigVMNodeExprAST(const FRigVMNodeExprAST&) = delete;
+
+	// returns the node from the model this expression is referencing
+	// @return the node from the model this expression is referencing
+	URigVMNode* GetNode() const { return Node; }
+
+protected:
+
+	// default constructor (protected so that only parser can access it)
+	FRigVMNodeExprAST(const FRigVMParserAST* InParser, EType InType, URigVMNode* InNode)
+		: FRigVMBlockExprAST(InParser, InType)
+		, Node(InNode)
+	{}
+
+
+	// overload of the type checking mechanism
+	virtual bool IsA(EType InType) const override
+	{
+		return false;
+	};
+
+private:
+
+	URigVMNode* Node;
+	friend class FRigVMParserAST;
+};
+
+/*
+ * An abstract syntax tree entry expression represents an entry point
+ * for a function or an event in an event graph.
+ * In C++ the entry point is the declaration: void main(...);
+ */
+class RIGVMDEVELOPER_API FRigVMEntryExprAST : public FRigVMNodeExprAST
+{
+public:
+
+	// virtual destructor
+	virtual ~FRigVMEntryExprAST() {}
+
+	// disable copy constructor
+	FRigVMEntryExprAST(const FRigVMEntryExprAST&) = delete;
+
+	// overload of the type checking mechanism
+	virtual bool IsA(EType InType) const override
+	{
+		if(FRigVMBlockExprAST::IsA(InType))
+		{
+			return true;
+		}
+		return InType == EType::Entry;
+	};
+
+	// returns the name of the entry / event
+	// @return the name of the entry / event
+	FName GetEventName() const;
+
+protected:
+
+	// default constructor (protected so that only parser can access it)
+	FRigVMEntryExprAST(const FRigVMParserAST* InParser, URigVMNode* InNode)
+		: FRigVMNodeExprAST(InParser, EType::Entry, InNode)
+	{}
+
+
+private:
+
+	friend class FRigVMParserAST;
+};
+
+/*
+ * An abstract syntax tree call extern expression represents the invocation
+ * of an extern function.
+ * In C++ the call extern is an invocation: FMath::Clamp(1.4, 0.0, 1.0);
+ * The call extern expression references a node (through parent class)
+ * from the model providing all of the relevant information for the invocation.
+ */
+class RIGVMDEVELOPER_API FRigVMCallExternExprAST: public FRigVMNodeExprAST
+{
+public:
+
+	// virtual destructor
+	virtual ~FRigVMCallExternExprAST() {}
+
+	// disable copy constructor
+	FRigVMCallExternExprAST(const FRigVMCallExternExprAST&) = delete;
+
+	// overload of the type checking mechanism
+	virtual bool IsA(EType InType) const override
+	{
+		return InType == EType::CallExtern;
+	};
+
+protected:
+
+	// default constructor (protected so that only parser can access it)
+	FRigVMCallExternExprAST(const FRigVMParserAST* InParser, URigVMNode* InNode)
+		: FRigVMNodeExprAST(InParser, EType::CallExtern, InNode)
+	{}
+
+private:
+
+	friend class FRigVMParserAST;
+};
+
+/*
+ * An abstract syntax tree no-op expression represents an expression which is
+ * relevant for the structure of the tree (for grouping for example) but which 
+ * itself has no operation connected to it.
+ * For the RigVM AST we use the no-op expression for representing reroute nodes
+ * in the model as well as parameter and variable getter nodes.
+ */
+class RIGVMDEVELOPER_API FRigVMNoOpExprAST : public FRigVMNodeExprAST
+{
+public:
+
+	// virtual destructor
+	virtual ~FRigVMNoOpExprAST() {}
+
+	// disable copy constructor
+	FRigVMNoOpExprAST(const FRigVMNoOpExprAST&) = delete;
+
+	// overload of the type checking mechanism
+	virtual bool IsA(EType InType) const override
+	{
+		return InType == EType::NoOp;
+	};
+
+protected:
+
+	// default constructor (protected so that only parser can access it)
+	FRigVMNoOpExprAST(const FRigVMParserAST* InParser, URigVMNode* InNode)
+		: FRigVMNodeExprAST(InParser, EType::NoOp, InNode)
+	{}
+
+private:
+
+	friend class FRigVMParserAST;
+};
+
+/*
+ * An abstract syntax tree var expression represents the definition of 
+ * mutable memory for a single variable.
+ * In C++ the var expression is a variable declaration: int A;
+ * The var expression references a pin from the model.
+ */
+class RIGVMDEVELOPER_API FRigVMVarExprAST: public FRigVMExprAST
+{
+public:
+
+	// virtual destructor
+	virtual ~FRigVMVarExprAST() {}
+
+	// disable copy constructor
+	FRigVMVarExprAST(const FRigVMVarExprAST&) = delete;
+
+	// overload of the type checking mechanism
+	virtual bool IsA(EType InType) const override
+	{
+		return InType == EType::Var;
+	};
+
+	// returns the pin in the model this variable is representing
+	// @return the pin in the model this variable is representing
+	URigVMPin* GetPin() const { return Pin; }
+
+	// returns the C++ data type of this variable
+	// @return the C++ data type of this variable
+	FString GetCPPType() const;
+
+	// returns the C++ data type object (ustruct / uenum)
+	// @return the C++ data type object (ustruct / uenum)
+	UObject* GetCPPTypeObject() const;
+
+	// returns the pin direction of this variable (input, output, hidden etc)
+	// @return the pin direction of this variable (input, output, hidden etc)
+	ERigVMPinDirection GetPinDirection() const;
+
+	// returns the default value on the pin for this variable
+	// @return the default value on the pin for this variable
+	FString GetDefaultValue() const;
+
+	// returns true if this variable is an execute context
+	// @return true if this variable is an execute context
+	bool IsExecuteContext() const;
+
+	// returns true if this variable is a graph parameter
+	// @return true if this variable is a graph parameter
+	bool IsGraphParameter() const;
+
+	// returns true if this variable is a graph variable
+	// @return true if this variable is a graph variable
+	bool IsGraphVariable() const;
+
+protected:
+
+	// default constructor (protected so that only parser can access it)
+	FRigVMVarExprAST(const FRigVMParserAST* InParser, URigVMPin* InPin, EType InType = EType::Var)
+		: FRigVMExprAST(InParser, InType)
+		, Pin(InPin)
+	{}
+
+private:
+
+	URigVMPin* Pin;
+
+	friend class FRigVMParserAST;
+	friend class URigVMCompiler;
+};
+
+/*
+ * An abstract syntax tree literal expression represents the definition of 
+ * const memory for a single variable - vs a var expression which is mutable.
+ * In C++ the literal expression is a literal declaration, for ex: const float PI = 3.14f;
+ * The literal expression references a pin from the model.
+ */
+class RIGVMDEVELOPER_API FRigVMLiteralExprAST : public FRigVMVarExprAST
+{
+public:
+
+	// virtual destructor
+	virtual ~FRigVMLiteralExprAST() {}
+
+	// disable copy constructor
+	FRigVMLiteralExprAST(const FRigVMLiteralExprAST&) = delete;
+
+	// overload of the type checking mechanism
+	virtual bool IsA(EType InType) const override
+	{
+		if(FRigVMVarExprAST::IsA(InType))
+		{
+			return true;
+		}
+		return InType == EType::Literal;
+	};
+
+protected:
+
+	// default constructor (protected so that only parser can access it)
+	FRigVMLiteralExprAST(const FRigVMParserAST* InParser, URigVMPin* InPin)
+		: FRigVMVarExprAST(InParser, InPin, EType::Literal)
+	{}
+
+private:
+
+	friend class FRigVMParserAST;
+};
+
+/*
+ * An abstract syntax tree assign expression represents the assignment of one
+ * expression to another. This can result in referencing memory from a to b or
+ * copying memory from a to b (thus the copy expression inherits the assign).
+ * In C++ the assign expression used for construction and copy, for ex: int32 A = B + 1;
+ * The assign expression references two pins / a link from the model.
+ */
+class RIGVMDEVELOPER_API FRigVMAssignExprAST : public FRigVMExprAST
+{
+public:
+
+	// virtual destructor
+	virtual ~FRigVMAssignExprAST() {}
+
+	// disable copy constructor
+	FRigVMAssignExprAST(const FRigVMAssignExprAST&) = delete;
+
+	// returns the source pin for this assignment
+	// @return the source pin for this assignment
+	URigVMPin* GetSourcePin() const { return SourcePin; }
+
+	// returns the target pin for this assignment
+	// @return the target pin for this assignment
+	URigVMPin* GetTargetPin() const { return TargetPin; }
+
+	// overload of the type checking mechanism
+	virtual bool IsA(EType InType) const override
+	{
+		return InType == EType::Assign;
+	};
+
+protected:
+
+	// default constructor (protected so that only parser can access it)
+	FRigVMAssignExprAST(const FRigVMParserAST* InParser, URigVMPin* InSourcePin, URigVMPin* InTargetPin, EType InType = EType::Assign)
+		: FRigVMExprAST(InParser, InType)
+		, SourcePin(InSourcePin)
+		, TargetPin(InTargetPin)
+	{}
+
+
+private:
+
+	URigVMPin* SourcePin;
+	URigVMPin* TargetPin;
+	friend class FRigVMParserAST;
+};
+
+/*
+ * An abstract syntax tree copy expression represents the an assignment of one
+ * expression to another which causes / requires a copy operation. 
+ * Within the RigVM AST this is only used for copying work state out of / into parameters
+ * or when composing / decomposing a structure (for ex: assigning a float to a vector.x).
+ * In C++ the copy expression is used for structures, for ex: FVector A = B;
+ * The copy expression references two pins / a link from the model.
+ */
+class RIGVMDEVELOPER_API FRigVMCopyExprAST : public FRigVMAssignExprAST
+{
+public:
+
+	// virtual destructor
+	virtual ~FRigVMCopyExprAST() {}
+
+	// disable copy constructor
+	FRigVMCopyExprAST(const FRigVMCopyExprAST&) = delete;
+
+	// overload of the type checking mechanism
+	virtual bool IsA(EType InType) const override
+	{
+		if(FRigVMAssignExprAST::IsA(InType))
+		{
+			return true;
+		}
+		return InType == EType::Copy;
+	};
+
+protected:
+
+	// default constructor (protected so that only parser can access it)
+	FRigVMCopyExprAST(const FRigVMParserAST* InParser, URigVMPin* InSourcePin, URigVMPin* InTargetPin)
+		: FRigVMAssignExprAST(InParser, InSourcePin, InTargetPin, EType::Copy)
+	{}
+
+
+private:
+
+	friend class FRigVMParserAST;
+};
+
+/*
+ * An abstract syntax tree cached value expression represents the reference to 
+ * a variable which needs to be calculated by a call extern expression.
+ * The first child of the cached value expression is the var expression to be
+ * computed / cached, the second child is the call extern expression to use.
+ */
+class RIGVMDEVELOPER_API FRigVMCachedValueExprAST : public FRigVMExprAST
+{
+public:
+
+	// virtual destructor
+	virtual ~FRigVMCachedValueExprAST() {}
+
+	// disable copy constructor
+	FRigVMCachedValueExprAST(const FRigVMCachedValueExprAST&) = delete;
+
+	// overload of the type checking mechanism
+	virtual bool IsA(EType InType) const override
+	{
+		return InType == EType::CachedValue;
+	};
+
+	// returns the var expression of this cached value (const)
+	// @return the var expression of this cached value (const)
+	const FRigVMVarExprAST* GetVarExpr() const { return ChildAt<FRigVMVarExprAST>(0); }
+
+	// returns the call extern expression of this cached value (const)
+	// @return the call extern expression of this cached value (const)
+	const FRigVMCallExternExprAST* GetCallExternExpr() const { return ChildAt<FRigVMCallExternExprAST>(1); }
+
+protected:
+
+	// default constructor (protected so that only parser can access it)
+	FRigVMCachedValueExprAST(const FRigVMParserAST* InParser)
+		: FRigVMExprAST(InParser, EType::CachedValue)
+	{}
+
+
+private:
+
+	friend class FRigVMParserAST;
+};
+
+/*
+ * An abstract syntax tree exit expression represents the exit out of an entry expression.
+ * In C++ the exit expression is a return from a main function.
+ */
+class RIGVMDEVELOPER_API FRigVMExitExprAST : public FRigVMExprAST
+{
+public:
+
+	// virtual destructor
+	virtual ~FRigVMExitExprAST() {}
+
+	// disable copy constructor
+	FRigVMExitExprAST(const FRigVMExitExprAST&) = delete;
+
+	// overload of the type checking mechanism
+	virtual bool IsA(EType InType) const override
+	{
+		return InType == EType::Exit;
+	};
+
+protected:
+
+	// default constructor (protected so that only parser can access it)
+	FRigVMExitExprAST(const FRigVMParserAST* InParser)
+		: FRigVMExprAST(InParser, EType::Exit)
+	{}
+
+private:
+
+	friend class FRigVMParserAST;
+};
+
+/*
+ * The settings to apply during the parse of the abstract syntax tree.
+ * The folding settings can affect the performance of the parse dramatically.
+ */
+USTRUCT(BlueprintType)
+struct RIGVMDEVELOPER_API FRigVMParserASTSettings
+{
+	GENERATED_BODY()
+
+	// remove no op nodes - used for reroutes and other expressions.
+	UPROPERTY(EditAnywhere, Category = "AST")
+	bool bFoldReroutes;
+
+	// fold assignments / copies
+	UPROPERTY(EditAnywhere, Category = "AST")
+	bool bFoldAssignments;
+
+	// fold literals and share memory
+	UPROPERTY(EditAnywhere, Category = "AST")
+	bool bFoldLiterals;
+
+	// static method to provide fast AST parse settings
+	static FRigVMParserASTSettings Fast()
+	{
+		FRigVMParserASTSettings Settings;
+		Settings.bFoldReroutes = false;
+		Settings.bFoldAssignments = false;
+		Settings.bFoldLiterals = false;
+		return Settings;
+	}
+
+	// static method to provide AST parse settings
+	// tuned for a fast executing runtime, but slow parse
+	static FRigVMParserASTSettings Optimized()
+	{
+		FRigVMParserASTSettings Settings;
+		Settings.bFoldReroutes = true;
+		Settings.bFoldAssignments = true;
+		Settings.bFoldLiterals = true;
+		return Settings;
+	}
+};
+
+/*
+ * The abstract syntax tree parser is the main object to parse a
+ * RigVM model graph. It's the memory owner for all expressions and
+ * provides functionality for introspection of the tree.
+ * The abstract syntax tree is then fed into the RigVMCompiler to 
+ * generate the byte code for the virtual machine.
+ */
+class RIGVMDEVELOPER_API FRigVMParserAST : public TSharedFromThis<FRigVMParserAST>
+{
+public:
+
+	// default constructor
+	// @param InGraph The graph / model to parse
+	// @param InSettings The parse settings to use
+	FRigVMParserAST(URigVMGraph* InGraph, const FRigVMParserASTSettings& InSettings = FRigVMParserASTSettings::Fast());
+
+	// default destructor
+	~FRigVMParserAST();
+
+	// returns the number of root expressions
+	int32 Num() const { return RootExpressions.Num(); }
+
+	// operator accessor for a given root expression
+	// @param InIndex the index of the expression to retrieve (bound = Num() - 1)
+	// @return the root expression with the given index
+	FORCEINLINE const FRigVMExprAST* operator[](int32 InIndex) const { return RootExpressions[InIndex]; }
+
+	// begin iterator accessor for the root expressions
+	FORCEINLINE TArray<FRigVMExprAST*>::RangedForConstIteratorType begin() const { return RootExpressions.begin(); }
+
+	// end iterator accessor for the root expressions
+	FORCEINLINE TArray<FRigVMExprAST*>::RangedForConstIteratorType end() const { return RootExpressions.end(); }
+
+	// accessor method for a given root expression
+	// @param InIndex the index of the expression to retrieve (bound = Num() - 1)
+	// @return the root expression with the given index
+	const FRigVMExprAST* At(int32 InIndex) const { return RootExpressions[InIndex]; }
+
+	// returns the expression for a given subject. subjects include nodes and pins.
+	// @param InSubject the subject to retrieve the expression for (node or pin)
+	// @return the expressoin for the given subject (or nullptr)
+	const FRigVMExprAST* GetExprForSubject(UObject* InSubject);
+
+	// Prepares the parser for cycle checking on a given pin.
+	// This marks up the parents and childen of the corresponding expression in the graph,
+	// to allow the client to determine if a new parent / child relationship could cause a cycle.
+	// @param InPin the pin to initialize the cycle checking for
+	void PrepareCycleChecking(URigVMPin* InPin);
+
+	// Performs a cycle check for a new potential link (assign or copy) between two pins.
+	// @param InSourcePin the source (left) pin of the potential assign / copy
+	// @param InTargetPin the target (right) pin of the potential assign / copy
+	// @param OutFailureReason an optional storage for the possible failure reason description
+	// @return true if the potential link (assign / copy) can be established
+	bool CanLink(URigVMPin* InSourcePin, URigVMPin* InTargetPin, FString* OutFailureReason);
+
+	// returns a string containing an indented tree structure
+	// for debugging purposes.
+	// @return the simple text representing this abstract syntax tree
+	FString DumpText() const;
+
+	// returns a string containing a dot file notation
+	// for debugging purposes.
+	// @return the dot file text representing this abstract syntax tree
+	FString DumpDot() const;
+
+private:
+
+	// removes a single expression from the parser
+	// @param InExpr the expression to remove
+	// @param bRefreshIndices flag to determine if the expression indices need to be refreshed
+	void RemoveExpression(FRigVMExprAST* InExpr, bool bRefreshIndices = true);
+
+	// removes an array of expressions from the parser
+	// @param InExprs the expressions to remove
+	// @param bRefreshIndices flag to determine if the expression indices need to be refreshed
+	void RemoveExpressions(TArray<FRigVMExprAST*> InExprs, bool bRefreshIndices = true);
+
+	// a static helper function to traverse along all parents of an expression,
+	// provided a predicate to return true if the traverse should continue,
+	// and false if the traverse should stop.
+	// @param InExpr the expression to start the traverse at
+	// @param InContinuePredicate the predicate to use to break out of the traverse (return false)
+	static void TraverseParents(const FRigVMExprAST* InExpr, TFunctionRef<bool(const FRigVMExprAST*)> InContinuePredicate);
+
+	// a static helper function to traverse along all children of an expression,
+	// provided a predicate to return true if the traverse should continue,
+	// and false if the traverse should stop.
+	// @param InExpr the expression to start the traverse at
+	// @param InContinuePredicate the predicate to use to break out of the traverse (return false)
+	static void TraverseChildren(const FRigVMExprAST* InExpr, TFunctionRef<bool(const FRigVMExprAST*)> InContinuePredicate);
+
+	// helper function to fold all entries with the same event name into one block
+	void FoldEntries();
+
+	// helper function to inject an exit expression at the end of every entry expressoin
+	void InjectExitsToEntries();
+
+	// helper function to refresh the expression indices (used after deleting an expression)
+	void RefreshExprIndices();
+
+	// helper function to fold / remove the no op expressions
+	void FoldNoOps();
+
+	// helper function to fold / merge redundant literals with the same value
+	void FoldLiterals();
+
+	// helper function to fold / remove obsolete assignments and reduce assignment chains
+	void FoldAssignments();
+
+	// traverse a single mutable node (constructs entry, call extern and other expressions)
+	FRigVMExprAST* TraverseMutableNode(URigVMNode* InNode, FRigVMExprAST* InParentExpr);
+
+	// traverse a single pure node (constructs call extern expressions)
+	FRigVMExprAST* TraverseNode(URigVMNode* InNode, FRigVMExprAST* InParentExpr);
+
+	// traverse an array of pins for a given node
+	TArray<FRigVMExprAST*> TraversePins(URigVMNode* InNode, FRigVMExprAST* InParentExpr);
+
+	// traverse a single pin (constructs var + literal expressions)
+	FRigVMExprAST* TraversePin(URigVMPin* InPin, FRigVMExprAST* InParentExpr);
+
+	// traverse a single link (constructs assign + copy expressions)
+	FRigVMExprAST* TraverseLink(URigVMLink* InLink, FRigVMExprAST* InParentExpr);
+
+	TMap<UObject*, FRigVMExprAST*> SubjectToExpression;
+	TArray<int32> NodeExpressionIndex;
+	TArray<FRigVMExprAST*> Expressions;
+	TArray<FRigVMExprAST*> RootExpressions;
+
+	enum ETraverseRelationShip
+	{
+		ETraverseRelationShip_Unknown,
+		ETraverseRelationShip_Parent,
+		ETraverseRelationShip_Child,
+		ETraverseRelationShip_Self
+	};
+
+	const FRigVMExprAST* LastCycleCheckExpr;
+	TArray<ETraverseRelationShip> CycleCheckFlags;
+
+	friend class FRigVMExprAST;
+};
+
diff --git a/Engine/Source/Developer/RigVMDeveloper/Public/RigVMCompiler/RigVMCompiler.h b/Engine/Source/Developer/RigVMDeveloper/Public/RigVMCompiler/RigVMCompiler.h
new file mode 100644
--- /dev/null
+++ b/Engine/Source/Developer/RigVMDeveloper/Public/RigVMCompiler/RigVMCompiler.h
@@ -0,0 +1,133 @@
+// Copyright Epic Games, Inc. All Rights Reserved.
+
+#pragma once
+
+#include "CoreMinimal.h"
+
+#include "RigVMModel/RigVMGraph.h"
+#include "RigVMModel/Nodes/RigVMStructNode.h"
+#include "RigVMModel/Nodes/RigVMVariableNode.h"
+#include "RigVMModel/Nodes/RigVMParameterNode.h"
+#include "RigVMModel/Nodes/RigVMCommentNode.h"
+#include "RigVMModel/Nodes/RigVMRerouteNode.h"
+#include "RigVMCore/RigVM.h"
+#include "RigVMCore/RigVMStruct.h"
+#include "RigVMCompiler/RigVMAST.h"
+
+#include "RigVMCompiler.generated.h"
+
+USTRUCT(BlueprintType)
+struct RIGVMDEVELOPER_API FRigVMCompileSettings
+{
+	GENERATED_BODY()
+
+public:
+
+	FRigVMCompileSettings();
+
+	UPROPERTY(EditAnywhere, BlueprintReadWrite, Category = FRigVMCompileSettings)
+	bool SurpressInfoMessages;
+
+	UPROPERTY(EditAnywhere, BlueprintReadWrite, Category = FRigVMCompileSettings)
+	bool SurpressWarnings;
+
+	UPROPERTY(EditAnywhere, BlueprintReadWrite, Category = FRigVMCompileSettings)
+	bool SurpressErrors;
+
+	UPROPERTY(EditAnywhere, BlueprintReadWrite, Category = FRigVMCompileSettings)
+	bool EnablePinWatches;
+
+	UPROPERTY(EditAnywhere, BlueprintReadWrite, Category = FRigVMCompileSettings)
+	bool SplitLiteralsFromWorkMemory;
+
+	UPROPERTY(EditAnywhere, BlueprintReadWrite, Category = FRigVMCompileSettings)
+	bool ConsolidateWorkRegisters;
+
+	UPROPERTY(EditAnywhere, Transient, BlueprintReadWrite, Category = FRigVMCompileSettings)
+	FRigVMParserASTSettings ASTSettings;
+};
+
+struct RIGVMDEVELOPER_API FRigVMCompilerWorkData
+{
+public:
+	bool bSetupMemory;
+	URigVM* VM;
+	FRigVMUserDataArray RigVMUserData;
+	TMap<FString, FRigVMOperand>* PinPathToOperand;
+	TMap<const FRigVMVarExprAST*, FRigVMOperand> ExprToOperand;
+	TMap<const FRigVMExprAST*, bool> ExprComplete;
+	TMap<FString, int32> ProcessedLinks;
+	TArray<FRigVMStruct*> DefaultStructs;
+	TMap<int32, int32> RegisterRefCount;
+	int32 NumInstructions;
+
+	int32 IncRefRegister(int32 InRegister, int32 InIncrement = 1);
+	int32 DecRefRegister(int32 InRegister, int32 InDecrement = 1);
+};
+
+UCLASS(BlueprintType)
+class RIGVMDEVELOPER_API URigVMCompiler : public UObject
+{
+	GENERATED_BODY()
+
+public:
+
+	URigVMCompiler();
+
+	UPROPERTY(BlueprintReadWrite, Category = FRigVMCompiler)
+	FRigVMCompileSettings Settings;
+
+	UFUNCTION(BlueprintCallable, Category = FRigVMCompiler)
+	bool Compile(URigVMGraph* InGraph, URigVM* OutVM)
+	{
+		return Compile(InGraph, OutVM, FRigVMUserDataArray(), nullptr);
+	}
+
+	bool Compile(URigVMGraph* InGraph, URigVM* OutVM, const FRigVMUserDataArray& InRigVMUserData, TMap<FString, FRigVMOperand>* OutOperands);
+	static UScriptStruct* GetScriptStructForCPPType(const FString& InCPPType);
+	static FString GetPinHash(URigVMPin* InPin, const FRigVMVarExprAST* InVarExpr, bool bIsDebugValue = false);
+
+private:
+
+	TArray<URigVMPin*> GetLinkedPins(URigVMPin* InPin, bool bInputs = true, bool bOutputs = true, bool bRecursive = true);
+	uint16 GetElementSizeFromCPPType(const FString& InCPPType, UScriptStruct* InScriptStruct);
+
+	void TraverseExpression(const FRigVMExprAST* InExpr, FRigVMCompilerWorkData& WorkData);
+	void TraverseChildren(const FRigVMExprAST* InExpr, FRigVMCompilerWorkData& WorkData);
+	void TraverseBlock(const FRigVMBlockExprAST* InExpr, FRigVMCompilerWorkData& WorkData);
+	void TraverseEntry(const FRigVMEntryExprAST* InExpr, FRigVMCompilerWorkData& WorkData);
+	void TraverseCallExtern(const FRigVMCallExternExprAST* InExpr, FRigVMCompilerWorkData& WorkData);
+	void TraverseNoOp(const FRigVMNoOpExprAST* InExpr, FRigVMCompilerWorkData& WorkData);
+	void TraverseVar(const FRigVMVarExprAST* InExpr, FRigVMCompilerWorkData& WorkData);
+	void TraverseLiteral(const FRigVMVarExprAST* InExpr, FRigVMCompilerWorkData& WorkData);
+	void TraverseAssign(const FRigVMAssignExprAST* InExpr, FRigVMCompilerWorkData& WorkData);
+	void TraverseCopy(const FRigVMCopyExprAST* InExpr, FRigVMCompilerWorkData& WorkData);
+	void TraverseCachedValue(const FRigVMCachedValueExprAST* InExpr, FRigVMCompilerWorkData& WorkData);
+	void TraverseExit(const FRigVMExitExprAST* InExpr, FRigVMCompilerWorkData& WorkData);
+
+	FRigVMOperand FindOrAddRegister(const FRigVMVarExprAST* InVarExpr, FRigVMCompilerWorkData& WorkData, bool bIsDebugValue = false);
+
+	void ReportInfo(const FString& InMessage);
+	void ReportWarning(const FString& InMessage);
+	void ReportError(const FString& InMessage);
+
+	template <typename FmtType, typename... Types>
+	void ReportInfof(const FmtType& Fmt, Types... Args)
+	{
+		ReportInfo(FString::Printf(Fmt, Args...));
+	}
+
+	template <typename FmtType, typename... Types>
+	void ReportWarningf(const FmtType& Fmt, Types... Args)
+	{
+		ReportWarning(FString::Printf(Fmt, Args...));
+	}
+
+	template <typename FmtType, typename... Types>
+	void ReportErrorf(const FmtType& Fmt, Types... Args)
+	{
+		ReportError(FString::Printf(Fmt, Args...));
+	}
+
+	friend class FRigVMCompilerImportErrorContext;
+};
diff --git a/Engine/Source/Developer/RigVMDeveloper/Public/RigVMModel/Nodes/RigVMRerouteNode.h b/Engine/Source/Developer/RigVMDeveloper/Public/RigVMModel/Nodes/RigVMRerouteNode.h
new file mode 100644
--- /dev/null
+++ b/Engine/Source/Developer/RigVMDeveloper/Public/RigVMModel/Nodes/RigVMRerouteNode.h
@@ -0,0 +1,46 @@
+// Copyright Epic Games, Inc. All Rights Reserved.
+
+#pragma once
+
+#include "RigVMModel/RigVMNode.h"
+#include "RigVMRerouteNode.generated.h"
+
+/**
+ * A reroute node is used to visually improve the 
+ * data flow within a Graph. Reroutes are purely 
+ * cosmetic and have no impact on the resulting
+ * VM whatsoever. Reroutes can furthermore be
+ * displayed as full nodes or as small circles.
+ */
+UCLASS(BlueprintType)
+class RIGVMDEVELOPER_API URigVMRerouteNode : public URigVMNode
+{
+	GENERATED_BODY()
+
+public:
+
+	// Default constructor
+	URigVMRerouteNode();
+
+	// Override of node title
+	virtual FString GetNodeTitle() const override { return RerouteName; }
+
+	// Returns true if this node should be shown as a full node,
+	// or false if this node should be shown as a small circle.
+	UFUNCTION(BlueprintCallable, Category = RigVMRerouteNode)
+	bool GetShowsAsFullNode() const;
+
+	virtual FLinearColor GetNodeColor() const override;
+
+private:
+
+	static const FString RerouteName;
+	static const FString ValueName;
+
+	UPROPERTY()
+	bool bShowAsFullNode;
+
+	friend class URigVMController;
+	friend class URigVMCompiler;
+};
+
diff --git a/Engine/Source/Developer/RigVMDeveloper/Public/RigVMModel/Nodes/RigVMStructNode.h b/Engine/Source/Developer/RigVMDeveloper/Public/RigVMModel/Nodes/RigVMStructNode.h
new file mode 100644
--- /dev/null
+++ b/Engine/Source/Developer/RigVMDeveloper/Public/RigVMModel/Nodes/RigVMStructNode.h
@@ -0,0 +1,56 @@
+// Copyright Epic Games, Inc. All Rights Reserved.
+
+#pragma once
+
+#include "RigVMModel/RigVMNode.h"
+#include "RigVMCore/RigVMStruct.h"
+#include "UObject/StructOnScope.h"
+#include "RigVMStructNode.generated.h"
+
+/**
+ * The Struct Node represents a Function Invocation of a RIGVM_METHOD
+ * declared on a USTRUCT. Struct Nodes have input / output pins for all
+ * struct UPROPERTY members.
+ */
+UCLASS(BlueprintType)
+class RIGVMDEVELOPER_API URigVMStructNode : public URigVMNode
+{
+	GENERATED_BODY()
+
+public:
+
+	// Override node tooltip
+	virtual FText GetToolTipText() const override;
+
+	// Returns the UStruct for this struct node
+	// (the struct declaring the RIGVM_METHOD)
+	UFUNCTION(BlueprintCallable, Category = RigVMStructNode)
+	UScriptStruct* GetScriptStruct() const;
+
+	// Returns the name of the declared RIGVM_METHOD
+	UFUNCTION(BlueprintCallable, Category = RigVMStructNode)
+	FName GetMethodName() const;
+
+	// Returns the default value for the struct as text
+	UFUNCTION(BlueprintCallable, Category = RigVMStructNode)
+	FString GetStructDefaultValue() const;
+
+	// Returns an instance of the struct with the current values.
+	// @param bUseDefault If set to true the default struct will be created - otherwise the struct will contains the values from the node
+	TSharedPtr<FStructOnScope> ConstructStructInstance(bool bUseDefault = false) const;
+
+protected:
+
+	virtual FText GetToolTipTextForPin(const URigVMPin* InPin) const override;
+
+private:
+
+	UPROPERTY()
+	UScriptStruct* ScriptStruct;
+
+	UPROPERTY()
+	FName MethodName;
+
+	friend class URigVMController;
+};
+
diff --git a/Engine/Source/Developer/RigVMDeveloper/Public/RigVMModel/Nodes/RigVMVariableNode.h b/Engine/Source/Developer/RigVMDeveloper/Public/RigVMModel/Nodes/RigVMVariableNode.h
new file mode 100644
--- /dev/null
+++ b/Engine/Source/Developer/RigVMDeveloper/Public/RigVMModel/Nodes/RigVMVariableNode.h
@@ -0,0 +1,97 @@
+// Copyright Epic Games, Inc. All Rights Reserved.
+
+#pragma once
+
+#include "RigVMModel/RigVMNode.h"
+#include "RigVMVariableNode.generated.h"
+
+/**
+ * The variable description is used to convey information
+ * about unique variables within a Graph. Multiple Variable
+ * Nodes can share the same variable description.
+ */
+USTRUCT(BlueprintType)
+struct FRigVMGraphVariableDescription
+{
+	GENERATED_BODY()
+
+public:
+
+	// comparison operator
+	bool operator ==(const FRigVMGraphVariableDescription& Other) const
+	{
+		return Name == Other.Name;
+	}
+
+	// The name of the variable
+	UPROPERTY(BlueprintReadOnly, Category = RigVMGraphVariableDescription)
+	FName Name;
+
+	// The C++ data type of the variable
+	UPROPERTY(BlueprintReadOnly, Category = RigVMGraphVariableDescription)
+	FString CPPType;
+
+	// The Struct of the C++ data type of the variable (or nullptr)
+	UPROPERTY(BlueprintReadOnly, Category = RigVMGraphVariableDescription)
+	UObject* CPPTypeObject;
+
+	// The default value of the variable
+	UPROPERTY(BlueprintReadOnly, Category = RigVMGraphVariableDescription)
+	FString DefaultValue;
+};
+
+/**
+ * The Variable Node represents a mutable value / local state within the
+ * the Function / Graph. Variable Node's can be a getter or a setter.
+ * Getters are pure nodes with just an output value pin, while setters
+ * are mutable nodes with an execute and input value pin.
+ */
+UCLASS(BlueprintType)
+class RIGVMDEVELOPER_API URigVMVariableNode : public URigVMNode
+{
+	GENERATED_BODY()
+
+public:
+
+	// Default constructor
+	URigVMVariableNode();
+
+	// Override of node title
+	virtual FString GetNodeTitle() const;
+
+	// Returns the name of the variable
+	UFUNCTION(BlueprintCallable, Category = RigVMVariableNode)
+	FName GetVariableName() const;
+
+	// Returns true if this node is a variable getter
+	UFUNCTION(BlueprintCallable, Category = RigVMVariableNode)
+	bool IsGetter() const;
+
+	// Returns the C++ data type of the variable
+	UFUNCTION(BlueprintCallable, Category = RigVMVariableNode)
+	FString GetCPPType() const;
+
+	// Returns the C++ data type struct of the variable (or nullptr)
+	UFUNCTION(BlueprintCallable, Category = RigVMVariableNode)
+	UObject* GetCPPTypeObject() const;
+
+	// Returns the default value of the variable as a string
+	UFUNCTION(BlueprintCallable, Category = RigVMVariableNode)
+	FString GetDefaultValue() const;
+
+	// Returns this variable node's variable description
+	UFUNCTION(BlueprintCallable, Category = RigVMVariableNode)
+	FRigVMGraphVariableDescription GetVariableDescription() const;
+
+	// Override of node title
+	virtual FLinearColor GetNodeColor() const override { return FLinearColor::Blue; }
+
+private:
+
+	static const FString VariableName;
+	static const FString ValueName;
+
+	friend class URigVMController;
+	friend class URigVMCompiler;
+};
+
