COPY from //Dev-Anim to //Dev-Main

diff --git a/Engine/Source/Developer/RigVMDeveloper/Public/RigVMModel/RigVMController.h b/Engine/Source/Developer/RigVMDeveloper/Public/RigVMModel/RigVMController.h
new file mode 100644
--- /dev/null
+++ b/Engine/Source/Developer/RigVMDeveloper/Public/RigVMModel/RigVMController.h
@@ -0,0 +1,462 @@
+// Copyright Epic Games, Inc. All Rights Reserved.
+
+#pragma once
+
+#include "RigVMGraph.h"
+#include "RigVMModel/Nodes/RigVMStructNode.h"
+#include "RigVMModel/Nodes/RigVMVariableNode.h"
+#include "RigVMModel/Nodes/RigVMParameterNode.h"
+#include "RigVMModel/Nodes/RigVMCommentNode.h"
+#include "RigVMModel/Nodes/RigVMRerouteNode.h"
+#include "RigVMController.generated.h"
+
+class URigVMActionStack;
+
+DECLARE_DELEGATE_RetVal_OneParam(bool, FRigVMController_ShouldStructUnfoldDelegate, const UStruct*)
+
+/**
+ * The Controller is the sole authority to perform changes
+ * on the Graph. The Controller itself is stateless.
+ * The Controller offers a Modified event to subscribe to 
+ * for user interface views - so they can be informed about
+ * any change that's happening within the Graph.
+ * The Controller routes all changes through the Graph itself,
+ * so you can have N Controllers performing edits on 1 Graph,
+ * and N Views subscribing to 1 Controller.
+ * In Python you can also subscribe to this event to be 
+ * able to react to topological changes of the Graph there.
+ */
+UCLASS(BlueprintType)
+class RIGVMDEVELOPER_API URigVMController : public UObject
+{
+	GENERATED_UCLASS_BODY()
+
+public:
+
+	// Default constructor
+	URigVMController();
+
+	// Default destructor
+	~URigVMController();
+
+	// Returns the currently edited Graph of this controller.
+	UFUNCTION(BlueprintCallable, Category = RigVMController)
+	URigVMGraph* GetGraph() const;
+
+	// Sets the currently edited Graph of this controller.
+	// This causes a GraphChanged modified event.
+	UFUNCTION(BlueprintCallable, Category = RigVMController)
+	void SetGraph(URigVMGraph* InGraph);
+
+	// The Modified event used to subscribe to changes
+	// happening within the Graph. This is broadcasted to 
+	// for any change happening - not only the changes 
+	// performed by this Controller - so it can be used
+	// for UI Views to react accordingly.
+	FRigVMGraphModifiedEvent& OnModified();
+
+	// Submits an event to the graph for broadcasting.
+	void Notify(ERigVMGraphNotifType InNotifType, UObject* InSubject);
+
+	// Resends all notifications
+	void ResendAllNotifications();
+
+	// Enables or disables the error reporting of this Controller.
+	UFUNCTION(BlueprintCallable, Category = RigVMController)
+	void EnableReporting(bool bEnabled = true) { bReportWarningsAndErrors = bEnabled; }
+
+#if WITH_EDITOR
+	// Note: The functions below are scoped with WITH_EDITOR since we are considering
+	// to move this code into the runtime in the future. Right now there's a dependency
+	// on the metadata of the USTRUCT - which is only available in the editor.
+
+	// Adds a Function / Struct Node to the edited Graph.
+	// StructNode represent a RIGVM_METHOD declaration on a USTRUCT.
+	// This causes a NodeAdded modified event.
+	UFUNCTION(BlueprintCallable, Category = RigVMController)
+	URigVMStructNode* AddStructNode(UScriptStruct* InScriptStruct, const FName& InMethodName, const FVector2D& InPosition = FVector2D::ZeroVector, const FString& InNodeName = TEXT(""), bool bUndo = true);
+
+	// Adds a Function / Struct Node to the edited Graph given its struct object path name.
+	// StructNode represent a RIGVM_METHOD declaration on a USTRUCT.
+	// This causes a NodeAdded modified event.
+	UFUNCTION(BlueprintCallable, Category = RigVMController)
+	URigVMStructNode* AddStructNodeFromStructPath(const FString& InScriptStructPath, const FName& InMethodName, const FVector2D& InPosition = FVector2D::ZeroVector, const FString& InNodeName = TEXT(""), bool bUndo = true);
+
+	// Adds a Variable Node to the edited Graph.
+	// Variables represent local work state for the function and
+	// can be read from and written to. 
+	// This causes a NodeAdded modified event.
+	UFUNCTION(BlueprintCallable, Category = RigVMController)
+	URigVMVariableNode* AddVariableNode(const FName& InVariableName, const FString& InCPPType, UObject* InCPPTypeObject, bool bIsGetter, const FString& InDefaultValue, const FVector2D& InPosition = FVector2D::ZeroVector, const FString& InNodeName = TEXT(""), bool bUndo = true);
+
+	// Adds a Variable Node to the edited Graph given a struct object path name.
+	// Variables represent local work state for the function and
+	// can be read from (bIsGetter == true) or written to (bIsGetter == false). 
+	// This causes a NodeAdded modified event.
+	UFUNCTION(BlueprintCallable, Category = RigVMController)
+	URigVMVariableNode* AddVariableNodeFromObjectPath(const FName& InVariableName, const FString& InCPPType, const FString& InCPPTypeObjectPath, bool bIsGetter, const FString& InDefaultValue, const FVector2D& InPosition = FVector2D::ZeroVector, const FString& InNodeName = TEXT(""), bool bUndo = true);
+
+	// Adds a Parameter Node to the edited Graph.
+	// Parameters represent input or output arguments to the Graph / Function.
+	// Input Parameters are constant values / literals.
+	// This causes a NodeAdded modified event.
+	UFUNCTION(BlueprintCallable, Category = RigVMController)
+	URigVMParameterNode* AddParameterNode(const FName& InParameterName, const FString& InCPPType, UObject* InCPPTypeObject, bool bIsInput, const FString& InDefaultValue, const FVector2D& InPosition = FVector2D::ZeroVector, const FString& InNodeName = TEXT(""), bool bUndo = true);
+
+	// Adds a Parameter Node to the edited Graph given a struct object path name.
+	// Parameters represent input or output arguments to the Graph / Function.
+	// Input Parameters are constant values / literals.
+	// This causes a NodeAdded modified event.
+	UFUNCTION(BlueprintCallable, Category = RigVMController)
+	URigVMParameterNode* AddParameterNodeFromObjectPath(const FName& InParameterName, const FString& InCPPType, const FString& InCPPTypeObjectPath, bool bIsInput, const FString& InDefaultValue, const FVector2D& InPosition = FVector2D::ZeroVector, const FString& InNodeName = TEXT(""), bool bUndo = true);
+
+	// Adds a Comment Node to the edited Graph.
+	// Comments can be used to annotate the Graph.
+	// This causes a NodeAdded modified event.
+	UFUNCTION(BlueprintCallable, Category = RigVMController)
+	URigVMCommentNode* AddCommentNode(const FString& InCommentText, const FVector2D& InPosition = FVector2D::ZeroVector, const FVector2D& InSize = FVector2D(400.f, 300.f), const FLinearColor& InColor = FLinearColor::Black, const FString& InNodeName = TEXT(""), bool bUndo = true);
+
+	// Adds a Reroute Node on an existing Link to the edited Graph.
+	// Reroute Nodes can be used to visually improve the data flow,
+	// they don't require any additional memory though and are purely
+	// cosmetic. This causes a NodeAdded modified event.
+	UFUNCTION(BlueprintCallable, Category = RigVMController)
+	URigVMRerouteNode* AddRerouteNodeOnLink(URigVMLink* InLink, bool bShowAsFullNode, const FVector2D& InPosition = FVector2D::ZeroVector, const FString& InNodeName = TEXT(""), bool bUndo = true);
+
+	// Adds a Reroute Node on an existing Link to the edited Graph given the Link's string representation.
+	// Reroute Nodes can be used to visually improve the data flow,
+	// they don't require any additional memory though and are purely
+	// cosmetic. This causes a NodeAdded modified event.
+	UFUNCTION(BlueprintCallable, Category = RigVMController)
+	URigVMRerouteNode* AddRerouteNodeOnLinkPath(const FString& InLinkPinPathRepresentation, bool bShowAsFullNode, const FVector2D& InPosition = FVector2D::ZeroVector, const FString& InNodeName = TEXT(""), bool bUndo = true);
+
+	// Adds a Reroute Node on an existing Pin to the editor Graph.
+	// Reroute Nodes can be used to visually improve the data flow,
+	// they don't require any additional memory though and are purely
+	// cosmetic. This causes a NodeAdded modified event.
+	UFUNCTION(BlueprintCallable, Category = RigVMController)
+	URigVMRerouteNode* AddRerouteNodeOnPin(const FString& InPinPath, bool bAsInput, bool bShowAsFullNode, const FVector2D& InPosition = FVector2D::ZeroVector, const FString& InNodeName = TEXT(""), bool bUndo = true);
+
+	// Adds a free Reroute Node
+	UFUNCTION(BlueprintCallable, Category = RigVMController)
+	URigVMRerouteNode* AddFreeRerouteNode(bool bShowAsFullNode, const FString& InCPPType, const FName& InCPPTypeObjectPath, bool bIsConstant, const FName& InCustomWidgetName, const FString& InDefaultValue, const FVector2D& InPosition = FVector2D::ZeroVector, const FString& InNodeName = TEXT(""), bool bUndo = true);
+
+	// Adds a Function / Struct Node to the edited Graph as an injected node
+	// StructNode represent a RIGVM_METHOD declaration on a USTRUCT.
+	// This causes a NodeAdded modified event.
+	UFUNCTION(BlueprintCallable, Category = RigVMController)
+	URigVMInjectionInfo* AddInjectedNode(const FString& InPinPath, bool bAsInput, UScriptStruct* InScriptStruct, const FName& InMethodName, const FName& InInputPinName, const FName& InOutputPinName, const FString& InNodeName = TEXT(""), bool bUndo = true);
+
+	// Adds a Function / Struct Node to the edited Graph as an injected node
+	// StructNode represent a RIGVM_METHOD declaration on a USTRUCT.
+	// This causes a NodeAdded modified event.
+	UFUNCTION(BlueprintCallable, Category = RigVMController)
+	URigVMInjectionInfo* AddInjectedNodeFromStructPath(const FString& InPinPath, bool bAsInput, const FString& InScriptStructPath, const FName& InMethodName, const FName& InInputPinName, const FName& InOutputPinName, const FString& InNodeName = TEXT(""), bool bUndo = true);
+
+	// Ejects the last injected node on a pin
+	UFUNCTION(BlueprintCallable, Category = RigVMController)
+	URigVMNode* EjectNodeFromPin(const FString& InPinPath, bool bUndo = true);
+
+	// Un-does the last action on the stack.
+	// Note: This should really only be used for unit tests,
+	// use the GEditor's main Undo method instead.
+	UFUNCTION(BlueprintCallable, Category = RigVMController)
+	bool Undo();
+
+	// Re-does the last action on the stack.
+	// Note: This should really only be used for unit tests,
+	// use the GEditor's main Undo method instead.
+	UFUNCTION(BlueprintCallable, Category = RigVMController)
+	bool Redo();
+
+	// Opens an undo bracket / scoped transaction for
+	// a series of actions to be performed as one step on the 
+	// Undo stack. This is primarily useful for Python.
+	// This causes a UndoBracketOpened modified event.
+	UFUNCTION(BlueprintCallable, Category = RigVMController)
+	bool OpenUndoBracket(const FString& InTitle);
+
+	// Closes an undo bracket / scoped transaction.
+	// This is primarily useful for Python.
+	// This causes a UndoBracketClosed modified event.
+	UFUNCTION(BlueprintCallable, Category = RigVMController)
+	bool CloseUndoBracket();
+
+	// Cancels an undo bracket / scoped transaction.
+	// This is primarily useful for Python.
+	// This causes a UndoBracketCanceled modified event.
+	UFUNCTION(BlueprintCallable, Category = RigVMController)
+	bool CancelUndoBracket();
+
+	// Exports the given nodes as text
+	UFUNCTION(BlueprintCallable, Category = RigVMController)
+	FString ExportNodesToText(const TArray<FName>& InNodeNames);
+
+	// Exports the selected nodes as text
+	UFUNCTION(BlueprintCallable, Category = RigVMController)
+	FString ExportSelectedNodesToText();
+
+	// Exports the given nodes as text
+	UFUNCTION(BlueprintCallable, Category = RigVMController)
+	bool CanImportNodesFromText(const FString& InText);
+
+	// Exports the given nodes as text
+	UFUNCTION(BlueprintCallable, Category = RigVMController)
+	TArray<FName> ImportNodesFromText(const FString& InText, bool bUndo = true);
+
+	// Returns a unique name
+	static FName GetUniqueName(const FName& InName, TFunction<bool(const FName&)> IsNameAvailableFunction);
+
+#endif
+
+	// Removes a node from the graph
+	// This causes a NodeRemoved modified event.
+	UFUNCTION(BlueprintCallable, Category = RigVMController)
+	bool RemoveNode(URigVMNode* InNode, bool bUndo = true);
+
+	// Removes a node from the graph given the node's name.
+	// This causes a NodeRemoved modified event.
+	UFUNCTION(BlueprintCallable, Category = RigVMController)
+	bool RemoveNodeByName(const FName& InNodeName, bool bUndo = true);
+
+	// Selects a single node in the graph.
+	// This causes a NodeSelected / NodeDeselected modified event.
+	UFUNCTION(BlueprintCallable, Category = RigVMController)
+	bool SelectNode(URigVMNode* InNode, bool bSelect = true, bool bUndo = true);
+
+	// Selects a single node in the graph by name.
+	// This causes a NodeSelected / NodeDeselected modified event.
+	UFUNCTION(BlueprintCallable, Category = RigVMController)
+	bool SelectNodeByName(const FName& InNodeName, bool bSelect = true, bool bUndo = true);
+
+	// Deselects all currently selected nodes in the graph.
+	// This might cause several NodeDeselected modified event.
+	UFUNCTION(BlueprintCallable, Category = RigVMController)
+	bool ClearNodeSelection(bool bUndo = true);
+
+	// Selects the nodes given the selection
+	// This might cause several NodeDeselected modified event.
+	UFUNCTION(BlueprintCallable, Category = RigVMController)
+	bool SetNodeSelection(const TArray<FName>& InNodeNames, bool bUndo = true);
+
+	// Sets the position of a node in the graph.
+	// This causes a NodePositionChanged modified event.
+	UFUNCTION(BlueprintCallable, Category = RigVMController)
+	bool SetNodePosition(URigVMNode* InNode, const FVector2D& InPosition, bool bUndo = true, bool bMergeUndoAction = false);
+
+	// Sets the position of a node in the graph by name.
+	// This causes a NodePositionChanged modified event.
+	UFUNCTION(BlueprintCallable, Category = RigVMController)
+	bool SetNodePositionByName(const FName& InNodeName, const FVector2D& InPosition, bool bUndo = true, bool bMergeUndoAction = false);
+
+	// Sets the size of a node in the graph.
+	// This causes a NodeSizeChanged modified event.
+	UFUNCTION(BlueprintCallable, Category = RigVMController)
+	bool SetNodeSize(URigVMNode* InNode, const FVector2D& InSize, bool bUndo = true, bool bMergeUndoAction = false);
+
+	// Sets the size of a node in the graph by name.
+	// This causes a NodeSizeChanged modified event.
+	UFUNCTION(BlueprintCallable, Category = RigVMController)
+	bool SetNodeSizeByName(const FName& InNodeName, const FVector2D& InSize, bool bUndo = true, bool bMergeUndoAction = false);
+
+	// Sets the color of a node in the graph.
+	// This causes a NodeColorChanged modified event.
+	UFUNCTION(BlueprintCallable, Category = RigVMController)
+	bool SetNodeColor(URigVMNode* InNode, const FLinearColor& InColor, bool bUndo = true, bool bMergeUndoAction = false);
+
+	// Sets the color of a node in the graph by name.
+	// This causes a NodeColorChanged modified event.
+	UFUNCTION(BlueprintCallable, Category = RigVMController)
+	bool SetNodeColorByName(const FName& InNodeName, const FLinearColor& InColor, bool bUndo = true, bool bMergeUndoAction = false);
+
+	// Sets the comment text of a comment node in the graph.
+	// This causes a CommentTextChanged modified event.
+	UFUNCTION(BlueprintCallable, Category = RigVMController)
+	bool SetCommentText(URigVMNode* InNode, const FString& InCommentText, bool bUndo = true);
+
+	// Sets the comment text of a comment node in the graph by name.
+	// This causes a CommentTextChanged modified event.
+	UFUNCTION(BlueprintCallable, Category = RigVMController)
+	bool SetCommentTextByName(const FName& InNodeName, const FString& InCommentText, bool bUndo = true);
+
+	// Sets the compactness of a reroute node in the graph.
+	// This causes a RerouteCompactnessChanged modified event.
+	UFUNCTION(BlueprintCallable, Category = RigVMController)
+	bool SetRerouteCompactness(URigVMNode* InNode, bool bShowAsFullNode, bool bUndo = true);
+
+	// Sets the compactness of a reroute node in the graph by name.
+	// This causes a RerouteCompactnessChanged modified event.
+	UFUNCTION(BlueprintCallable, Category = RigVMController)
+	bool SetRerouteCompactnessByName(const FName& InNodeName, bool bShowAsFullNode, bool bUndo = true);
+
+	// Renames a variable in the graph.
+	// This causes a VariableRenamed modified event.
+	UFUNCTION(BlueprintCallable, Category = RigVMController)
+	bool RenameVariable(const FName& InOldName, const FName& InNewName, bool bUndo = true);
+
+	// Renames a parameter in the graph.
+	// This causes a ParameterRenamed modified event.
+	UFUNCTION(BlueprintCallable, Category = RigVMController)
+	bool RenameParameter(const FName& InOldName, const FName& InNewName, bool bUndo = true);
+
+	// Sets the pin to be expanded or not
+	// This causes a PinExpansionChanged modified event.
+	UFUNCTION(BlueprintCallable, Category = RigVMController)
+	bool SetPinExpansion(const FString& InPinPath, bool bIsExpanded, bool bUndo = true);
+
+	// Sets the pin to be watched (or not)
+	// This causes a PinWatchedChanged modified event.
+	UFUNCTION(BlueprintCallable, Category = RigVMController)
+	bool SetPinIsWatched(const FString& InPinPath, bool bIsWatched, bool bUndo = true);
+
+	// Returns the default value of a pin given its pinpath.
+	UFUNCTION(BlueprintCallable, Category = RigVMController)
+	FString GetPinDefaultValue(const FString& InPinPath);
+
+	// Sets the default value of a pin given its pinpath.
+	// This causes a PinDefaultValueChanged modified event.
+	UFUNCTION(BlueprintCallable, Category = RigVMController)
+	bool SetPinDefaultValue(const FString& InPinPath, const FString& InDefaultValue, bool bResizeArrays = true, bool bUndo = true, bool bMergeUndoAction = false);
+
+	// Resets the default value of a pin given its pinpath.
+	// This causes a PinDefaultValueChanged modified event.
+	UFUNCTION(BlueprintCallable, Category = RigVMController)
+	bool ResetPinDefaultValue(const FString& InPinPath, bool bUndo = true);
+
+	// Adds an array element pin to the end of an array pin.
+	// This causes a PinArraySizeChanged modified event.
+	UFUNCTION(BlueprintCallable, Category = RigVMController)
+	FString AddArrayPin(const FString& InArrayPinPath, const FString& InDefaultValue = TEXT(""), bool bUndo = true);
+
+	// Duplicates an array element pin.
+	// This causes a PinArraySizeChanged modified event.
+	UFUNCTION(BlueprintCallable, Category = RigVMController)
+	FString DuplicateArrayPin(const FString& InArrayElementPinPath, bool bUndo = true);
+
+	// Inserts an array element pin into an array pin.
+	// This causes a PinArraySizeChanged modified event.
+	UFUNCTION(BlueprintCallable, Category = RigVMController)
+	FString InsertArrayPin(const FString& InArrayPinPath, int32 InIndex = -1, const FString& InDefaultValue = TEXT(""), bool bUndo = true);
+
+	// Removes an array element pin from an array pin.
+	// This causes a PinArraySizeChanged modified event.
+	UFUNCTION(BlueprintCallable, Category = RigVMController)
+	bool RemoveArrayPin(const FString& InArrayElementPinPath, bool bUndo = true);
+
+	// Removes all (but one) array element pin from an array pin.
+	// This causes a PinArraySizeChanged modified event.
+	UFUNCTION(BlueprintCallable, Category = RigVMController)
+	bool ClearArrayPin(const FString& InArrayPinPath, bool bUndo = true);
+
+	// Sets the size of the array pin
+	// This causes a PinArraySizeChanged modified event.
+	UFUNCTION(BlueprintCallable, Category = RigVMController)
+	bool SetArrayPinSize(const FString& InArrayPinPath, int32 InSize, const FString& InDefaultValue = TEXT(""), bool bUndo = true);
+
+	// Adds a link to the graph.
+	// This causes a LinkAdded modified event.
+	UFUNCTION(BlueprintCallable, Category = RigVMController)
+	bool AddLink(const FString& InOutputPinPath, const FString& InInputPinPath, bool bUndo = true);
+
+	// Removes a link from the graph.
+	// This causes a LinkRemoved modified event.
+	UFUNCTION(BlueprintCallable, Category = RigVMController)
+	bool BreakLink(const FString& InOutputPinPath, const FString& InInputPinPath, bool bUndo = true);
+
+	// Removes all links on a given pin from the graph.
+	// This might cause multiple LinkRemoved modified event.
+	UFUNCTION(BlueprintCallable, Category = RigVMController)
+	bool BreakAllLinks(const FString& InPinPath, bool bAsInput = true, bool bUndo = true);
+
+	// Sets the execute context struct type to use
+	void SetExecuteContextStruct(UStruct* InExecuteContextStruct);
+
+	// A delegate that can be set to change the struct unfolding behaviour
+	FRigVMController_ShouldStructUnfoldDelegate UnfoldStructDelegate;
+
+	int32 DetachLinksFromPinObjects();
+	int32 ReattachLinksToPinObjects();
+
+#if WITH_EDITOR
+	void RepopulatePinsOnNode(URigVMNode* InNode);
+#endif
+
+private:
+
+	UPROPERTY(BlueprintReadOnly, Category = RigVMController, meta = (ScriptName = "ModifiedEvent", AllowPrivateAccess = "true"))
+	FRigVMGraphModifiedDynamicEvent ModifiedEventDynamic;
+
+	FRigVMGraphModifiedEvent ModifiedEventStatic;
+	void HandleModifiedEvent(ERigVMGraphNotifType InNotifType, URigVMGraph* InGraph, UObject* InSubject);
+
+	FString GetValidNodeName(const FString& InPrefix);
+	bool IsValidGraph();
+	bool IsValidNodeForGraph(URigVMNode* InNode);
+	bool IsValidPinForGraph(URigVMPin* InPin);
+	bool IsValidLinkForGraph(URigVMLink* InLink);
+	void AddPinsForStruct(UStruct* InStruct, URigVMNode* InNode, URigVMPin* InParentPin, ERigVMPinDirection InPinDirection, const FString& InDefaultValue, bool bAutoExpandArrays);
+	void AddPinsForArray(FArrayProperty* InArrayProperty, URigVMNode* InNode, URigVMPin* InParentPin, ERigVMPinDirection InPinDirection, const TArray<FString>& InDefaultValues, bool bAutoExpandArrays);
+	void ConfigurePinFromProperty(FProperty* InProperty, URigVMPin* InOutPin, ERigVMPinDirection InPinDirection = ERigVMPinDirection::Invalid);
+	void ConfigurePinFromPin(URigVMPin* InOutPin, URigVMPin* InPin);
+	virtual bool ShouldStructBeUnfolded(const UStruct* InStruct);
+	virtual bool ShouldPinBeUnfolded(URigVMPin* InPin);
+	void SetPinDefaultValue(URigVMPin* InPin, const FString& InDefaultValue, bool bResizeArrays, bool bUndo, bool bMergeUndoAction);
+	bool ResetPinDefaultValue(URigVMPin* InPin, bool bUndo);
+	static TArray<FString> SplitDefaultValue(const FString& InDefaultValue);
+	URigVMPin* InsertArrayPin(URigVMPin* ArrayPin, int32 InIndex, const FString& InDefaultValue, bool bUndo);
+	bool RemovePin(URigVMPin* InPinToRemove, bool bUndo);
+	FProperty* FindPropertyForPin(const FString& InPinPath);
+	bool AddLink(URigVMPin* OutputPin, URigVMPin* InputPin, bool bUndo);
+	bool BreakLink(URigVMPin* OutputPin, URigVMPin* InputPin, bool bUndo);
+	bool BreakAllLinks(URigVMPin* Pin, bool bAsInput, bool bUndo);
+	void BreakAllLinksRecursive(URigVMPin* Pin, bool bAsInput, bool bTowardsParent, bool bUndo);
+	bool SetPinExpansion(URigVMPin* InPin, bool bIsExpanded, bool bUndo = true);
+	void ExpandPinRecursively(URigVMPin* InPin, bool bUndo);
+	bool SetPinIsWatched(URigVMPin* InPin, bool bIsWatched, bool bUndo);
+	bool SetVariableName(URigVMVariableNode* InVariableNode, const FName& InVariableName, bool bUndo);
+	bool SetParameterName(URigVMParameterNode* InParameterNode, const FName& InParameterName, bool bUndo);
+	static void ForEveryPinRecursively(URigVMPin* InPin, TFunction<void(URigVMPin*)> OnEachPinFunction);
+	static void ForEveryPinRecursively(URigVMNode* InNode, TFunction<void(URigVMPin*)> OnEachPinFunction);
+
+	void ReportWarning(const FString& InMessage);
+	void ReportError(const FString& InMessage);
+
+	template <typename FmtType, typename... Types>
+	void ReportWarningf(const FmtType& Fmt, Types... Args)
+	{
+		ReportWarning(FString::Printf(Fmt, Args...));
+	}
+
+	template <typename FmtType, typename... Types>
+	void ReportErrorf(const FmtType& Fmt, Types... Args)
+	{
+		ReportError(FString::Printf(Fmt, Args...));
+	}
+
+	static FLinearColor GetColorFromMetadata(const FString& InMetadata);
+	static void CreateDefaultValueForStructIfRequired(UScriptStruct* InStruct, FString& OutDefaultValue);
+	static void PostProcessDefaultValue(URigVMPin* Pin, FString& OutDefaultValue);
+
+#if WITH_EDITOR
+	void RewireLinks(URigVMPin* OldPin, URigVMPin* NewPin, bool bAsInput, bool bUndo, TArray<URigVMLink*> InLinks = TArray<URigVMLink*>());
+#endif
+
+	UPROPERTY(transient)
+	URigVMGraph* Graph;
+
+	UPROPERTY()
+	UStruct* ExecuteContextStruct;
+
+	UPROPERTY(transient)
+	URigVMActionStack* ActionStack;
+
+	bool bSuspendNotifications;
+	bool bReportWarningsAndErrors;
+
+	friend class URigVMGraph;
+	friend class URigVMActionStack;
+	friend class URigVMCompiler;
+	friend struct FRigVMControllerObjectFactory;
+	friend struct FRigVMAddRerouteNodeAction;
+};
+
diff --git a/Engine/Source/Developer/RigVMDeveloper/Public/RigVMModel/RigVMControllerActions.h b/Engine/Source/Developer/RigVMDeveloper/Public/RigVMModel/RigVMControllerActions.h
new file mode 100644
--- /dev/null
+++ b/Engine/Source/Developer/RigVMDeveloper/Public/RigVMModel/RigVMControllerActions.h
@@ -0,0 +1,934 @@
+// Copyright Epic Games, Inc. All Rights Reserved.
+
+#pragma once
+
+#include "UObject/PropertyPortFlags.h"
+#include "RigVMController.h"
+#if WITH_EDITOR
+#include "ScopedTransaction.h"
+#endif
+#include "RigVMControllerActions.generated.h"
+
+/**
+ * ================================================================================
+ * The RigVMController doesn't rely on the transaction system for performing,
+ * tracking, undoing and redoing changes. Instead it uses an action stack which
+ * stores small serialized structs for each occured action. The reason for this 
+ * is the subscription model and Python support: We need the Graph to broadcast
+ * events to all subscribers - independently from where the action is coming from.
+ * This includes UI views, scripting or undo / redo.
+ * Each action supports the concept of 'merging', so multiple color change actions
+ * for example can be merged into a single action. This avoids the need for tracking
+ * action scope - and makes the integration with UI code simple.
+ * The Controller's ActionStack integrates into the editor's transaction stack
+ * using the ActionIndex property: Transactions on that property cause the 
+ * actionstack to consecutively undo or redo actions until the expected stack
+ * size is reached.
+ * ================================================================================
+ */
+
+struct FRigVMBaseAction;
+
+/**
+ * The action key is used for serializing and storing an action in the stack,
+ * or within another action.
+ */
+USTRUCT()
+struct FRigVMActionKey
+{
+	GENERATED_BODY()
+
+	UPROPERTY()
+	FString ScriptStructPath;
+
+	UPROPERTY()
+	FString ExportedText;
+
+	template<class ActionType>
+	void Set(const ActionType& InAction)
+	{
+		UScriptStruct* ScriptStruct = ActionType::StaticStruct();
+		FRigVMActionKey Key;
+		ScriptStructPath = ScriptStruct->GetPathName();
+		ScriptStruct->ExportText(ExportedText, &InAction, nullptr, nullptr, PPF_None, nullptr);
+	}
+};
+
+/**
+ * The action wrapper is used to extract an action from a serialized key.
+ */
+struct FRigVMActionWrapper
+{
+public:
+	FRigVMActionWrapper(const FRigVMActionKey& Key);
+	~FRigVMActionWrapper();
+
+	FRigVMBaseAction* GetAction();
+	FString ExportText();
+
+private:
+	FRigVMActionWrapper(const FRigVMActionWrapper& Other) = delete;
+	FRigVMActionWrapper& operator = (const FRigVMActionWrapper& Other) = delete;
+
+	UScriptStruct* ScriptStruct;
+	TArray<uint8> Data;
+};
+
+/**
+ * The base action is the base struct for all actions, and provides
+ * access to sub actions, merge functionality as well as undo and redo
+ * base implementations.
+ */
+USTRUCT()
+struct FRigVMBaseAction
+{
+	GENERATED_BODY()
+
+	// Default constructor
+	FRigVMBaseAction()
+		 : Title(TEXT("ACtion"))
+	{
+	}
+
+	// Default destructor
+	virtual ~FRigVMBaseAction() {};
+
+	// Returns the title of the action - used for the Edit menu's undo / redo
+	virtual FString GetTitle() const { return Title; }
+
+	// Trys to merge the action with another action and 
+	// returns true if successfull.
+	virtual bool Merge(const FRigVMBaseAction* Other);
+
+	// Un-does the action and returns true if successfull.
+	virtual bool Undo(URigVMController* InController);
+
+	// Re-does the action and returns true if successfull.
+	virtual bool Redo(URigVMController* InController);
+
+	// Adds a child / sub action to this one
+	template<class ActionType>
+	void AddAction(const ActionType& InAction)
+	{
+		FRigVMActionKey Key;
+		Key.Set<ActionType>(InAction);
+		SubActions.Add(Key);
+	}
+
+	UPROPERTY()
+	FString Title;
+
+	UPROPERTY()
+	TArray<FRigVMActionKey> SubActions;
+};
+
+/**
+ * The Action Stack can be used to track actions happening on a
+ * Graph. Currently the only owner of the ActionStack is the Controller.
+ * Actions can be added to the stack, or they can be understood as
+ * scopes / brackets. For this you can use BeginAction / EndAction / CancelAction
+ * to open / close a bracket. Open brackets automatically record additional
+ * actions occuring during the bracket's lifetime.
+ */
+UCLASS()
+class URigVMActionStack : public UObject
+{
+	GENERATED_BODY()
+
+public:
+
+	// Begins an action and opens a bracket / scope.
+	template<class T>
+	void BeginAction(T& InAction)
+	{
+		if (CurrentActions.Num() > 0)
+		{
+			// catch erroreous duplicate calls to begin action
+			ensure(CurrentActions.Last() != &InAction);
+		}
+		CurrentActions.Add((FRigVMBaseAction*)&InAction);
+
+		ModifiedEvent.Broadcast(ERigVMGraphNotifType::InteractionBracketOpened, nullptr, nullptr);
+	}
+
+	// Ends an action and closes a bracket / scope.
+	template<class ActionType>
+	void EndAction(ActionType& InAction, bool bPerformMerge = false)
+	{
+		ensure(CurrentActions.Num() > 0);
+		ensure((FRigVMBaseAction*)&InAction == CurrentActions.Last());
+		CurrentActions.Pop();
+		AddAction(InAction, bPerformMerge);
+
+		ModifiedEvent.Broadcast(ERigVMGraphNotifType::InteractionBracketClosed, nullptr, nullptr);
+	}
+
+	// Cancels an action, closes a bracket / scope and discards all 
+	// actions to this point.
+	template<class ActionType>
+	void CancelAction(ActionType& InAction)
+	{
+		ensure(CurrentActions.Num() > 0);
+		ensure((FRigVMBaseAction*)&InAction == CurrentActions.Last());
+		CurrentActions.Pop();
+
+		ModifiedEvent.Broadcast(ERigVMGraphNotifType::InteractionBracketCanceled, nullptr, nullptr);
+	}
+
+	// Adds an action to the stack. Optionally this can perform
+	// a potential merge of this action with the previous action to
+	// compact the stack.
+	template<class ActionType>
+	void AddAction(const ActionType& InAction, bool bPerformMerge = false)
+	{
+		TArray<FRigVMActionKey>* ActionList = &UndoActions;
+		if (CurrentActions.Num() > 0)
+		{
+			ActionList = &CurrentActions[CurrentActions.Num()-1]->SubActions;
+		}
+
+		bool bMergeIfPossible = bPerformMerge;
+		if (bMergeIfPossible)
+		{
+			bMergeIfPossible = false;
+			if (ActionList->Num() > 0 && InAction.SubActions.Num() == 0)
+			{
+				if (ActionList->Last().ScriptStructPath == ActionType::StaticStruct()->GetPathName())
+				{
+					FRigVMActionWrapper Wrapper((*ActionList)[ActionList->Num() - 1]);
+					if (Wrapper.GetAction()->SubActions.Num() == 0)
+					{
+						bMergeIfPossible = Wrapper.GetAction()->Merge(&InAction);
+						if (bMergeIfPossible)
+						{
+							(*ActionList)[ActionList->Num()-1].ExportedText = Wrapper.ExportText();
+						}
+					}
+				}
+			}
+		}
+
+		if (!bMergeIfPossible)
+		{
+			FRigVMActionKey Key;
+			Key.Set<ActionType>(InAction);
+
+			ActionList->Add(Key);
+
+			if (CurrentActions.Num() == 0)
+			{
+				RedoActions.Reset();
+
+#if WITH_EDITOR
+				const FScopedTransaction Transaction(FText::FromString(InAction.GetTitle()));
+				SetFlags(RF_Transactional);
+				Modify();
+#endif
+
+				ActionIndex = ActionIndex + 1;
+			}
+		}
+	}
+
+	// Opens an undo bracket / scope to record actions into.
+	// This is primary useful for Python.
+	UFUNCTION()
+	bool OpenUndoBracket(const FString& InTitle);
+
+	// Closes an undo bracket / scope.
+	// This is primary useful for Python.
+	UFUNCTION()
+	bool CloseUndoBracket();
+
+	// Cancels an undo bracket / scope.
+	// This is primary useful for Python.
+	UFUNCTION()
+	bool CancelUndoBracket();
+
+	// Pops the last action from the undo stack and perform undo on it.
+	// Note: This should really only be used for unit tests,
+	// use the GEditor's main Undo method instead.
+	UFUNCTION()
+	bool Undo(URigVMController* InController);
+
+	// Pops the last action from the redo stack and perform redo on it.
+	// Note: This should really only be used for unit tests,
+	// use the GEditor's main Redo method instead.
+	UFUNCTION()
+	bool Redo(URigVMController* InController);
+
+#if WITH_EDITOR
+	virtual void PostTransacted(const FTransactionObjectEvent& TransactionEvent) override;
+#endif
+
+	FRigVMGraphModifiedEvent& OnModified() { return ModifiedEvent; }
+
+private:
+
+	UPROPERTY()
+	int32 ActionIndex;
+
+	UPROPERTY(NonTransactional)
+	TArray<FRigVMActionKey> UndoActions;
+
+	UPROPERTY(NonTransactional)
+	TArray<FRigVMActionKey> RedoActions;
+
+	TArray<FRigVMBaseAction*> CurrentActions;
+	TArray<FRigVMBaseAction*> BracketActions;
+
+	FRigVMGraphModifiedEvent ModifiedEvent;
+};
+
+/**
+ * An action which inverses the child actions,
+ * it performs undo on redo and vice versa.
+ */
+USTRUCT()
+struct FRigVMInverseAction : public FRigVMBaseAction
+{
+	GENERATED_BODY()
+
+	virtual ~FRigVMInverseAction() {};
+	virtual bool Undo(URigVMController* InController);
+	virtual bool Redo(URigVMController* InController);
+};
+
+/**
+ * An action adding a struct node to the graph.
+ */
+USTRUCT()
+struct FRigVMAddStructNodeAction : public FRigVMBaseAction
+{
+	GENERATED_BODY()
+
+public:
+
+	FRigVMAddStructNodeAction();
+	FRigVMAddStructNodeAction(URigVMStructNode* InNode);
+	virtual ~FRigVMAddStructNodeAction() {};
+	virtual bool Undo(URigVMController* InController) override;
+	virtual bool Redo(URigVMController* InController) override;
+
+	UPROPERTY()
+	FString ScriptStructPath;
+
+	UPROPERTY()
+	FName MethodName;
+	
+	UPROPERTY()
+	FVector2D Position;
+
+	UPROPERTY()
+	FString NodePath;
+};
+
+/**
+ * An action adding a variable node to the graph.
+ */
+USTRUCT()
+struct FRigVMAddVariableNodeAction : public FRigVMBaseAction
+{
+	GENERATED_BODY()
+
+public:
+
+	FRigVMAddVariableNodeAction();
+	FRigVMAddVariableNodeAction(URigVMVariableNode* InNode);
+	virtual ~FRigVMAddVariableNodeAction() {};
+	virtual bool Undo(URigVMController* InController) override;
+	virtual bool Redo(URigVMController* InController) override;
+
+	UPROPERTY()
+	FName VariableName;
+	
+	UPROPERTY()
+	FString CPPType;
+
+	UPROPERTY()
+	FString CPPTypeObjectPath;
+
+	UPROPERTY()
+	bool bIsGetter;
+
+	UPROPERTY()
+	FString DefaultValue;
+
+	UPROPERTY()
+	FVector2D Position;
+
+	UPROPERTY()
+	FString NodePath;
+};
+
+/**
+ * An action adding a parameter node to the graph.
+ */
+USTRUCT()
+struct FRigVMAddParameterNodeAction : public FRigVMBaseAction
+{
+	GENERATED_BODY()
+
+public:
+
+	FRigVMAddParameterNodeAction();
+	FRigVMAddParameterNodeAction(URigVMParameterNode* InNode);
+	virtual ~FRigVMAddParameterNodeAction() {};
+	virtual bool Undo(URigVMController* InController) override;
+	virtual bool Redo(URigVMController* InController) override;
+
+	UPROPERTY()
+	FName ParameterName;
+	
+	UPROPERTY()
+	FString CPPType;
+
+	UPROPERTY()
+	FString CPPTypeObjectPath;
+
+	UPROPERTY()
+	bool bIsInput;
+
+	UPROPERTY()
+	FString DefaultValue;
+
+	UPROPERTY()
+	FVector2D Position;
+
+	UPROPERTY()
+	FString NodePath;
+};
+
+/**
+ * An action adding a comment node to the graph.
+ */
+USTRUCT()
+struct FRigVMAddCommentNodeAction : public FRigVMBaseAction
+{
+	GENERATED_BODY()
+
+public:
+
+	FRigVMAddCommentNodeAction();
+	FRigVMAddCommentNodeAction(URigVMCommentNode* InNode);
+	virtual ~FRigVMAddCommentNodeAction() {};
+	virtual bool Undo(URigVMController* InController) override;
+	virtual bool Redo(URigVMController* InController) override;
+
+	UPROPERTY()
+	FString CommentText;
+
+	UPROPERTY()
+	FVector2D Position;
+
+	UPROPERTY()
+	FVector2D Size;
+
+	UPROPERTY()
+	FLinearColor Color;
+
+	UPROPERTY()
+	FString NodePath;
+};
+
+/**
+ * An action adding a reroute node to the graph.
+ */
+USTRUCT()
+struct FRigVMAddRerouteNodeAction : public FRigVMBaseAction
+{
+	GENERATED_BODY()
+
+public:
+
+	FRigVMAddRerouteNodeAction();
+	FRigVMAddRerouteNodeAction(URigVMRerouteNode* InNode);
+	virtual ~FRigVMAddRerouteNodeAction() {};
+	virtual bool Undo(URigVMController* InController) override;
+	virtual bool Redo(URigVMController* InController) override;
+
+	UPROPERTY()
+	bool bShowAsFullNode;
+
+	UPROPERTY()
+	FString CPPType;
+
+	UPROPERTY()
+	FName CPPTypeObjectPath;
+
+	UPROPERTY()
+	FString DefaultValue;
+
+	UPROPERTY()
+	bool bIsConstant;
+
+	UPROPERTY()
+	FName CustomWidgetName;
+
+	UPROPERTY()
+	FVector2D Position;
+
+	UPROPERTY()
+	FString NodePath;
+};
+
+/**
+ * An action adding an injected node to the graph.
+ */
+USTRUCT()
+struct FRigVMAddInjectedNodeAction : public FRigVMBaseAction
+{
+	GENERATED_BODY()
+
+public:
+
+	FRigVMAddInjectedNodeAction();
+	FRigVMAddInjectedNodeAction(URigVMInjectionInfo* InInjectionInfo);
+	virtual ~FRigVMAddInjectedNodeAction() {};
+	virtual bool Undo(URigVMController* InController) override;
+	virtual bool Redo(URigVMController* InController) override;
+
+	UPROPERTY()
+	FString PinPath;
+
+	UPROPERTY()
+	bool bAsInput;
+
+	UPROPERTY()
+	FString ScriptStructPath;
+
+	UPROPERTY()
+	FName MethodName;
+
+	UPROPERTY()
+	FName InputPinName;
+
+	UPROPERTY()
+	FName OutputPinName;
+
+	UPROPERTY()
+	FString NodePath;
+};
+
+/**
+ * An action removing a node from the graph.
+ */
+USTRUCT()
+struct FRigVMRemoveNodeAction : public FRigVMBaseAction
+{
+	GENERATED_BODY()
+
+public:
+
+	FRigVMRemoveNodeAction() {}
+	FRigVMRemoveNodeAction(URigVMNode* InNode);
+	virtual ~FRigVMRemoveNodeAction() {};
+	virtual bool Undo(URigVMController* InController) override;
+	virtual bool Redo(URigVMController* InController) override;
+
+	UPROPERTY()
+	FRigVMActionKey InverseActionKey;
+};
+
+/**
+ * An action selecting or deselecting a node in the graph.
+ */
+USTRUCT()
+struct FRigVMSetNodeSelectionAction : public FRigVMBaseAction
+{
+	GENERATED_BODY()
+
+public:
+
+	FRigVMSetNodeSelectionAction();
+	FRigVMSetNodeSelectionAction(URigVMGraph* InGraph, TArray<FName> InNewSelection);
+	virtual ~FRigVMSetNodeSelectionAction() {};
+	virtual bool Undo(URigVMController* InController) override;
+	virtual bool Redo(URigVMController* InController) override;
+
+	UPROPERTY()
+	TArray<FName> NewSelection;
+
+	UPROPERTY()
+	TArray<FName> OldSelection;
+};
+
+/**
+ * An action setting a node's position in the graph.
+ */
+USTRUCT()
+struct FRigVMSetNodePositionAction : public FRigVMBaseAction
+{
+	GENERATED_BODY()
+
+public:
+
+	FRigVMSetNodePositionAction()
+	{
+		OldPosition = NewPosition = FVector2D::ZeroVector;
+	}
+	FRigVMSetNodePositionAction(URigVMNode* InNode, const FVector2D& InNewPosition);
+	virtual ~FRigVMSetNodePositionAction() {};
+	virtual bool Merge(const FRigVMBaseAction* Other);
+	virtual bool Undo(URigVMController* InController) override;
+	virtual bool Redo(URigVMController* InController) override;
+
+	UPROPERTY()
+	FString NodePath;
+
+	UPROPERTY()
+	FVector2D OldPosition;
+
+	UPROPERTY()
+	FVector2D NewPosition;
+};
+
+/**
+ * An action setting a node's size in the graph.
+ */
+USTRUCT()
+struct FRigVMSetNodeSizeAction : public FRigVMBaseAction
+{
+	GENERATED_BODY()
+
+public:
+
+	FRigVMSetNodeSizeAction()
+	{
+		OldSize = NewSize = FVector2D::ZeroVector;
+	}
+	FRigVMSetNodeSizeAction(URigVMNode* InNode, const FVector2D& InNewSize);
+	virtual ~FRigVMSetNodeSizeAction() {};
+	virtual bool Merge(const FRigVMBaseAction* Other);
+	virtual bool Undo(URigVMController* InController) override;
+	virtual bool Redo(URigVMController* InController) override;
+
+	UPROPERTY()
+	FString NodePath;
+
+	UPROPERTY()
+	FVector2D OldSize;
+
+	UPROPERTY()
+	FVector2D NewSize;
+};
+
+/**
+ * An action setting a node's color in the graph.
+ */
+USTRUCT()
+struct FRigVMSetNodeColorAction : public FRigVMBaseAction
+{
+	GENERATED_BODY()
+
+public:
+
+	FRigVMSetNodeColorAction()
+	{
+		OldColor = NewColor = FLinearColor::Black;
+	}
+	FRigVMSetNodeColorAction(URigVMNode* InNode, const FLinearColor& InNewColor);
+	virtual ~FRigVMSetNodeColorAction() {};
+	virtual bool Merge(const FRigVMBaseAction* Other);
+	virtual bool Undo(URigVMController* InController) override;
+	virtual bool Redo(URigVMController* InController) override;
+
+	UPROPERTY()
+	FString NodePath;
+
+	UPROPERTY()
+	FLinearColor OldColor;
+
+	UPROPERTY()
+	FLinearColor NewColor;
+};
+
+/**
+ * An action setting a comment node's text in the graph.
+ */
+USTRUCT()
+struct FRigVMSetCommentTextAction : public FRigVMBaseAction
+{
+	GENERATED_BODY()
+
+public:
+
+	FRigVMSetCommentTextAction() {}
+	FRigVMSetCommentTextAction(URigVMCommentNode* InNode, const FString& InNewText);
+	virtual ~FRigVMSetCommentTextAction() {};
+	virtual bool Undo(URigVMController* InController) override;
+	virtual bool Redo(URigVMController* InController) override;
+
+	UPROPERTY()
+	FString NodePath;
+
+	UPROPERTY()
+	FString OldText;
+
+	UPROPERTY()
+	FString NewText;
+};
+
+/**
+ * An action setting a reroute node's compactness in the graph.
+ */
+USTRUCT()
+struct FRigVMSetRerouteCompactnessAction : public FRigVMBaseAction
+{
+	GENERATED_BODY()
+
+public:
+
+	FRigVMSetRerouteCompactnessAction();
+	FRigVMSetRerouteCompactnessAction(URigVMRerouteNode* InNode, bool InShowAsFullNode);
+	virtual ~FRigVMSetRerouteCompactnessAction() {};
+	virtual bool Undo(URigVMController* InController) override;
+	virtual bool Redo(URigVMController* InController) override;
+
+	UPROPERTY()
+	FString NodePath;
+
+	UPROPERTY()
+	bool OldShowAsFullNode;
+
+	UPROPERTY()
+	bool NewShowAsFullNode;
+};
+
+/**
+ * An action renaming a variable in the graph.
+ */
+USTRUCT()
+struct FRigVMRenameVariableAction : public FRigVMBaseAction
+{
+	GENERATED_BODY()
+
+public:
+
+	FRigVMRenameVariableAction() {}
+	FRigVMRenameVariableAction(const FName& InOldVariableName, const FName& InNewVariableName);
+	virtual ~FRigVMRenameVariableAction() {};
+	virtual bool Undo(URigVMController* InController) override;
+	virtual bool Redo(URigVMController* InController) override;
+
+	UPROPERTY()
+	FString OldVariableName;
+
+	UPROPERTY()
+	FString NewVariableName;
+};
+
+/**
+ * An action renaming a parameter in the graph.
+ */
+USTRUCT()
+struct FRigVMRenameParameterAction : public FRigVMBaseAction
+{
+	GENERATED_BODY()
+
+public:
+
+	FRigVMRenameParameterAction() {}
+	FRigVMRenameParameterAction(const FName& InOldParameterName, const FName& InNewParameterName);
+	virtual ~FRigVMRenameParameterAction() {};
+	virtual bool Undo(URigVMController* InController) override;
+	virtual bool Redo(URigVMController* InController) override;
+
+	UPROPERTY()
+	FString OldParameterName;
+
+	UPROPERTY()
+	FString NewParameterName;
+};
+
+/**
+ * An action setting a pin's expansion state in the graph.
+ */
+USTRUCT()
+struct FRigVMSetPinExpansionAction : public FRigVMBaseAction
+{
+	GENERATED_BODY()
+
+public:
+
+	FRigVMSetPinExpansionAction()
+	{
+		PinPath = FString();
+		OldIsExpanded = NewIsExpanded = false;
+	}
+
+	FRigVMSetPinExpansionAction(URigVMPin* InPin, bool bNewIsExpanded);
+	virtual ~FRigVMSetPinExpansionAction() {};
+	virtual bool Undo(URigVMController* InController) override;
+	virtual bool Redo(URigVMController* InController) override;
+
+	UPROPERTY()
+	FString PinPath;
+
+	UPROPERTY()
+	bool OldIsExpanded;
+
+	UPROPERTY()
+	bool NewIsExpanded;
+};
+
+/**
+ * An action setting a pin's watch state in the graph.
+ */
+USTRUCT()
+struct FRigVMSetPinWatchAction : public FRigVMBaseAction
+{
+	GENERATED_BODY()
+
+public:
+
+	FRigVMSetPinWatchAction()
+	{
+		OldIsWatched = NewIsWatched = false;
+	}
+
+	FRigVMSetPinWatchAction(URigVMPin* InPin, bool bNewIsWatched);
+	virtual ~FRigVMSetPinWatchAction() {};
+	virtual bool Undo(URigVMController* InController) override;
+	virtual bool Redo(URigVMController* InController) override;
+
+	UPROPERTY()
+	FString PinPath;
+
+	UPROPERTY()
+	bool OldIsWatched;
+
+	UPROPERTY()
+	bool NewIsWatched;
+};
+
+/**
+ * An action setting a pin's default value in the graph.
+ */
+USTRUCT()
+struct FRigVMSetPinDefaultValueAction : public FRigVMBaseAction
+{
+	GENERATED_BODY()
+
+public:
+
+	FRigVMSetPinDefaultValueAction() {}
+	FRigVMSetPinDefaultValueAction(URigVMPin* InPin, const FString& InNewDefaultValue);
+	virtual ~FRigVMSetPinDefaultValueAction() {};
+	virtual bool Merge(const FRigVMBaseAction* Other);
+	virtual bool Undo(URigVMController* InController) override;
+	virtual bool Redo(URigVMController* InController) override;
+
+	UPROPERTY()
+	FString PinPath;
+
+	UPROPERTY()
+	FString OldDefaultValue;
+
+	UPROPERTY()
+	FString NewDefaultValue;
+};
+
+/**
+ * An action inserting a new array pin in the graph.
+ */
+USTRUCT()
+struct FRigVMInsertArrayPinAction : public FRigVMBaseAction
+{
+	GENERATED_BODY()
+
+public:
+
+	FRigVMInsertArrayPinAction()
+		: Index(0)
+	{
+	}
+	FRigVMInsertArrayPinAction(URigVMPin* InArrayPin, int32 InIndex, const FString& InNewDefaultValue);
+	virtual ~FRigVMInsertArrayPinAction() {};
+	virtual bool Undo(URigVMController* InController) override;
+	virtual bool Redo(URigVMController* InController) override;
+
+	UPROPERTY()
+	FString ArrayPinPath;
+
+	UPROPERTY()
+	int32 Index;
+
+	UPROPERTY()
+	FString NewDefaultValue;
+};
+
+/**
+ * An action removing an array pin from the graph.
+ */
+USTRUCT()
+struct FRigVMRemoveArrayPinAction : public FRigVMBaseAction
+{
+	GENERATED_BODY()
+
+public:
+
+	FRigVMRemoveArrayPinAction()
+		: Index(0)
+	{
+	}
+	FRigVMRemoveArrayPinAction(URigVMPin* InArrayElementPin);
+	virtual ~FRigVMRemoveArrayPinAction() {};
+	virtual bool Undo(URigVMController* InController) override;
+	virtual bool Redo(URigVMController* InController) override;
+
+	UPROPERTY()
+	FString ArrayPinPath;
+
+	UPROPERTY()
+	int32 Index;
+
+	UPROPERTY()
+	FString DefaultValue;
+};
+
+/**
+ * An action adding a new link to the graph.
+ */
+USTRUCT()
+struct FRigVMAddLinkAction : public FRigVMBaseAction
+{
+	GENERATED_BODY()
+
+public:
+
+	FRigVMAddLinkAction() {}
+	FRigVMAddLinkAction(URigVMPin* InOutputPin, URigVMPin* InInputPin);
+	virtual ~FRigVMAddLinkAction() {};
+	virtual bool Undo(URigVMController* InController) override;
+	virtual bool Redo(URigVMController* InController) override;
+
+	UPROPERTY()
+	FString OutputPinPath;
+
+	UPROPERTY()
+	FString InputPinPath;
+};
+
+/**
+ * An action removing a link from the graph.
+ */
+USTRUCT()
+struct FRigVMBreakLinkAction : public FRigVMBaseAction
+{
+	GENERATED_BODY()
+
+public:
+
+	FRigVMBreakLinkAction() {}
+	FRigVMBreakLinkAction(URigVMPin* InOutputPin, URigVMPin* InInputPin);
+	virtual ~FRigVMBreakLinkAction() {};
+	virtual bool Undo(URigVMController* InController) override;
+	virtual bool Redo(URigVMController* InController) override;
+
+	UPROPERTY()
+	FString OutputPinPath;
+
+	UPROPERTY()
+	FString InputPinPath;
+};
\ No newline at end of file
diff --git a/Engine/Source/Developer/RigVMDeveloper/Public/RigVMModel/RigVMGraph.h b/Engine/Source/Developer/RigVMDeveloper/Public/RigVMModel/RigVMGraph.h
new file mode 100644
--- /dev/null
+++ b/Engine/Source/Developer/RigVMDeveloper/Public/RigVMModel/RigVMGraph.h
@@ -0,0 +1,116 @@
+// Copyright Epic Games, Inc. All Rights Reserved.
+
+#pragma once
+
+#include "RigVMNode.h"
+#include "RigVMLink.h"
+#include "RigVMNotifications.h"
+#include "Nodes/RigVMVariableNode.h"
+#include "Nodes/RigVMParameterNode.h"
+#include "RigVMCompiler/RigVMAST.h"
+#include "RigVMGraph.generated.h"
+
+/**
+ * The Graph represents a Function definition
+ * using Nodes as statements.
+ * Graphs can be compiled into a URigVM using the 
+ * FRigVMCompiler. 
+ * Graphs provide access to its Nodes, Pins and
+ * Links.
+ */
+UCLASS(BlueprintType)
+class RIGVMDEVELOPER_API URigVMGraph : public UObject
+{
+	GENERATED_BODY()
+
+public:
+
+	// Default constructor
+	URigVMGraph();
+
+	// Returns all of the Nodes within this Graph.
+	UFUNCTION(BlueprintCallable, Category = RigVMGraph)
+	const TArray<URigVMNode*>& GetNodes() const;
+
+	// Returns all of the Links within this Graph.
+	UFUNCTION(BlueprintCallable, Category = RigVMGraph)
+	const TArray<URigVMLink*>& GetLinks() const;
+
+	// Returns a list of unique Variable descriptions within this Graph.
+	// Multiple Variable Nodes can share the same description.
+	UFUNCTION(BlueprintCallable, Category = RigVMGraph)
+	TArray<FRigVMGraphVariableDescription> GetVariableDescriptions() const;
+
+	// Returns a list of unique Parameter descriptions within this Graph.
+	// Multiple Parameter Nodes can share the same description.
+	UFUNCTION(BlueprintCallable, Category = RigVMGraph)
+	TArray<FRigVMGraphParameterDescription> GetParameterDescriptions() const;
+
+	// Returns a Node given its name (or nullptr).
+	UFUNCTION(BlueprintCallable, Category = RigVMGraph)
+	URigVMNode* FindNodeByName(const FName& InNodeName) const;
+
+	// Returns a Node given its path (or nullptr).
+	// (for now this is the same as finding a node by its name.)
+	UFUNCTION(BlueprintCallable, Category = RigVMGraph)
+	URigVMNode* FindNode(const FString& InNodePath) const;
+
+	// Returns a Pin given its path, for example "Node.Color.R".
+	UFUNCTION(BlueprintCallable, Category = RigVMGraph)
+	URigVMPin* FindPin(const FString& InPinPath) const;
+
+	// Returns a link given its string representation,
+	// for example "NodeA.Color.R -> NodeB.Translation.X"
+	UFUNCTION(BlueprintCallable, Category = RigVMGraph)
+	URigVMLink* FindLink(const FString& InLinkPinPathRepresentation) const;
+
+	// Returns true if a Node with a given name is selected.
+	UFUNCTION(BlueprintCallable, Category = RigVMGraph)
+	bool IsNodeSelected(const FName& InNodeName) const;
+
+	// Returns the names of all currently selected Nodes.
+	UFUNCTION(BlueprintCallable, Category = RigVMGraph)
+	const TArray<FName>& GetSelectNodes() const;
+
+	// Returns the modified event, which can be used to 
+	// subscribe to changes happening within the Graph.
+	FRigVMGraphModifiedEvent& OnModified();
+
+	void PrepareCycleChecking(URigVMPin* InPin, bool bAsInput);
+	virtual bool CanLink(URigVMPin* InSourcePin, URigVMPin* InTargetPin, FString* OutFailureReason = nullptr);
+	
+	TSharedPtr<FRigVMParserAST> GetDiagnosticsAST(bool bForceRefresh = false);
+	TSharedPtr<FRigVMParserAST> GetRuntimeAST(const FRigVMParserASTSettings& InSettings = FRigVMParserASTSettings::Optimized(), bool bForceRefresh = false);
+	void ClearAST(bool bClearDiagnostics = true, bool bClearRuntime = true);
+
+private:
+
+	void RepopulatePinLinks();
+
+	FRigVMGraphModifiedEvent ModifiedEvent;
+	void Notify(ERigVMGraphNotifType InNotifType, UObject* InSubject);
+
+	UPROPERTY()
+	TArray<URigVMNode*> Nodes;
+
+	UPROPERTY()
+	TArray<URigVMLink*> Links;
+
+	UPROPERTY()
+	TArray<FName> SelectedNodes;
+
+	TSharedPtr<FRigVMParserAST> DiagnosticsAST;
+	TSharedPtr<FRigVMParserAST> RuntimeAST;
+
+#if WITH_EDITOR
+	TArray<TSharedPtr<FString>> VariableNames;
+	TArray<TSharedPtr<FString>> ParameterNames;
+#endif
+
+	bool IsNameAvailable(const FString& InName);
+
+	friend class URigVMController;
+	friend class URigVMCompiler;
+	friend struct FRigVMControllerObjectFactory;
+};
+
diff --git a/Engine/Source/Developer/RigVMDeveloper/Public/RigVMModel/RigVMPin.h b/Engine/Source/Developer/RigVMDeveloper/Public/RigVMModel/RigVMPin.h
new file mode 100644
--- /dev/null
+++ b/Engine/Source/Developer/RigVMDeveloper/Public/RigVMModel/RigVMPin.h
@@ -0,0 +1,349 @@
+// Copyright Epic Games, Inc. All Rights Reserved.
+
+#pragma once
+
+#include "CoreMinimal.h"
+#include "UObject/ObjectMacros.h"
+#include "RigVMPin.generated.h"
+
+class URigVMGraph;
+class URigVMNode;
+class URigVMStructNode;
+class URigVMPin;
+class URigVMLink;
+
+/**
+ * The Pin Direction is used to differentiate different kinds of 
+ * pins in the data flow graph - inputs, outputs etc.
+ */
+UENUM(BlueprintType)
+enum class ERigVMPinDirection : uint8
+{
+	Input, // A const input value
+	Output, // A mutable output value
+	IO, // A mutable input and output value
+	Visible, // A const value that cannot be connected to
+	Hidden, // A mutable hidden value (used for interal state)
+	Invalid // The max value for this enum - used for guarding.
+};
+
+/**
+ * The Injected Info is used for injecting a node on a pin.
+ * Injected nodes are not visible to the user, but they are normal
+ * nodes on the graph.
+ */
+UCLASS(BlueprintType)
+class RIGVMDEVELOPER_API URigVMInjectionInfo: public UObject
+{
+	GENERATED_BODY()
+
+public:
+
+	URigVMInjectionInfo()
+	{
+		bInjectedAsInput = true;
+	}
+
+	UPROPERTY()
+	URigVMStructNode* StructNode;
+
+	UPROPERTY()
+	bool bInjectedAsInput;
+
+	UPROPERTY()
+	URigVMPin* InputPin;
+
+	UPROPERTY()
+	URigVMPin* OutputPin;
+
+	// Returns the graph of this injected node.
+	UFUNCTION(BlueprintCallable, Category = RigVMInjectionInfo)
+	URigVMGraph* GetGraph() const;
+
+	// Returns the pin of this injected node.
+	UFUNCTION(BlueprintCallable, Category = RigVMInjectionInfo)
+	URigVMPin* GetPin() const;
+};
+
+
+/**
+ * The Visual Debugging Info is used for visually displaying
+ * Data flowing through a pin. Typically this is attached to an input
+ * pin causes the pin to inject a node on the link driving it
+ */
+
+
+/**
+ * The Pin represents a single connector / pin on a node in the RigVM model.
+ * Pins can be connected based on rules. Pins also provide access to a 'PinPath',
+ * which essentially represents . separated list of names to reach the pin within
+ * the owning graph. PinPaths are unique.
+ * In comparison to the EdGraph Pin the URigVMPin supports the concept of 'SubPins',
+ * so child / parent relationships between pins. A FVector Pin for example might
+ * have its X, Y and Z components as SubPins. Array Pins will have its elements as
+ * SubPins, and so on.
+ * A URigVMPin is owned solely by a URigVMNode.
+ */
+UCLASS(BlueprintType)
+class RIGVMDEVELOPER_API URigVMPin : public UObject
+{
+	GENERATED_BODY()
+
+public:
+
+	// Splits a PinPath at the start, so for example "Node.Color.R" becomes "Node" and "Color.R"
+	static bool SplitPinPathAtStart(const FString& InPinPath, FString& LeftMost, FString& Right);
+
+	// Splits a PinPath at the start, so for example "Node.Color.R" becomes "Node.Color" and "R"
+	static bool SplitPinPathAtEnd(const FString& InPinPath, FString& Left, FString& RightMost);
+
+	// Splits a PinPath into all segments, so for example "Node.Color.R" becomes ["Node", "Color", "R"]
+	static bool SplitPinPath(const FString& InPinPath, TArray<FString>& Parts);
+
+	// Joins a PinPath from to segments, so for example "Node.Color" and "R" becomes "Node.Color.R"
+	static FString JoinPinPath(const FString& Left, const FString& Right);
+
+	// Joins a PinPath from to segments, so for example ["Node", "Color", "R"] becomes "Node.Color.R"
+	static FString JoinPinPath(const TArray<FString>& InParts);
+
+	// Default constructor
+	URigVMPin();
+
+	// Returns a . separated path containing all names of the pin and its owners,
+	// this includes the node name, for example "Node.Color.R"
+	UFUNCTION(BlueprintCallable, Category = RigVMPin)
+	FString GetPinPath() const;
+
+	// Returns a . separated path containing all names of the pin within its main
+	// memory owner / storage. This is typically used to create an offset pointer
+	// within memory (FRigVMRegisterOffset).
+	// So for example for a PinPath such as "Node.Transform.Translation.X" the 
+	// corresponding SegmentPath is "Translation.X", since the transform is the
+	// storage / memory.
+	UFUNCTION(BlueprintCallable, Category = RigVMPin)
+	FString GetSegmentPath() const;
+
+	// Returns the display label of the pin
+	UFUNCTION(BlueprintCallable, Category = RigVMPin)
+	FName GetDisplayName() const;
+
+	// Returns the direction of the pin
+	UFUNCTION(BlueprintCallable, Category = RigVMPin)
+	ERigVMPinDirection GetDirection() const;
+
+	// Returns true if the pin is currently expanded
+	UFUNCTION(BlueprintCallable, Category = RigVMPin)
+	bool IsExpanded() const;
+
+	// Returns true if the pin is a constant value / literal
+	UFUNCTION(BlueprintCallable, Category = RigVMPin)
+	bool IsConstant() const;
+
+	// Returns true if the pin should be watched
+	UFUNCTION(BlueprintCallable, Category = RigVMPin)
+	bool RequiresWatch() const;
+
+	// Returns true if the data type of the Pin is a struct
+	UFUNCTION(BlueprintCallable, Category = RigVMPin)
+	bool IsStruct() const;
+
+	// Returns true if the Pin is a SubPin within a struct
+	UFUNCTION(BlueprintCallable, Category = RigVMPin)
+	bool IsStructMember() const;
+
+	// Returns true if the data type of the Pin is an array
+	UFUNCTION(BlueprintCallable, Category = RigVMPin)
+	bool IsArray() const;
+
+	// Returns true if the Pin is a SubPin within an array
+	UFUNCTION(BlueprintCallable, Category = RigVMPin)
+	bool IsArrayElement() const;
+
+	// Returns the index of the Pin within the node / parent Pin
+	UFUNCTION(BlueprintCallable, Category = RigVMPin)
+	int32 GetPinIndex() const;
+
+	// Returns the number of elements within an array Pin
+	UFUNCTION(BlueprintCallable, Category = RigVMPin)
+	int32 GetArraySize() const;
+
+	// Returns the C++ data type of the pin
+	UFUNCTION(BlueprintCallable, Category = RigVMPin)
+	FString GetCPPType() const;
+
+	// Returns the C++ data type of an element of the Pin array
+	UFUNCTION(BlueprintCallable, Category = RigVMPin)
+	FString GetArrayElementCppType() const;
+
+	// Returns true if the C++ data type is FString or FName
+	UFUNCTION(BlueprintCallable, Category = RigVMPin)
+	bool IsStringType() const;
+
+	// Returns true if the C++ data type is an execute context
+	UFUNCTION(BlueprintCallable, Category = RigVMPin)
+	bool IsExecuteContext() const;
+
+	// Returns the default value of the Pin as a string.
+	// Note that this value is computed based on the Pin's
+	// SubPins - so for example for a FVector typed Pin
+	// the default value is actually composed out of the
+	// default values of the X, Y and Z SubPins.
+	UFUNCTION(BlueprintCallable, Category = RigVMPin)
+	FString GetDefaultValue() const;
+
+	// Returns the name of a custom widget to be used
+	// for editing the Pin.
+	UFUNCTION(BlueprintCallable, Category = RigVMPin)
+	FName GetCustomWidgetName() const;
+
+	// Returns the tooltip of this pin
+	UFUNCTION(BlueprintCallable, Category = RigVMPin)
+	FText GetToolTipText() const;
+
+	// Returns the struct of the data type of the Pin,
+	// or nullptr otherwise.
+	UFUNCTION(BlueprintCallable, Category = RigVMPin)
+	UObject* GetCPPTypeObject() const;
+
+	// Returns the struct of the data type of the Pin,
+	// or nullptr otherwise.
+	UFUNCTION(BlueprintCallable, Category = RigVMPin)
+	UScriptStruct* GetScriptStruct() const;
+
+	// Returns the enum of the data type of the Pin,
+	// or nullptr otherwise.
+	UFUNCTION(BlueprintCallable, Category = RigVMPin)
+	UEnum* GetEnum() const;
+
+	// Returns the parent Pin - or nullptr if the Pin
+	// is nested directly below a node.
+	UFUNCTION(BlueprintCallable, Category = RigVMPin)
+	URigVMPin* GetParentPin() const;
+
+	// Returns the top-most parent Pin, so for example
+	// for "Node.Transform.Translation.X" this returns
+	// the Pin for "Node.Transform".
+	UFUNCTION(BlueprintCallable, Category = RigVMPin)
+	URigVMPin* GetRootPin() const;
+
+	// Returns the pin to be used for a link.
+	// This might differ from this actual pin, since
+	// the pin might contain injected nodes.
+	UFUNCTION(BlueprintCallable, Category = RigVMPin)
+	URigVMPin* GetPinForLink() const;
+
+	// Returns the original pin for a pin on an injected
+	// node. This can be used to determine where a link
+	// should go in the user interface
+	UFUNCTION(BlueprintCallable, Category = RigVMPin)
+	URigVMPin* GetOriginalPinFromInjectedNode() const;
+
+	// Returns all of the SubPins of this one.
+	UFUNCTION(BlueprintCallable, Category = RigVMPin)
+	const TArray<URigVMPin*>& GetSubPins() const;
+
+	// Returns a SubPin given a name / path or nullptr.
+	UFUNCTION(BlueprintCallable, Category = RigVMPin)
+	URigVMPin* FindSubPin(const FString& InPinPath) const;
+
+	// Returns true if this Pin is linked to another Pin
+	UFUNCTION(BlueprintCallable, Category = RigVMPin)
+	bool IsLinkedTo(URigVMPin* InPin) const;
+
+	// Returns all of the links linked to this Pin.
+	UFUNCTION(BlueprintCallable, Category = RigVMPin)
+	const TArray<URigVMLink*>& GetLinks() const;
+
+	// Returns all of the linked source Pins,
+	// using this Pin as the target.
+	UFUNCTION(BlueprintCallable, Category = RigVMPin)
+	TArray<URigVMPin*> GetLinkedSourcePins(bool bRecursive = false) const;
+
+	// Returns all of the linked target Pins,
+	// using this Pin as the source.
+	UFUNCTION(BlueprintCallable, Category = RigVMPin)
+	TArray<URigVMPin*> GetLinkedTargetPins(bool bRecursive = false) const;
+
+	// Returns all of the source pins
+	// using this Pin as the target.
+	UFUNCTION(BlueprintCallable, Category = RigVMPin)
+	TArray<URigVMLink*> GetSourceLinks(bool bRecursive = false) const;
+
+	// Returns all of the target links,
+	// using this Pin as the source.
+	UFUNCTION(BlueprintCallable, Category = RigVMPin)
+	TArray<URigVMLink*> GetTargetLinks(bool bRecursive = false) const;
+
+	// Returns the node of this Pin.
+	UFUNCTION(BlueprintCallable, Category = RigVMPin)
+	URigVMNode* GetNode() const;
+
+	// Returns the graph of this Pin.
+	UFUNCTION(BlueprintCallable, Category = RigVMPin)
+	URigVMGraph* GetGraph() const;
+
+	// Returns true is the two provided source and target Pins
+	// can be linked to one another.
+	static bool CanLink(URigVMPin* InSourcePin, URigVMPin* InTargetPin, FString* OutFailureReason = nullptr);
+
+	// Returns true if this pin has injected nodes
+	bool HasInjectedNodes() const { return InjectionInfos.Num() > 0; }
+
+	// Returns the injected nodes this pin contains.
+	const TArray<URigVMInjectionInfo*> GetInjectedNodes() const { return InjectionInfos; }
+
+	// helper function to retrieve an object from a path
+	static UObject* FindObjectFromCPPTypeObjectPath(const FString& InObjectPath);
+	template<class T>
+	FORCEINLINE static T* FindObjectFromCPPTypeObjectPath(const FString& InObjectPath)
+	{
+		return Cast<T>(FindObjectFromCPPTypeObjectPath(InObjectPath));
+	}
+
+private:
+
+	void UpdateCPPTypeObjectIfRequired() const;
+	void SetNameFromIndex();
+
+	UPROPERTY()
+	FName DisplayName;
+
+	UPROPERTY()
+	ERigVMPinDirection Direction;
+
+	UPROPERTY()
+	bool bIsExpanded;
+
+	UPROPERTY()
+	bool bIsConstant;
+
+	UPROPERTY(transient)
+	bool bRequiresWatch;
+
+	UPROPERTY()
+	FString CPPType;
+
+	UPROPERTY(transient)
+	UObject* CPPTypeObject;
+
+	UPROPERTY()
+	FName CPPTypeObjectPath;
+
+	UPROPERTY()
+	FString DefaultValue;
+
+	UPROPERTY()
+	FName CustomWidgetName;
+
+	UPROPERTY()
+	TArray<URigVMPin*> SubPins;
+
+	UPROPERTY(transient)
+	TArray<URigVMLink*> Links;
+
+	UPROPERTY()
+	TArray<URigVMInjectionInfo*> InjectionInfos;
+
+	friend class URigVMController;
+	friend class URigVMGraph;
+};
diff --git a/Engine/Source/Editor/AnimGraph/Classes/AnimGraphNode_ApplyMeshSpaceAdditive.h b/Engine/Source/Editor/AnimGraph/Classes/AnimGraphNode_ApplyMeshSpaceAdditive.h
--- a/Engine/Source/Editor/AnimGraph/Classes/AnimGraphNode_ApplyMeshSpaceAdditive.h
+++ b/Engine/Source/Editor/AnimGraph/Classes/AnimGraphNode_ApplyMeshSpaceAdditive.h
@@ -17,15 +17,16 @@ class UAnimGraphNode_ApplyMeshSpaceAdditive : public UAnimGraphNode_Base
 	FAnimNode_ApplyMeshSpaceAdditive Node;
 
 	// UEdGraphNode interface
-	virtual FText GetTooltipText() const override;
-	virtual FText GetNodeTitle(ENodeTitleType::Type TitleType) const override;
-	virtual FLinearColor GetNodeTitleColor() const override;
-	virtual void PostEditChangeProperty(struct FPropertyChangedEvent& PropertyChangedEvent) override;
+	FText GetTooltipText() const override;
+	FText GetNodeTitle(ENodeTitleType::Type TitleType) const override;
+	FLinearColor GetNodeTitleColor() const override;
+	void PostEditChangeProperty(struct FPropertyChangedEvent& PropertyChangedEvent) override;
 	// End of UEdGraphNode interface
 
 	// UAnimGraphNode_Base interface
-	virtual void CustomizePinData(UEdGraphPin* Pin, FName SourcePropertyName, int32 ArrayIndex) const override;
-	virtual FString GetNodeCategory() const override;
-	virtual void ValidateAnimNodeDuringCompilation(class USkeleton* ForSkeleton, class FCompilerResultsLog& MessageLog) override;
+	FString GetNodeCategory() const override;
+	void CustomizePinData(UEdGraphPin* Pin, FName SourcePropertyName, int32 ArrayIndex) const override;
+	void CustomizeDetails(IDetailLayoutBuilder& DetailBuilder) override;
+	void ValidateAnimNodeDuringCompilation(class USkeleton* ForSkeleton, class FCompilerResultsLog& MessageLog) override;
 	// End of UAnimGraphNode_Base interface
 };
diff --git a/Engine/Source/Editor/AnimGraph/Classes/AnimPreviewAttacheInstance.h b/Engine/Source/Editor/AnimGraph/Classes/AnimPreviewAttacheInstance.h
--- a/Engine/Source/Editor/AnimGraph/Classes/AnimPreviewAttacheInstance.h
+++ b/Engine/Source/Editor/AnimGraph/Classes/AnimPreviewAttacheInstance.h
@@ -4,7 +4,7 @@
 
 #include "CoreMinimal.h"
 #include "UObject/ObjectMacros.h"
-#include "AnimCustomInstance.h"
+#include "AnimCustomInstanceHelper.h"
 #include "Animation/AnimInstanceProxy.h"
 #include "AnimNodes/AnimNode_CopyPoseFromMesh.h"
 #include "AnimPreviewAttacheInstance.generated.h"
@@ -41,7 +41,7 @@ private:
  */
 
 UCLASS(transient, NotBlueprintable, noteditinlinenew)
-class ANIMGRAPH_API UAnimPreviewAttacheInstance : public UAnimCustomInstance
+class ANIMGRAPH_API UAnimPreviewAttacheInstance : public UAnimInstance
 {
 	GENERATED_UCLASS_BODY()
 
diff --git a/Engine/Source/Editor/AnimGraph/Classes/AnimPreviewInstance.h b/Engine/Source/Editor/AnimGraph/Classes/AnimPreviewInstance.h
--- a/Engine/Source/Editor/AnimGraph/Classes/AnimPreviewInstance.h
+++ b/Engine/Source/Editor/AnimGraph/Classes/AnimPreviewInstance.h
@@ -83,27 +83,22 @@ public:
 		SkeletalControlAlpha = FMath::Clamp<float>(InSkeletalControlAlpha, 0.f, 1.f);
 	}
 
-	void SetKey(FSimpleDelegate InOnSetKeyCompleteDelegate)
+#if WITH_EDITOR	
+	void SetKey()
 	{
-#if WITH_EDITOR
 		bSetKey = true;
-		OnSetKeyCompleteDelegate = InOnSetKeyCompleteDelegate;
-#endif
 	}
 
-	void SetKey()
+	FDelegateHandle AddKeyCompleteDelegate(FSimpleMulticastDelegate::FDelegate InOnSetKeyCompleteDelegate)
 	{
-#if WITH_EDITOR
-		bSetKey = true;
-#endif
+		return OnSetKeyCompleteDelegate.Add(InOnSetKeyCompleteDelegate);
 	}
 
-	void SetKeyCompleteDelegate(FSimpleDelegate InOnSetKeyCompleteDelegate)
+	void RemoveKeyCompleteDelegate(FDelegateHandle InDelegateHandle)
 	{
-#if WITH_EDITOR
-		OnSetKeyCompleteDelegate = InOnSetKeyCompleteDelegate;
-#endif
+		OnSetKeyCompleteDelegate.Remove(InDelegateHandle);
 	}
+#endif
 
 	void RefreshCurveBoneControllers(UAnimationAsset* AssetToRefreshFrom);
 
@@ -153,7 +148,7 @@ private:
 	/**
 	 * Delegate to call after Key is set
 	 */
-	FSimpleDelegate OnSetKeyCompleteDelegate;
+	FSimpleMulticastDelegate OnSetKeyCompleteDelegate;
 
 	/** Shared parameters for previewing blendspace or animsequence **/
 	float SkeletalControlAlpha;
@@ -266,30 +261,29 @@ public:
 
 	bool GetForceRetargetBasePose() const;
 
+#if WITH_EDITOR	
 	/**
 	 * Convert current modified bone transforms (BoneControllers) to transform curves (CurveControllers)
 	 * it does based on CurrentTime. This function does not set key directly here. 
 	 * It does wait until next update, and it gets the delta of transform before applying curves, and 
-	 * creates curves from it, so you'll need delegate if you'd like to do something after
-	 * 
-	 * @param Delegate To be called once set key is completed
+	 * creates curves from it, so you'll need delegate if you'd like to do something after (set with SetKeyCompleteDelegate)
 	 */
-	void SetKey(FSimpleDelegate InOnSetKeyCompleteDelegate);
+	void SetKey();
 
 	/**
-	 * Convert current modified bone transforms (BoneControllers) to transform curves (CurveControllers)
-	 * it does based on CurrentTime. This function does not set key directly here. 
-	 * It does wait until next update, and it gets the delta of transform before applying curves, and 
-	 * creates curves from it, so you'll need delegate if you'd like to do something after (set with SetKeyCompleteDelegate)
+	 * Add the delegate to be called when a key is set.
+	 * 
+	 * @param Delegate To be called once set key is completed
 	 */
-	void SetKey();
+	FDelegateHandle AddKeyCompleteDelegate(FSimpleMulticastDelegate::FDelegate InOnSetKeyCompleteDelegate);
 
 	/**
-	 * Set the delegate to be called when a key is set.
+	 * Add the delegate to be called when a key is set.
 	 * 
 	 * @param Delegate To be called once set key is completed
 	 */
-	void SetKeyCompleteDelegate(FSimpleDelegate InOnSetKeyCompleteDelegate);
+	void RemoveKeyCompleteDelegate(FDelegateHandle InDelegateHandle);
+#endif
 
 	/** 
 	 * Refresh Curve Bone Controllers based on TransformCurves from Animation data
diff --git a/Engine/Source/Editor/AnimGraph/Private/AnimGraphNode_ApplyMeshSpaceAdditive.cpp b/Engine/Source/Editor/AnimGraph/Private/AnimGraphNode_ApplyMeshSpaceAdditive.cpp
--- a/Engine/Source/Editor/AnimGraph/Private/AnimGraphNode_ApplyMeshSpaceAdditive.cpp
+++ b/Engine/Source/Editor/AnimGraph/Private/AnimGraphNode_ApplyMeshSpaceAdditive.cpp
@@ -2,9 +2,14 @@
 
 #include "AnimGraphNode_ApplyMeshSpaceAdditive.h"
 
+#include "DetailLayoutBuilder.h"
+#include "ScopedTransaction.h"
+
 #include "Animation/AnimationSettings.h"
+#include "Kismet2/BlueprintEditorUtils.h"
 #include "Kismet2/CompilerResultsLog.h"
 
+
 /////////////////////////////////////////////////////
 // UAnimGraphNode_ApplyAdditive
 
@@ -36,22 +41,117 @@ void UAnimGraphNode_ApplyMeshSpaceAdditive::CustomizePinData(UEdGraphPin* Pin, F
 
 	if (Pin->PinName == GET_MEMBER_NAME_STRING_CHECKED(FAnimNode_ApplyMeshSpaceAdditive, Alpha))
 	{
+		Pin->bHidden = Node.AlphaInputType != EAnimAlphaInputType::Float;
+
+		if (!Pin->bHidden)
+		{
+			Pin->PinFriendlyName = Node.AlphaScaleBiasClamp.GetFriendlyName(Node.AlphaScaleBias.GetFriendlyName(Pin->PinFriendlyName));
+		}
+	}
+	else if (Pin->PinName == GET_MEMBER_NAME_STRING_CHECKED(FAnimNode_ApplyMeshSpaceAdditive, bAlphaBoolEnabled))
+	{
+		Pin->bHidden = Node.AlphaInputType != EAnimAlphaInputType::Bool;
+	}
+	else if (Pin->PinName == GET_MEMBER_NAME_STRING_CHECKED(FAnimNode_ApplyMeshSpaceAdditive, AlphaCurveName))
+	{
+		Pin->bHidden = Node.AlphaInputType != EAnimAlphaInputType::Curve;
+
 		if (!Pin->bHidden)
 		{
-			Pin->PinFriendlyName = Node.AlphaScaleBias.GetFriendlyName(Pin->PinFriendlyName);
+			Pin->PinFriendlyName = Node.AlphaScaleBiasClamp.GetFriendlyName(Pin->PinFriendlyName);
 		}
 	}
 }
 
+void UAnimGraphNode_ApplyMeshSpaceAdditive::CustomizeDetails(IDetailLayoutBuilder& DetailBuilder)
+{
+	Super::CustomizeDetails(DetailBuilder);
+
+	TSharedRef<IPropertyHandle> NodeHandle = DetailBuilder.GetProperty(FName(TEXT("Node")), GetClass());
+
+	if (Node.AlphaInputType != EAnimAlphaInputType::Bool)
+	{
+		DetailBuilder.HideProperty(NodeHandle->GetChildHandle(GET_MEMBER_NAME_CHECKED(FAnimNode_ApplyMeshSpaceAdditive, bAlphaBoolEnabled)));
+		DetailBuilder.HideProperty(NodeHandle->GetChildHandle(GET_MEMBER_NAME_CHECKED(FAnimNode_ApplyMeshSpaceAdditive, AlphaBoolBlend)));
+	}
+
+	if (Node.AlphaInputType != EAnimAlphaInputType::Float)
+	{
+		DetailBuilder.HideProperty(NodeHandle->GetChildHandle(GET_MEMBER_NAME_CHECKED(FAnimNode_ApplyMeshSpaceAdditive, Alpha)));
+		DetailBuilder.HideProperty(NodeHandle->GetChildHandle(GET_MEMBER_NAME_CHECKED(FAnimNode_ApplyMeshSpaceAdditive, AlphaScaleBias)));
+	}
+
+	if (Node.AlphaInputType != EAnimAlphaInputType::Curve)
+	{
+		DetailBuilder.HideProperty(NodeHandle->GetChildHandle(GET_MEMBER_NAME_CHECKED(FAnimNode_ApplyMeshSpaceAdditive, AlphaCurveName)));
+	}
+
+	if (Node.AlphaInputType == EAnimAlphaInputType::Bool)
+	{
+		DetailBuilder.HideProperty(NodeHandle->GetChildHandle(GET_MEMBER_NAME_CHECKED(FAnimNode_ApplyMeshSpaceAdditive, AlphaScaleBiasClamp)));
+	}
+}
+
 void UAnimGraphNode_ApplyMeshSpaceAdditive::PostEditChangeProperty(struct FPropertyChangedEvent& PropertyChangedEvent)
 {
 	const FName PropertyName = (PropertyChangedEvent.Property ? PropertyChangedEvent.Property->GetFName() : NAME_None);
 
 	// Reconstruct node to show updates to PinFriendlyNames.
-	if ((PropertyName == GET_MEMBER_NAME_STRING_CHECKED(FAnimNode_ApplyMeshSpaceAdditive, AlphaScaleBias)))
+	if (PropertyName == GET_MEMBER_NAME_STRING_CHECKED(FAnimNode_ApplyMeshSpaceAdditive, AlphaScaleBias) || 
+		PropertyName == GET_MEMBER_NAME_STRING_CHECKED(FInputScaleBiasClamp, bMapRange) || 
+		PropertyName == GET_MEMBER_NAME_STRING_CHECKED(FInputRange, Min) || 
+		PropertyName == GET_MEMBER_NAME_STRING_CHECKED(FInputRange, Max) || 
+		PropertyName == GET_MEMBER_NAME_STRING_CHECKED(FInputScaleBiasClamp, Scale) || 
+		PropertyName == GET_MEMBER_NAME_STRING_CHECKED(FInputScaleBiasClamp, Bias) || 
+		PropertyName == GET_MEMBER_NAME_STRING_CHECKED(FInputScaleBiasClamp, bClampResult) || 
+		PropertyName == GET_MEMBER_NAME_STRING_CHECKED(FInputScaleBiasClamp, ClampMin) || 
+		PropertyName == GET_MEMBER_NAME_STRING_CHECKED(FInputScaleBiasClamp, ClampMax) || 
+		PropertyName == GET_MEMBER_NAME_STRING_CHECKED(FInputScaleBiasClamp, bInterpResult) || 
+		PropertyName == GET_MEMBER_NAME_STRING_CHECKED(FInputScaleBiasClamp, InterpSpeedIncreasing) || 
+		PropertyName == GET_MEMBER_NAME_STRING_CHECKED(FInputScaleBiasClamp, InterpSpeedDecreasing))
 	{
 		ReconstructNode();
 	}
+	else if (PropertyName == GET_MEMBER_NAME_STRING_CHECKED(FAnimNode_ApplyMeshSpaceAdditive, AlphaInputType))
+	{
+		UEdGraphPin* PinToBreakLinksTo = nullptr;
+
+		// Break links to pins that are getting hidden.
+		for (UEdGraphPin* Pin: Pins)
+		{
+			if (Pin->PinName == GET_MEMBER_NAME_STRING_CHECKED(FAnimNode_ApplyMeshSpaceAdditive, Alpha) && 
+				Node.AlphaInputType != EAnimAlphaInputType::Float)
+			{
+				PinToBreakLinksTo = Pin;
+				break;
+			}
+			else if (Pin->PinName == GET_MEMBER_NAME_STRING_CHECKED(FAnimNode_ApplyMeshSpaceAdditive, bAlphaBoolEnabled) && 
+					 Node.AlphaInputType != EAnimAlphaInputType::Bool)
+			{
+				PinToBreakLinksTo = Pin;
+				break;
+			}
+			else if (Pin->PinName == GET_MEMBER_NAME_STRING_CHECKED(FAnimNode_ApplyMeshSpaceAdditive, AlphaCurveName) &&
+					 Node.AlphaInputType != EAnimAlphaInputType::Curve)
+			{
+				PinToBreakLinksTo = Pin;
+				break;
+			}
+		}
+
+		if (PinToBreakLinksTo != nullptr)
+		{
+			FScopedTransaction Transaction(LOCTEXT("ChangeAlphaInputType", "Change Alpha Input Type"));
+			Modify();
+			PinToBreakLinksTo->BreakAllPinLinks();
+			ReconstructNode();
+
+			FBlueprintEditorUtils::MarkBlueprintAsStructurallyModified(GetBlueprint());
+		}
+
+	}
+
+
 
 	Super::PostEditChangeProperty(PropertyChangedEvent);
 }
@@ -59,7 +159,6 @@ void UAnimGraphNode_ApplyMeshSpaceAdditive::PostEditChangeProperty(struct FPrope
 FString UAnimGraphNode_ApplyMeshSpaceAdditive::GetNodeCategory() const
 {
 	return TEXT("Blends");
-	//@TODO: TEXT("Apply additive to normal pose"), TEXT("Apply additive pose"));
 }
 
 void UAnimGraphNode_ApplyMeshSpaceAdditive::ValidateAnimNodeDuringCompilation(class USkeleton* ForSkeleton, class FCompilerResultsLog& MessageLog)
diff --git a/Engine/Source/Editor/AnimGraph/Private/AnimGraphNode_PoseDriver.cpp b/Engine/Source/Editor/AnimGraph/Private/AnimGraphNode_PoseDriver.cpp
--- a/Engine/Source/Editor/AnimGraph/Private/AnimGraphNode_PoseDriver.cpp
+++ b/Engine/Source/Editor/AnimGraph/Private/AnimGraphNode_PoseDriver.cpp
@@ -118,6 +118,18 @@ void UAnimGraphNode_PoseDriver::ValidateAnimNodeDuringCompilation(USkeleton* For
 		MessageLog.Warning(*LOCTEXT("SourceBoneNotFound", "@@ - Entry in SourceBones not found").ToString(), this);
 	}
 
+	TArray<FRBFTarget> RBFTargets;
+	Node.GetRBFTargets(RBFTargets);
+	TArray<int> InvalidTargets;
+	if (!FRBFSolver::ValidateTargets(Node.RBFParams, RBFTargets, InvalidTargets))
+	{
+		for (int TargetIdx: InvalidTargets)
+		{
+			MessageLog.Error(*LOCTEXT("PoseDriver_InvalidTarget", "@@ - '@@' is an invalid or duplicate target.").ToString(),
+				this, GetData(Node.PoseTargets[TargetIdx].DrivenName.ToString()));
+		}
+	}
+
 	Super::ValidateAnimNodeDuringCompilation(ForSkeleton, MessageLog);
 }
 
@@ -209,6 +221,7 @@ void UAnimGraphNode_PoseDriver::CopyNodeDataToPreviewNode(FAnimNode_Base* InPrev
 {
 	FAnimNode_PoseDriver* PreviewPoseDriver = static_cast<FAnimNode_PoseDriver*>(InPreviewNode);
 
+	PreviewPoseDriver->RBFParams.SolverType = Node.RBFParams.SolverType;
 	PreviewPoseDriver->RBFParams.Radius = Node.RBFParams.Radius;
 	PreviewPoseDriver->RBFParams.Function = Node.RBFParams.Function;
 	PreviewPoseDriver->RBFParams.DistanceMethod = Node.RBFParams.DistanceMethod;
diff --git a/Engine/Source/Editor/AnimGraph/Private/AnimPreviewInstance.cpp b/Engine/Source/Editor/AnimGraph/Private/AnimPreviewInstance.cpp
--- a/Engine/Source/Editor/AnimGraph/Private/AnimPreviewInstance.cpp
+++ b/Engine/Source/Editor/AnimGraph/Private/AnimPreviewInstance.cpp
@@ -4,6 +4,7 @@
 #include "AnimPreviewInstance.h"
 #include "Animation/DebugSkelMeshComponent.h"
 #include "AnimationRuntime.h"
+#include "Animation/AnimSequence.h"
 
 #if WITH_EDITOR
 #include "ScopedTransaction.h"
@@ -412,7 +413,7 @@ void FAnimPreviewInstanceProxy::SetKeyImplementation(const FCompactPose& PreCont
 
 		ResetModifiedBone(false);
 
-		OnSetKeyCompleteDelegate.ExecuteIfBound();
+		OnSetKeyCompleteDelegate.Broadcast();
 	}
 #endif
 }
@@ -518,21 +519,25 @@ void UAnimPreviewInstance::ResetModifiedBone(bool bCurveController/*=false*/)
 	GetProxyOnGameThread<FAnimPreviewInstanceProxy>().ResetModifiedBone(bCurveController);
 }
 
-void UAnimPreviewInstance::SetKey(FSimpleDelegate InOnSetKeyCompleteDelegate)
-{
-	GetProxyOnGameThread<FAnimPreviewInstanceProxy>().SetKey(InOnSetKeyCompleteDelegate);
-}
+#if WITH_EDITOR	
 
 void UAnimPreviewInstance::SetKey()
 {
 	GetProxyOnGameThread<FAnimPreviewInstanceProxy>().SetKey();
 }
 
-void UAnimPreviewInstance::SetKeyCompleteDelegate(FSimpleDelegate InOnSetKeyCompleteDelegate)
+FDelegateHandle UAnimPreviewInstance::AddKeyCompleteDelegate(FSimpleMulticastDelegate::FDelegate InOnSetKeyCompleteDelegate)
 {
-	GetProxyOnGameThread<FAnimPreviewInstanceProxy>().SetKeyCompleteDelegate(InOnSetKeyCompleteDelegate);
+	return GetProxyOnGameThread<FAnimPreviewInstanceProxy>().AddKeyCompleteDelegate(InOnSetKeyCompleteDelegate);
 }
 
+void UAnimPreviewInstance::RemoveKeyCompleteDelegate(FDelegateHandle InDelegateHandle)
+{
+	GetProxyOnGameThread<FAnimPreviewInstanceProxy>().RemoveKeyCompleteDelegate(InDelegateHandle);
+}
+
+#endif
+
 void UAnimPreviewInstance::RefreshCurveBoneControllers()
 {
 	GetProxyOnGameThread<FAnimPreviewInstanceProxy>().RefreshCurveBoneControllers(CurrentAsset);
@@ -703,7 +708,7 @@ void UAnimPreviewInstance::MontagePreview_StepForward()
 		// Add DELTA to prefer next frame when we're close to the boundary
 		float CurrentFraction = Proxy.GetCurrentTime() / Montage->SequenceLength + DELTA;
 		float NextFrame = FMath::Clamp<float>(FMath::FloorToFloat(CurrentFraction * NumFrames) + 1.0f, 0, NumFrames);
-		float NewTime = Montage->SequenceLength * (NextFrame / NumFrames);
+		float NewTime = Montage->SequenceLength * (NextFrame / (NumFrames-1));
 
 		GetSkelMeshComponent()->GlobalAnimRateScale = 1.0f;
 		GetSkelMeshComponent()->TickAnimation(NewTime - Proxy.GetCurrentTime(), false);
@@ -759,7 +764,7 @@ void UAnimPreviewInstance::MontagePreview_StepBackward()
 		// Add DELTA to prefer next frame when we're close to the boundary
 		float CurrentFraction = Proxy.GetCurrentTime() / Montage->SequenceLength + DELTA;
 		float NextFrame = FMath::Clamp<float>(FMath::FloorToFloat(CurrentFraction * NumFrames) - 1.0f, 0, NumFrames);
-		float NewTime = Montage->SequenceLength * (NextFrame / NumFrames);
+		float NewTime = Montage->SequenceLength * (NextFrame / (NumFrames-1));
 
 		GetSkelMeshComponent()->GlobalAnimRateScale = 1.0f;
 		GetSkelMeshComponent()->TickAnimation(FMath::Abs(NewTime - Proxy.GetCurrentTime()), false);
diff --git a/Engine/Source/Editor/AnimGraph/Private/PoseDriverDetails.cpp b/Engine/Source/Editor/AnimGraph/Private/PoseDriverDetails.cpp
--- a/Engine/Source/Editor/AnimGraph/Private/PoseDriverDetails.cpp
+++ b/Engine/Source/Editor/AnimGraph/Private/PoseDriverDetails.cpp
@@ -257,6 +257,7 @@ TSharedRef< SWidget > SPDD_TargetRow::GenerateWidgetForColumn(const FName& Colum
 							.OptionsSource(&DistanceMethodOptions)
 							.OnGenerateWidget(this, &SPDD_TargetRow::MakeDrivenNameWidget)
 							.OnSelectionChanged(this, &SPDD_TargetRow::OnDistanceMethodChanged)
+							.IsEnabled_Lambda([this]() { return IsOverrideEnabled(); })
 							.Content()
 							[
 								SNew(STextBlock)
@@ -274,6 +275,7 @@ TSharedRef< SWidget > SPDD_TargetRow::GenerateWidgetForColumn(const FName& Colum
 							.OptionsSource(&FunctionTypeOptions)
 							.OnGenerateWidget(this, &SPDD_TargetRow::MakeDrivenNameWidget)
 							.OnSelectionChanged(this, &SPDD_TargetRow::OnFunctionTypeChanged)
+							.IsEnabled_Lambda([this]() { return IsOverrideEnabled(); })
 							.Content()
 							[
 								SNew(STextBlock)
@@ -316,6 +318,7 @@ TSharedRef< SWidget > SPDD_TargetRow::GenerateWidgetForColumn(const FName& Colum
 					[
 						SNew(SBox)
 						.Visibility_Lambda([=]() { return IsCustomCurveEnabled() ? EVisibility::Visible : EVisibility::Collapsed;  })
+						.IsEnabled_Lambda([this]() { return IsOverrideEnabled(); })
 						.Content()
 						[
 							SAssignNew(CurveEditor, SCurveEditor)
@@ -437,6 +440,18 @@ int32 SPDD_TargetRow::GetTargetIndex() const
 	return TargetInfo->TargetIndex;
 }
 
+bool SPDD_TargetRow::IsOverrideEnabled() const
+{
+	bool bOverrideEnabled = true;
+	UAnimGraphNode_PoseDriver* Driver = GetPoseDriverGraphNode();
+	if (Driver)
+	{
+		bOverrideEnabled = Driver->Node.RBFParams.SolverType == ERBFSolverType::Additive;
+	}
+
+	return bOverrideEnabled;
+}
+
 float SPDD_TargetRow::GetTargetWeight() const
 {
 	float OutWeight = 0.f;
diff --git a/Engine/Source/Editor/AnimGraph/Private/PoseDriverDetails.h b/Engine/Source/Editor/AnimGraph/Private/PoseDriverDetails.h
--- a/Engine/Source/Editor/AnimGraph/Private/PoseDriverDetails.h
+++ b/Engine/Source/Editor/AnimGraph/Private/PoseDriverDetails.h
@@ -68,9 +68,13 @@ public:
 	void NotifyTargetChanged();
 	/** Get current weight of this target in preview */
 	float GetTargetWeight() const;
+
 	/** Get index of target this represents on pose driver */
 	int32 GetTargetIndex() const;
 
+	/** If we should enable the display of override controls */
+	bool IsOverrideEnabled() const;
+
 	int32 GetTransRotWidgetIndex() const;
 	TOptional<float> GetTranslation(int32 BoneIndex, EAxis::Type Axis) const;
 	TOptional<float> GetRotation(int32 BoneIndex, EAxis::Type Axis) const;
diff --git a/Engine/Source/Editor/AnimationEditor/Private/AnimationEditor.cpp b/Engine/Source/Editor/AnimationEditor/Private/AnimationEditor.cpp
--- a/Engine/Source/Editor/AnimationEditor/Private/AnimationEditor.cpp
+++ b/Engine/Source/Editor/AnimationEditor/Private/AnimationEditor.cpp
@@ -49,6 +49,7 @@
 #include "ISkeletonTreeItem.h"
 #include "Algo/Transform.h"
 #include "ISequenceRecorder.h"
+#include "IAnimSequenceCurveEditor.h"
 
 const FName AnimationEditorAppIdentifier = FName(TEXT("AnimationEditorApp"));
 
@@ -59,10 +60,12 @@ const FName AnimationEditorTabs::SkeletonTreeTab(TEXT("SkeletonTreeView"));
 const FName AnimationEditorTabs::ViewportTab(TEXT("Viewport"));
 const FName AnimationEditorTabs::AdvancedPreviewTab(TEXT("AdvancedPreviewTab"));
 const FName AnimationEditorTabs::DocumentTab(TEXT("Document"));
+const FName AnimationEditorTabs::CurveEditorTab(TEXT("CurveEditor"));
 const FName AnimationEditorTabs::AssetBrowserTab(TEXT("SequenceBrowser"));
 const FName AnimationEditorTabs::AssetDetailsTab(TEXT("AnimAssetPropertiesTab"));
 const FName AnimationEditorTabs::CurveNamesTab(TEXT("AnimCurveViewerTab"));
 const FName AnimationEditorTabs::SlotNamesTab(TEXT("SkeletonSlotNames"));
+const FName AnimationEditorTabs::AnimMontageSectionsTab(TEXT("AnimMontageSections"));
 
 DEFINE_LOG_CATEGORY(LogAnimationEditor);
 
@@ -402,6 +405,14 @@ TSharedPtr<SDockTab> FAnimationEditor::OpenNewAnimationDocumentTab(UAnimationAss
 		GetPersonaToolkit()->GetPreviewScene()->SetPreviewAnimationAsset(InAnimAsset);
 		GetPersonaToolkit()->SetAnimationAsset(InAnimAsset);
 
+		// Close existing opened curve tab
+		if(AnimCurveDocumentTab.IsValid())
+		{
+			AnimCurveDocumentTab.Pin()->RequestCloseTab();
+		}
+
+		AnimCurveDocumentTab.Reset();
+
 		struct Local
 		{
 			static FText GetObjectName(UObject* Object)
@@ -437,6 +448,22 @@ TSharedPtr<SDockTab> FAnimationEditor::OpenNewAnimationDocumentTab(UAnimationAss
 			SharedAnimDocumentTab = OpenedTab;
 		}
 
+		// Invoke the preview tab if this is a montage
+		if(InAnimAsset->IsA<UAnimMontage>())
+		{
+			TabManager->InvokeTab(AnimationEditorTabs::AnimMontageSectionsTab);
+			OnSectionsChanged.Broadcast();
+		}
+		else
+		{
+			// Close existing opened montage sections tab
+			TSharedPtr<SDockTab> OpenMontageSectionsTab = TabManager->FindExistingLiveTab(AnimationEditorTabs::AnimMontageSectionsTab);
+			if(OpenMontageSectionsTab.IsValid())
+			{
+				OpenMontageSectionsTab->RequestCloseTab();
+			}	
+		}
+
 		if (SequenceBrowser.IsValid())
 		{
 			SequenceBrowser.Pin()->SelectAsset(InAnimAsset);
@@ -450,6 +477,55 @@ TSharedPtr<SDockTab> FAnimationEditor::OpenNewAnimationDocumentTab(UAnimationAss
 	return OpenedTab;
 }
 
+void FAnimationEditor::EditCurves(UAnimSequenceBase* InAnimSequence, const TArray<FCurveEditInfo>& InCurveInfo, const TSharedPtr<ITimeSliderController>& InExternalTimeSliderController)
+{
+	FPersonaModule& PersonaModule = FModuleManager::GetModuleChecked<FPersonaModule>("Persona");
+
+	if(!AnimCurveDocumentTab.IsValid())
+	{
+		TSharedRef<IAnimSequenceCurveEditor> NewCurveEditor = PersonaModule.CreateCurveWidgetForAnimDocument(SharedThis(this), GetPersonaToolkit()->GetPreviewScene(), InAnimSequence, InExternalTimeSliderController, TabManager);
+		CurveEditor = NewCurveEditor;
+
+		TSharedPtr<SDockTab> CurveTab = SNew(SDockTab)
+			.Label(LOCTEXT("CurveEditorTabTitle", "Curve Editor"))
+			.TabRole(ETabRole::DocumentTab)
+			.TabColorScale(GetTabColorScale())
+			[
+				NewCurveEditor
+			];
+
+		AnimCurveDocumentTab = CurveTab;
+
+		TabManager->InsertNewDocumentTab(AnimationEditorTabs::CurveEditorTab, FTabManager::ESearchPreference::RequireClosedTab, CurveTab.ToSharedRef());
+	}
+	else
+	{
+		TabManager->DrawAttention(AnimCurveDocumentTab.Pin().ToSharedRef());
+	}
+
+	check(CurveEditor.IsValid());
+
+	CurveEditor.Pin()->ResetCurves();
+
+	for(const FCurveEditInfo& CurveInfo : InCurveInfo)
+	{
+		CurveEditor.Pin()->AddCurve(CurveInfo.CurveDisplayName, CurveInfo.CurveColor, CurveInfo.Name, CurveInfo.Type, CurveInfo.CurveIndex, CurveInfo.OnCurveModified);
+	}
+
+	CurveEditor.Pin()->ZoomToFit();
+}
+
+void FAnimationEditor::StopEditingCurves(const TArray<FCurveEditInfo>& InCurveInfo)
+{
+	if(CurveEditor.IsValid())
+	{
+		for(const FCurveEditInfo& CurveInfo : InCurveInfo)
+		{
+			CurveEditor.Pin()->RemoveCurve(CurveInfo.Name, CurveInfo.Type, CurveInfo.CurveIndex);
+		}
+	}
+}
+
 void FAnimationEditor::HandleSectionsChanged()
 {
 	OnSectionsChanged.Broadcast();
@@ -547,7 +623,7 @@ void FAnimationEditor::OnApplyCompression()
 		TArray<TWeakObjectPtr<UAnimSequence>> AnimSequences;
 		AnimSequences.Add(AnimSequence);
 		FPersonaModule& PersonaModule = FModuleManager::GetModuleChecked<FPersonaModule>("Persona");
-		PersonaModule.ApplyCompression(AnimSequences);
+		PersonaModule.ApplyCompression(AnimSequences, false);
 	}
 }
 
diff --git a/Engine/Source/Editor/AnimationEditor/Private/AnimationEditor.h b/Engine/Source/Editor/AnimationEditor/Private/AnimationEditor.h
--- a/Engine/Source/Editor/AnimationEditor/Private/AnimationEditor.h
+++ b/Engine/Source/Editor/AnimationEditor/Private/AnimationEditor.h
@@ -22,7 +22,10 @@ class ISkeletonTree;
 class UAnimationAsset;
 class USkeletalMeshComponent;
 class UAnimSequence;
+class UAnimSequenceBase;
 class ISkeletonTreeItem;
+class IAnimSequenceCurveEditor;
+struct FRichCurve;
 
 namespace AnimationEditorModes
 {
@@ -38,10 +41,12 @@ namespace AnimationEditorTabs
 	extern const FName ViewportTab;
 	extern const FName AdvancedPreviewTab;
 	extern const FName DocumentTab;
+	extern const FName CurveEditorTab;
 	extern const FName AssetBrowserTab;
 	extern const FName AssetDetailsTab;
 	extern const FName CurveNamesTab;
 	extern const FName SlotNamesTab;
+	extern const FName AnimMontageSectionsTab;
 }
 
 class FAnimationEditor : public IAnimationEditor, public FGCObject, public FEditorUndoClient, public FTickableEditorObject
@@ -57,6 +62,8 @@ public:
 	/** IAnimationEditor interface */
 	virtual void SetAnimationAsset(UAnimationAsset* AnimAsset) override;
 	virtual IAnimationSequenceBrowser* GetAssetBrowser() const override;
+	virtual void EditCurves(UAnimSequenceBase* InAnimSequence, const TArray<FCurveEditInfo>& InCurveInfo, const TSharedPtr<ITimeSliderController>& InExternalTimeSliderController) override;
+	virtual void StopEditingCurves(const TArray<FCurveEditInfo>& InCurveInfo) override;
 
 	/** IHasPersonaToolkit interface */
 	virtual TSharedRef<class IPersonaToolkit> GetPersonaToolkit() const override { return PersonaToolkit.ToSharedRef(); }
@@ -196,6 +203,12 @@ private:
 	/** The animation document currently being edited */
 	TWeakPtr<SDockTab> SharedAnimDocumentTab;
 
+	/** The animation document's curves that are currently being edited */
+	TWeakPtr<SDockTab> AnimCurveDocumentTab;
+
 	/** Sequence Browser **/
 	TWeakPtr<class IAnimationSequenceBrowser> SequenceBrowser;
+
+	/** The anim sequence curve editor */
+	TWeakPtr<IAnimSequenceCurveEditor> CurveEditor;
 };
diff --git a/Engine/Source/Editor/AnimationEditor/Private/AnimationEditorMode.cpp b/Engine/Source/Editor/AnimationEditor/Private/AnimationEditorMode.cpp
--- a/Engine/Source/Editor/AnimationEditor/Private/AnimationEditorMode.cpp
+++ b/Engine/Source/Editor/AnimationEditor/Private/AnimationEditorMode.cpp
@@ -35,8 +35,9 @@ FAnimationEditorMode::FAnimationEditorMode(TSharedRef<FWorkflowCentricApplicatio
 	TabFactories.RegisterFactory(PersonaModule.CreateCurveViewerTabFactory(InHostingApp, InSkeletonTree->GetEditableSkeleton(), AnimationEditor->GetPersonaToolkit()->GetPreviewScene(), AnimationEditor->OnPostUndo, OnObjectsSelected));
 	TabFactories.RegisterFactory(PersonaModule.CreateSkeletonSlotNamesTabFactory(InHostingApp, InSkeletonTree->GetEditableSkeleton(), AnimationEditor->OnPostUndo, FOnObjectSelected::CreateSP(&AnimationEditor.Get(), &FAnimationEditor::HandleObjectSelected)));
 	TabFactories.RegisterFactory(PersonaModule.CreateAnimNotifiesTabFactory(InHostingApp, InSkeletonTree->GetEditableSkeleton(), OnObjectsSelected));
+	TabFactories.RegisterFactory(PersonaModule.CreateAnimMontageSectionsTabFactory(InHostingApp, AnimationEditor->GetPersonaToolkit(), AnimationEditor->OnSectionsChanged));
 
-	TabLayout = FTabManager::NewLayout("Standalone_AnimationEditor_Layout_v1.2")
+	TabLayout = FTabManager::NewLayout("Standalone_AnimationEditor_Layout_v1.3")
 		->AddArea
 		(
 			FTabManager::NewPrimaryArea()
@@ -84,6 +85,8 @@ FAnimationEditorMode::FAnimationEditorMode(TSharedRef<FWorkflowCentricApplicatio
 						->SetSizeCoefficient(0.4f)
 						->SetHideTabWell(true)
 						->AddTab(AnimationEditorTabs::DocumentTab, ETabState::ClosedTab)
+						->AddTab(AnimationEditorTabs::CurveEditorTab, ETabState::ClosedTab)
+						->SetForegroundTab(AnimationEditorTabs::DocumentTab)
 					)
 				)
 				->Split
@@ -106,6 +109,7 @@ FAnimationEditorMode::FAnimationEditorMode(TSharedRef<FWorkflowCentricApplicatio
 						->SetSizeCoefficient(0.4f)
 						->SetHideTabWell(false)
 						->AddTab(AnimationEditorTabs::AssetBrowserTab, ETabState::OpenedTab)
+						->AddTab(AnimationEditorTabs::AnimMontageSectionsTab, ETabState::ClosedTab)
 						->AddTab(AnimationEditorTabs::CurveNamesTab, ETabState::ClosedTab)
 						->AddTab(AnimationEditorTabs::SlotNamesTab, ETabState::ClosedTab)
 					)
diff --git a/Engine/Source/Editor/AnimationEditor/Public/IAnimationEditor.h b/Engine/Source/Editor/AnimationEditor/Public/IAnimationEditor.h
--- a/Engine/Source/Editor/AnimationEditor/Public/IAnimationEditor.h
+++ b/Engine/Source/Editor/AnimationEditor/Public/IAnimationEditor.h
@@ -5,9 +5,14 @@
 #include "CoreMinimal.h"
 #include "WorkflowOrientedApp/WorkflowCentricApplication.h"
 #include "IHasPersonaToolkit.h"
+#include "Animation/SmartName.h"
 
 class UAnimationAsset;
 class IAnimationSequenceBrowser;
+class UAnimSequenceBase;
+struct FRichCurve;
+class ITimeSliderController;
+enum class ERawCurveTrackTypes : uint8;
 
 class IAnimationEditor : public FWorkflowCentricApplication, public IHasPersonaToolkit
 {
@@ -17,4 +22,41 @@ public:
 
 	/** Get the asset browser we host */
 	virtual IAnimationSequenceBrowser* GetAssetBrowser() const = 0;
+
+	/** Support structure for EditCurves */
+	struct FCurveEditInfo
+	{
+		FCurveEditInfo(const FText& InCurveDisplayName, const FLinearColor& InCurveColor, const FSmartName& InName, ERawCurveTrackTypes InType, int32 InCurveIndex, FSimpleDelegate OnCurveModified = FSimpleDelegate())
+			: CurveDisplayName(InCurveDisplayName)
+			, CurveColor(InCurveColor)
+			, Name(InName)
+			, Type(InType)
+			, CurveIndex(InCurveIndex)
+			, OnCurveModified(OnCurveModified)
+		{}
+
+		FCurveEditInfo(const FSmartName& InName, ERawCurveTrackTypes InType, int32 InCurveIndex)
+			: Name(InName)
+			, Type(InType)
+			, CurveIndex(InCurveIndex)
+		{}
+
+		bool operator==(const FCurveEditInfo& InCurveEditInfo) const
+		{
+			return Name.UID == InCurveEditInfo.Name.UID && Type == InCurveEditInfo.Type && CurveIndex == InCurveEditInfo.CurveIndex;
+		}
+
+		FText CurveDisplayName;
+		FLinearColor CurveColor;
+		FSmartName Name;
+		ERawCurveTrackTypes Type;
+		int32 CurveIndex;
+		FSimpleDelegate OnCurveModified;
+	};
+
+	/** Edit the specified curves on the specified sequence */
+	virtual void EditCurves(UAnimSequenceBase* InAnimSequence, const TArray<FCurveEditInfo>& InCurveInfo, const TSharedPtr<ITimeSliderController>& InExternalTimeSliderController) = 0;
+
+	/** Stop editing the specified curves */
+	virtual void StopEditingCurves(const TArray<FCurveEditInfo>& InCurveInfo) = 0;
 };
diff --git a/Engine/Source/Editor/AnimationModifiers/AnimationModifiers.Build.cs b/Engine/Source/Editor/AnimationModifiers/AnimationModifiers.Build.cs
--- a/Engine/Source/Editor/AnimationModifiers/AnimationModifiers.Build.cs
+++ b/Engine/Source/Editor/AnimationModifiers/AnimationModifiers.Build.cs
@@ -22,7 +22,6 @@ public class AnimationModifiers : ModuleRules
                 "Kismet",
                 "AssetTools",
                 "ClassViewer",
-                "Persona",
                 "AnimationEditor",
                 "SkeletonEditor",
                 "AssetRegistry"                
diff --git a/Engine/Source/Editor/AnimationModifiers/Private/AnimationBlueprintLibrary.cpp b/Engine/Source/Editor/AnimationModifiers/Private/AnimationBlueprintLibrary.cpp
--- a/Engine/Source/Editor/AnimationModifiers/Private/AnimationBlueprintLibrary.cpp
+++ b/Engine/Source/Editor/AnimationModifiers/Private/AnimationBlueprintLibrary.cpp
@@ -4,8 +4,10 @@
 
 #include "Animation/AnimSequence.h"
 #include "Animation/AnimationAsset.h"
+#include "Animation/AnimBoneCompressionSettings.h"
+#include "Animation/AnimCurveCompressionSettings.h"
 #include "Animation/AnimMetaData.h"
-#include "Animation/AnimNotifies/AnimNotifyState.h" 
+#include "Animation/AnimNotifies/AnimNotifyState.h"
 #include "Animation/Skeleton.h"
 #include "Animation/AnimNotifies/AnimNotify.h"
 #include "BonePose.h" 
@@ -132,27 +134,60 @@ const FRawAnimSequenceTrack&  UAnimationBlueprintLibrary::GetRawAnimationTrackBy
 	return AnimationSequence->GetRawAnimationTrack(TrackIndex);	
 }
 
-void UAnimationBlueprintLibrary::GetCompressionScheme(const UAnimSequence* AnimationSequence, UAnimCompress*& CompressionScheme)
+void UAnimationBlueprintLibrary::GetBoneCompressionSettings(const UAnimSequence* AnimationSequence, UAnimBoneCompressionSettings*& CompressionSettings)
 {
-	if (AnimationSequence)
+	if (AnimationSequence == nullptr)
 	{
-		CompressionScheme = AnimationSequence->CompressionScheme;
+		UE_LOG(LogAnimationBlueprintLibrary, Warning, TEXT("Invalid Animation Sequence supplied for GetBoneCompressionSettings"));
+		return;
 	}
-	else
+
+	CompressionSettings = AnimationSequence->BoneCompressionSettings;
+}
+
+void UAnimationBlueprintLibrary::SetBoneCompressionSettings(UAnimSequence* AnimationSequence, UAnimBoneCompressionSettings* CompressionSettings)
+{
+	if (AnimationSequence == nullptr)
 	{
-		UE_LOG(LogAnimationBlueprintLibrary, Warning, TEXT("Invalid Animation Sequence supplied for GetCompressionScheme"));
+		UE_LOG(LogAnimationBlueprintLibrary, Warning, TEXT("Invalid Animation Sequence supplied for SetBoneCompressionSettings"));
+		return;
 	}
+
+	if (CompressionSettings == nullptr || !CompressionSettings->AreSettingsValid())
+	{
+		UE_LOG(LogAnimationBlueprintLibrary, Warning, TEXT("Invalid Bone Compression Settings supplied for SetBoneCompressionSettings"));
+		return;
+	}
+
+	AnimationSequence->BoneCompressionSettings = CompressionSettings;
 }
 
-void UAnimationBlueprintLibrary::SetCompressionScheme(UAnimSequence* AnimationSequence, UAnimCompress* CompressionScheme)
+void UAnimationBlueprintLibrary::GetCurveCompressionSettings(const UAnimSequence* AnimationSequence, UAnimCurveCompressionSettings*& CompressionSettings)
 {
-	if (AnimationSequence)
+	if (AnimationSequence == nullptr)
 	{
-		AnimationSequence->CompressionScheme = CompressionScheme;
+		UE_LOG(LogAnimationBlueprintLibrary, Warning, TEXT("Invalid Animation Sequence supplied for GetCurveCompressionSettings"));
+		return;
 	}
+
+	CompressionSettings = AnimationSequence->CurveCompressionSettings;
+}
+
+void UAnimationBlueprintLibrary::SetCurveCompressionSettings(UAnimSequence* AnimationSequence, UAnimCurveCompressionSettings* CompressionSettings)
+{
+	if (AnimationSequence == nullptr)
 	{
-		UE_LOG(LogAnimationBlueprintLibrary, Warning, TEXT("Invalid Animation Sequence supplied for SetCompressionScheme"));
+		UE_LOG(LogAnimationBlueprintLibrary, Warning, TEXT("Invalid Animation Sequence supplied for SetCurveCompressionSettings"));
+		return;
 	}
+
+	if (CompressionSettings == nullptr || !CompressionSettings->AreSettingsValid())
+	{
+		UE_LOG(LogAnimationBlueprintLibrary, Warning, TEXT("Invalid Bone Compression Settings supplied for SetCurveCompressionSettings"));
+		return;
+	}
+
+	AnimationSequence->CurveCompressionSettings = CompressionSettings;
 }
 
 void UAnimationBlueprintLibrary::GetAdditiveAnimationType(const UAnimSequence* AnimationSequence, TEnumAsByte<enum EAdditiveAnimationType>& AdditiveAnimationType)
