Incremental Static Analysis Win64 (MSVC) - Warning: Dereferencing NULL pointer Property in rigvmcontroller.cpp

diff --git a/Engine/Source/Developer/RigVMDeveloper/Private/RigVMModel/RigVMController.cpp b/Engine/Source/Developer/RigVMDeveloper/Private/RigVMModel/RigVMController.cpp
--- a/Engine/Source/Developer/RigVMDeveloper/Private/RigVMModel/RigVMController.cpp
+++ b/Engine/Source/Developer/RigVMDeveloper/Private/RigVMModel/RigVMController.cpp
@@ -2186,10 +2186,7 @@ bool URigVMController::ResetPinDefaultValue(URigVMPin* InPin, bool bUndo)
 
 		UStruct* Struct = StructNode->ScriptStruct;
 		FProperty* Property = Struct->FindPropertyByName(*Parts[PartIndex++]);
-		if (!ensure(Property != nullptr))
-		{
-			return false;
-		}
+		check(Property);
 
 		const uint8* Memory = StructOnScope->GetStructMemory();
 		Memory = Property->ContainerPtrToValuePtr<uint8>(Memory);
@@ -2199,6 +2196,7 @@ bool URigVMController::ResetPinDefaultValue(URigVMPin* InPin, bool bUndo)
 			if (FArrayProperty* ArrayProperty = CastField<FArrayProperty>(Property))
 			{
 				Property = ArrayProperty->Inner;
+				check(Property);
 				int32 ArrayIndex = FCString::Atoi(*Parts[PartIndex++]);
 				TArray<uint8>* ArrayPtr = (TArray<uint8>*)Memory;
 				Memory = ArrayPtr->GetData() + Property->GetSize() * ArrayIndex;
@@ -2209,6 +2207,7 @@ bool URigVMController::ResetPinDefaultValue(URigVMPin* InPin, bool bUndo)
 			{
 				Struct = StructProperty->Struct;
 				Property = Struct->FindPropertyByName(*Parts[PartIndex++]);
+				check(Property);
 				Memory = Property->ContainerPtrToValuePtr<uint8>(Memory);
 				continue;
 			}
@@ -2219,6 +2218,7 @@ bool URigVMController::ResetPinDefaultValue(URigVMPin* InPin, bool bUndo)
 		if (Memory)
 		{
 			FString DefaultValue;
+			check(Property);
 			Property->ExportTextItem(DefaultValue, Memory, nullptr, nullptr, PPF_None);
 
 			if (!DefaultValue.IsEmpty())
