Instead of accessing sub array elements on a default struct we create a new structonscope for the inner property and use that instead. The array property might be empty - so relying on that isn't stable and can cause crashes as hit in UE-88394

Merging

//UE4/Dev-Anim/Engine/Source/Developer/RigVMDeveloper/Private/RigVMModel/RigVMController.cpp

to //UE4/Release-4.25/Engine/Source/Developer/RigVMDeveloper/Private/RigVMModel/RigVMController.cpp

diff --git a/Engine/Source/Developer/RigVMDeveloper/Private/RigVMModel/RigVMController.cpp b/Engine/Source/Developer/RigVMDeveloper/Private/RigVMModel/RigVMController.cpp
--- a/Engine/Source/Developer/RigVMDeveloper/Private/RigVMModel/RigVMController.cpp
+++ b/Engine/Source/Developer/RigVMDeveloper/Private/RigVMModel/RigVMController.cpp
@@ -2234,7 +2234,7 @@ bool URigVMController::ResetPinDefaultValue(URigVMPin* InPin, bool bUndo)
 		FProperty* Property = Struct->FindPropertyByName(*Parts[PartIndex++]);
 		check(Property);
 
-		const uint8* Memory = StructOnScope->GetStructMemory();
+		uint8* Memory = StructOnScope->GetStructMemory();
 		Memory = Property->ContainerPtrToValuePtr<uint8>(Memory);
 
 		while (PartIndex < Parts.Num() && Property != nullptr)
@@ -2243,9 +2243,15 @@ bool URigVMController::ResetPinDefaultValue(URigVMPin* InPin, bool bUndo)
 			{
 				Property = ArrayProperty->Inner;
 				check(Property);
-				int32 ArrayIndex = FCString::Atoi(*Parts[PartIndex++]);
-				TArray<uint8>* ArrayPtr = (TArray<uint8>*)Memory;
-				Memory = ArrayPtr->GetData() + Property->GetSize() * ArrayIndex;
+				PartIndex++;
+
+				if (FStructProperty* StructProperty = CastField<FStructProperty>(Property))
+				{
+					UScriptStruct* InnerStruct = StructProperty->Struct;
+					StructOnScope = MakeShareable(new FStructOnScope(InnerStruct));
+					Memory = (uint8 *)StructOnScope->GetStructMemory();
+					InnerStruct->InitializeDefaultValue(Memory);
+				}
 				continue;
 			}
 
