RigVM AST: Moved the lookup from AST Node to Expression Index to a map, since now with nested nodes the Node index itself cannot be relied on.

Merging

//UE4/Dev-Anim/Engine/Source/Developer/RigVMDeveloper/...

to //UE4/Dev-Anim/Engine/Source/Developer/RigVMDeveloper/...

diff --git a/Engine/Source/Developer/RigVMDeveloper/Private/RigVMCompiler/RigVMAST.cpp b/Engine/Source/Developer/RigVMDeveloper/Private/RigVMCompiler/RigVMAST.cpp
--- a/Engine/Source/Developer/RigVMDeveloper/Private/RigVMCompiler/RigVMAST.cpp
+++ b/Engine/Source/Developer/RigVMDeveloper/Private/RigVMCompiler/RigVMAST.cpp
@@ -481,13 +481,6 @@ FRigVMParserAST::FRigVMParserAST(URigVMGraph* InGraph, const FRigVMParserASTSett
 	LastCycleCheckExpr = nullptr;
 
 	const TArray<URigVMNode*> Nodes = InGraph->GetNodes();
-
-	NodeExpressionIndex.SetNumUninitialized(Nodes.Num());
-	for (int32& Index : NodeExpressionIndex)
-	{
-		Index = INDEX_NONE;
-	}
-
 	for (URigVMNode* Node : Nodes)
 	{
 		if(Node->IsEvent())
@@ -502,11 +495,14 @@ FRigVMParserAST::FRigVMParserAST(URigVMGraph* InGraph, const FRigVMParserASTSett
 	for (int32 PassIndex = 0; PassIndex < 2; PassIndex++)
 	{
 		const bool bTraverseMutable = PassIndex == 0;
-		for (int32 NodeIndex = 0; NodeIndex < NodeExpressionIndex.Num(); NodeIndex++)
+		for (int32 NodeIndex = 0; NodeIndex < Nodes.Num(); NodeIndex++)
 		{
-			if (NodeExpressionIndex[NodeIndex] != INDEX_NONE)
+			if (const int32* ExprIndex = NodeExpressionIndex.Find(Nodes[NodeIndex]))
 			{
-				continue;
+				if (*ExprIndex != INDEX_NONE)
+				{
+					continue;
+				}
 			}
 
 			if (Nodes[NodeIndex]->IsMutable() == bTraverseMutable)
@@ -598,7 +594,7 @@ FRigVMExprAST* FRigVMParserAST::TraverseMutableNode(URigVMNode* InNode, FRigVMEx
 		InParentExpr = NodeExpr;
 	}
 	SubjectToExpression.Add(InNode, NodeExpr);
-	NodeExpressionIndex[InNode->GetNodeIndex()] = NodeExpr->GetIndex();
+	NodeExpressionIndex.Add(InNode, NodeExpr->GetIndex());
 
 	TraversePins(InNode, NodeExpr);
 
@@ -648,7 +644,7 @@ FRigVMExprAST* FRigVMParserAST::TraverseNode(URigVMNode* InNode, FRigVMExprAST*
 	NodeExpr->Name = InNode->GetFName();
 	NodeExpr->AddParent(InParentExpr);
 	SubjectToExpression.Add(InNode, NodeExpr);
-	NodeExpressionIndex[InNode->GetNodeIndex()] = NodeExpr->GetIndex();
+	NodeExpressionIndex.Add(InNode, NodeExpr->GetIndex());
 
 	TraversePins(InNode, NodeExpr);
 
diff --git a/Engine/Source/Developer/RigVMDeveloper/Public/RigVMCompiler/RigVMAST.h b/Engine/Source/Developer/RigVMDeveloper/Public/RigVMCompiler/RigVMAST.h
--- a/Engine/Source/Developer/RigVMDeveloper/Public/RigVMCompiler/RigVMAST.h
+++ b/Engine/Source/Developer/RigVMDeveloper/Public/RigVMCompiler/RigVMAST.h
@@ -974,7 +974,7 @@ private:
 	FRigVMExprAST* TraverseLink(URigVMLink* InLink, FRigVMExprAST* InParentExpr);
 
 	TMap<UObject*, FRigVMExprAST*> SubjectToExpression;
-	TArray<int32> NodeExpressionIndex;
+	TMap<UObject*, int32> NodeExpressionIndex;
 	TArray<FRigVMExprAST*> Expressions;
 	TArray<FRigVMExprAST*> RootExpressions;
 
