Control Rig: Ensure to define missing parameters AFTER initial traverse of AST

diff --git a/Engine/Source/Developer/RigVMDeveloper/Private/RigVMCompiler/RigVMCompiler.cpp b/Engine/Source/Developer/RigVMDeveloper/Private/RigVMCompiler/RigVMCompiler.cpp
--- a/Engine/Source/Developer/RigVMDeveloper/Private/RigVMCompiler/RigVMCompiler.cpp
+++ b/Engine/Source/Developer/RigVMDeveloper/Private/RigVMCompiler/RigVMCompiler.cpp
@@ -126,18 +126,14 @@ bool URigVMCompiler::Compile(URigVMGraph* InGraph, URigVM* OutVM, const FRigVMUs
 	WorkData.NumInstructions = 0;
 
 	// define all parameters independent from sorted nodes
-	for (URigVMNode* Node : InGraph->GetNodes())
+	for (const FRigVMExprAST* Expr : AST->Expressions)
 	{
-		if (URigVMParameterNode* ParameterNode = Cast<URigVMParameterNode>(Node))
+		if (Expr->IsA(FRigVMExprAST::EType::Var))
 		{
-			const FRigVMExprAST* ParameterExpr = AST->GetExprForSubject(ParameterNode);
-			check(ParameterExpr);
-
-			for (const FRigVMExprAST* ChildExpr : *ParameterExpr)
+			const FRigVMVarExprAST* VarExpr = Expr->To<FRigVMVarExprAST>();
 			{
-				if (ChildExpr->IsA(FRigVMExprAST::EType::Var))
+				if(Cast<URigVMParameterNode>(VarExpr->GetPin()->GetNode()))
 				{
-					const FRigVMVarExprAST* VarExpr = ChildExpr->To<FRigVMVarExprAST>();
 					if (VarExpr->GetPin()->GetName() == TEXT("Value"))
 					{
 						FindOrAddRegister(VarExpr, WorkData, false /* watchvalue */);
@@ -147,6 +143,12 @@ bool URigVMCompiler::Compile(URigVMGraph* InGraph, URigVM* OutVM, const FRigVMUs
 		}
 	}
 
+	WorkData.ExprComplete.Reset();
+	for (FRigVMExprAST* RootExpr : *AST)
+	{
+		TraverseExpression(RootExpr, WorkData);
+	}
+
 	// If a parameter node has no corresponding AST node, because it's not on the execution path,
 	// then add a dummy parameter with no register index, so that we can properly add the pins on the 
 	// ControlRig node later and not lose any existing external connections.
@@ -171,12 +173,6 @@ bool URigVMCompiler::Compile(URigVMGraph* InGraph, URigVM* OutVM, const FRigVMUs
 		}
 	}
 
-	WorkData.ExprComplete.Reset();
-	for (FRigVMExprAST* RootExpr : *AST)
-	{
-		TraverseExpression(RootExpr, WorkData);
-	}
-
 	WorkData.bSetupMemory = false;
 	WorkData.ExprComplete.Reset();
 	for (FRigVMExprAST* RootExpr : *AST)
@@ -1030,7 +1026,7 @@ FRigVMOperand URigVMCompiler::FindOrAddRegister(const FRigVMVarExprAST* InVarExp
 			FName Name = ParameterNode->GetParameterName();
 			ERigVMParameterType ParameterType = ParameterNode->IsInput() ? ERigVMParameterType::Input : ERigVMParameterType::Output;
 			FRigVMParameter Parameter(ParameterType, Name, Operand.GetRegisterIndex(), Pin->GetCPPType(), ScriptStruct);
-			WorkData.VM->ParametersNameMap.Add(Parameter.Name, WorkData.VM->Parameters.Add(Parameter));
+			WorkData.VM->ParametersNameMap.FindOrAdd(Parameter.Name) = WorkData.VM->Parameters.Add(Parameter);
 		}
 	}
 	ensure(Operand.IsValid());
diff --git a/Engine/Source/Developer/RigVMDeveloper/Public/RigVMCompiler/RigVMAST.h b/Engine/Source/Developer/RigVMDeveloper/Public/RigVMCompiler/RigVMAST.h
--- a/Engine/Source/Developer/RigVMDeveloper/Public/RigVMCompiler/RigVMAST.h
+++ b/Engine/Source/Developer/RigVMDeveloper/Public/RigVMCompiler/RigVMAST.h
@@ -990,5 +990,6 @@ private:
 	TArray<ETraverseRelationShip> CycleCheckFlags;
 
 	friend class FRigVMExprAST;
+	friend class URigVMCompiler;
 };
 
