Control Rig: Rapidly undoing "Add visual debug" causes ensure

diff --git a/Engine/Source/Developer/RigVMDeveloper/Private/RigVMCompiler/RigVMAST.cpp b/Engine/Source/Developer/RigVMDeveloper/Private/RigVMCompiler/RigVMAST.cpp
--- a/Engine/Source/Developer/RigVMDeveloper/Private/RigVMCompiler/RigVMAST.cpp
+++ b/Engine/Source/Developer/RigVMDeveloper/Private/RigVMCompiler/RigVMAST.cpp
@@ -842,13 +842,13 @@ void FRigVMParserAST::FoldEntries()
 					{
 						if (ChildExpr->To<FRigVMVarExprAST>()->IsExecuteContext())
 						{
-							ExpressionsToRemove.Add(ChildExpr);
+							ExpressionsToRemove.AddUnique(ChildExpr);
 							continue;
 						}
 					}
 					ChildExpr->AddParent(FoldEntry);
 				}
-				ExpressionsToRemove.Add(Entry);
+				ExpressionsToRemove.AddUnique(Entry);
 			}
 			else
 			{
@@ -996,10 +996,10 @@ void FRigVMParserAST::FoldAssignments()
 			}
 		}
 
-		ExpressionsToRemove.Add(AssignExpr);
+		ExpressionsToRemove.AddUnique(AssignExpr);
 		if (Parent->Parents.Num() == 0)
 		{
-			ExpressionsToRemove.Add(Parent);
+			ExpressionsToRemove.AddUnique(Parent);
 		}
 	}
 
@@ -1036,7 +1036,7 @@ void FRigVMParserAST::FoldLiterals()
 				{
 					Parent->ReplaceChild(Expression, *MappedExpr);
 				}
-				ExpressionsToRemove.Add(Expression);
+				ExpressionsToRemove.AddUnique(Expression);
 			}
 			else
 			{
diff --git a/Engine/Source/Developer/RigVMDeveloper/Private/RigVMCompiler/RigVMCompiler.cpp b/Engine/Source/Developer/RigVMDeveloper/Private/RigVMCompiler/RigVMCompiler.cpp
--- a/Engine/Source/Developer/RigVMDeveloper/Private/RigVMCompiler/RigVMCompiler.cpp
+++ b/Engine/Source/Developer/RigVMDeveloper/Private/RigVMCompiler/RigVMCompiler.cpp
@@ -75,8 +75,7 @@ int32 FRigVMCompilerWorkData::DecRefRegister(int32 InRegister, int32 InDecrement
 	if (int32* RefCountPtr = RegisterRefCount.Find(InRegister))
 	{
 		int32& RefCount = *RefCountPtr;
-		RefCount -= InDecrement;
-		ensure(RefCount >= 0);
+		RefCount = FMath::Max<int32>(0, RefCount - InDecrement);
 		return RefCount;
 	}
 	return 0;
diff --git a/Engine/Source/Developer/RigVMDeveloper/Private/RigVMModel/RigVMController.cpp b/Engine/Source/Developer/RigVMDeveloper/Private/RigVMModel/RigVMController.cpp
--- a/Engine/Source/Developer/RigVMDeveloper/Private/RigVMModel/RigVMController.cpp
+++ b/Engine/Source/Developer/RigVMDeveloper/Private/RigVMModel/RigVMController.cpp
@@ -3891,7 +3891,7 @@ void URigVMController::RewireLinks(URigVMPin* InOldPin, URigVMPin* InNewPin, boo
 			check(NewPin);
 
 			BreakLink(Link->GetSourcePin(), Link->GetTargetPin(), false);
-			ensure(AddLink(Link->GetSourcePin(), NewPin, false));
+			AddLink(Link->GetSourcePin(), NewPin, false);
 		}
 	}
 	else
@@ -3909,7 +3909,7 @@ void URigVMController::RewireLinks(URigVMPin* InOldPin, URigVMPin* InNewPin, boo
 			check(NewPin);
 
 			BreakLink(Link->GetSourcePin(), Link->GetTargetPin(), false);
-			ensure(AddLink(NewPin, Link->GetTargetPin(), false));
+			AddLink(NewPin, Link->GetTargetPin(), false);
 		}
 	}
 }
