Control Rig: Ensure to clean up injected nodes during undo / redo

The injected node action was just calling RemoveNode, which outside of undo / redo does not clean up links. We had some stale / transient nodes due to this in the repro case. By adding an additional argument to force recursion into BreakAllLinks etc this is fixed.

diff --git a/Engine/Source/Developer/RigVMDeveloper/Private/RigVMModel/RigVMController.cpp b/Engine/Source/Developer/RigVMDeveloper/Private/RigVMModel/RigVMController.cpp
--- a/Engine/Source/Developer/RigVMDeveloper/Private/RigVMModel/RigVMController.cpp
+++ b/Engine/Source/Developer/RigVMDeveloper/Private/RigVMModel/RigVMController.cpp
@@ -1334,7 +1334,7 @@ FName URigVMController::GetUniqueName(const FName& InName, TFunction<bool(const
 
 #endif
 
-bool URigVMController::RemoveNode(URigVMNode* InNode, bool bUndo)
+bool URigVMController::RemoveNode(URigVMNode* InNode, bool bUndo, bool bRecursive)
 {
 	if (!IsValidNodeForGraph(InNode))
 	{
@@ -1370,7 +1370,7 @@ bool URigVMController::RemoveNode(URigVMNode* InNode, bool bUndo)
 		}
 	}
 
-	if (bUndo)
+	if (bUndo || bRecursive)
 	{
 		SelectNode(InNode, false, bUndo);
 
@@ -1425,13 +1425,13 @@ bool URigVMController::RemoveNode(URigVMNode* InNode, bool bUndo)
 	return true;
 }
 
-bool URigVMController::RemoveNodeByName(const FName& InNodeName, bool bUndo)
+bool URigVMController::RemoveNodeByName(const FName& InNodeName, bool bUndo, bool bRecursive)
 {
 	if (!IsValidGraph())
 	{
 		return false;
 	}
-	return RemoveNode(Graph->FindNodeByName(InNodeName), bUndo);
+	return RemoveNode(Graph->FindNodeByName(InNodeName), bUndo, bRecursive);
 }
 
 bool URigVMController::SelectNode(URigVMNode* InNode, bool bSelect, bool bUndo)
diff --git a/Engine/Source/Developer/RigVMDeveloper/Private/RigVMModel/RigVMControllerActions.cpp b/Engine/Source/Developer/RigVMDeveloper/Private/RigVMModel/RigVMControllerActions.cpp
--- a/Engine/Source/Developer/RigVMDeveloper/Private/RigVMModel/RigVMControllerActions.cpp
+++ b/Engine/Source/Developer/RigVMDeveloper/Private/RigVMModel/RigVMControllerActions.cpp
@@ -443,7 +443,7 @@ bool FRigVMAddInjectedNodeAction::Undo(URigVMController* InController)
 	{
 		return false;
 	}
-	return InController->RemoveNodeByName(*NodePath, false);
+	return InController->RemoveNodeByName(*NodePath, false, true);
 }
 
 bool FRigVMAddInjectedNodeAction::Redo(URigVMController* InController)
diff --git a/Engine/Source/Developer/RigVMDeveloper/Public/RigVMModel/RigVMController.h b/Engine/Source/Developer/RigVMDeveloper/Public/RigVMModel/RigVMController.h
--- a/Engine/Source/Developer/RigVMDeveloper/Public/RigVMModel/RigVMController.h
+++ b/Engine/Source/Developer/RigVMDeveloper/Public/RigVMModel/RigVMController.h
@@ -212,12 +212,12 @@ public:
 	// Removes a node from the graph
 	// This causes a NodeRemoved modified event.
 	UFUNCTION(BlueprintCallable, Category = RigVMController)
-	bool RemoveNode(URigVMNode* InNode, bool bUndo = true);
+	bool RemoveNode(URigVMNode* InNode, bool bUndo = true, bool bRecursive = false);
 
 	// Removes a node from the graph given the node's name.
 	// This causes a NodeRemoved modified event.
 	UFUNCTION(BlueprintCallable, Category = RigVMController)
-	bool RemoveNodeByName(const FName& InNodeName, bool bUndo = true);
+	bool RemoveNodeByName(const FName& InNodeName, bool bUndo = true, bool bRecursive = false);
 
 	// Selects a single node in the graph.
 	// This causes a NodeSelected / NodeDeselected modified event.
