Control Rig: Feature to bind pins to external variables, model + controller work (view todo)

diff --git a/Engine/Plugins/Experimental/ControlRig/Source/ControlRigDeveloper/Private/ControlRigBlueprint.cpp b/Engine/Plugins/Experimental/ControlRig/Source/ControlRigDeveloper/Private/ControlRigBlueprint.cpp
--- a/Engine/Plugins/Experimental/ControlRig/Source/ControlRigDeveloper/Private/ControlRigBlueprint.cpp
+++ b/Engine/Plugins/Experimental/ControlRig/Source/ControlRigDeveloper/Private/ControlRigBlueprint.cpp
@@ -1898,7 +1898,7 @@ void UControlRigBlueprint::OnVariableRemoved(const FName& InVarName)
 {
 	if (Controller)
 	{
-		Controller->RemoveVariableNodes(InVarName, true);
+		Controller->OnExternalVariableRemoved(InVarName, true);
 	}
 	BroadcastExternalVariablesChangedEvent();
 }
@@ -1907,7 +1907,7 @@ void UControlRigBlueprint::OnVariableRenamed(const FName& InOldVarName, const FN
 {
 	if (Controller)
 	{
-		Controller->RenameVariableNodes(InOldVarName, InNewVarName, true);
+		Controller->OnExternalVariableRenamed(InOldVarName, InNewVarName, true);
 	}
 	BroadcastExternalVariablesChangedEvent();
 }
@@ -1919,11 +1919,11 @@ void UControlRigBlueprint::OnVariableTypeChanged(const FName& InVarName, FEdGrap
 		FRigVMExternalVariable NewVariable = UControlRig::GetExternalVariableFromPinType(InVarName, InNewPinType);
 		if (NewVariable.IsValid(true)) // allow nullptr
 		{
-			Controller->ChangeVariableNodesType(InVarName, NewVariable.TypeName.ToString(), NewVariable.TypeObject, true);
+			Controller->OnExternalVariableTypeChanged(InVarName, NewVariable.TypeName.ToString(), NewVariable.TypeObject, true);
 		}
 		else
 		{
-			Controller->RemoveVariableNodes(InVarName, true);
+			Controller->OnExternalVariableRemoved(InVarName, true);
 		}
 	}
 	BroadcastExternalVariablesChangedEvent();
diff --git a/Engine/Plugins/Experimental/ControlRig/Source/ControlRigDeveloper/Private/Graph/ControlRigGraph.cpp b/Engine/Plugins/Experimental/ControlRig/Source/ControlRigDeveloper/Private/Graph/ControlRigGraph.cpp
--- a/Engine/Plugins/Experimental/ControlRig/Source/ControlRigDeveloper/Private/Graph/ControlRigGraph.cpp
+++ b/Engine/Plugins/Experimental/ControlRig/Source/ControlRigDeveloper/Private/Graph/ControlRigGraph.cpp
@@ -496,6 +496,7 @@ void UControlRigGraph::HandleModifiedEvent(ERigVMGraphNotifType InNotifType, URi
 		case ERigVMGraphNotifType::PinArraySizeChanged:
 		case ERigVMGraphNotifType::PinDirectionChanged:
 		case ERigVMGraphNotifType::PinTypeChanged:
+		case ERigVMGraphNotifType::PinBoundVariableChanged:
 		{
 			if (URigVMPin* ModelPin = Cast<URigVMPin>(InSubject))
 			{
diff --git a/Engine/Plugins/Experimental/ControlRig/Source/ControlRigEditor/Private/Editor/ControlRigEditor.cpp b/Engine/Plugins/Experimental/ControlRig/Source/ControlRigEditor/Private/Editor/ControlRigEditor.cpp
--- a/Engine/Plugins/Experimental/ControlRig/Source/ControlRigEditor/Private/Editor/ControlRigEditor.cpp
+++ b/Engine/Plugins/Experimental/ControlRig/Source/ControlRigEditor/Private/Editor/ControlRigEditor.cpp
@@ -1808,6 +1808,7 @@ void FControlRigEditor::HandleModifiedEvent(ERigVMGraphNotifType InNotifType, UR
 			break;
 		}
 		case ERigVMGraphNotifType::PinArraySizeChanged:
+		case ERigVMGraphNotifType::PinBoundVariableChanged:
 		{
 			URigVMPin* Pin = Cast<URigVMPin>(InSubject);
 
diff --git a/Engine/Source/Developer/RigVMDeveloper/Private/RigVMModel/RigVMController.cpp b/Engine/Source/Developer/RigVMDeveloper/Private/RigVMModel/RigVMController.cpp
--- a/Engine/Source/Developer/RigVMDeveloper/Private/RigVMModel/RigVMController.cpp
+++ b/Engine/Source/Developer/RigVMDeveloper/Private/RigVMModel/RigVMController.cpp
@@ -460,7 +460,7 @@ void URigVMController::RefreshVariableNode(const FName& InNodeName, const FName&
 	}
 }
 
-void URigVMController::RemoveVariableNodes(const FName& InVarName, bool bUndo)
+void URigVMController::OnExternalVariableRemoved(const FName& InVarName, bool bUndo)
 {
 	if (!IsValidGraph())
 	{
@@ -489,9 +489,19 @@ void URigVMController::RemoveVariableNodes(const FName& InVarName, bool bUndo)
 				if (VariablePin->GetDefaultValue() == VarNameStr)
 				{
 					RemoveNode(Node, bUndo, true);
+					continue;
 				}
 			}
 		}
+
+		TArray<URigVMPin*> AllPins = Node->GetAllPinsRecursively();
+		for (URigVMPin* Pin : AllPins)
+		{
+			if (Pin->GetBoundVariableName() == InVarName)
+			{
+				BindPinToVariable(Pin, NAME_None, true);
+			}
+		}
 	}
 
 	if (bUndo)
@@ -500,7 +510,7 @@ void URigVMController::RemoveVariableNodes(const FName& InVarName, bool bUndo)
 	}
 }
 
-void URigVMController::RenameVariableNodes(const FName& InOldVarName, const FName& InNewVarName, bool bUndo)
+void URigVMController::OnExternalVariableRenamed(const FName& InOldVarName, const FName& InNewVarName, bool bUndo)
 {
 	if (!IsValidGraph())
 	{
@@ -529,9 +539,19 @@ void URigVMController::RenameVariableNodes(const FName& InOldVarName, const FNam
 				if (VariablePin->GetDefaultValue() == VarNameStr)
 				{
 					RefreshVariableNode(Node->GetFName(), InNewVarName, FString(), nullptr, bUndo);
+					continue;
 				}
 			}
 		}
+
+		TArray<URigVMPin*> AllPins = Node->GetAllPinsRecursively();
+		for (URigVMPin* Pin : AllPins)
+		{
+			if (Pin->GetBoundVariableName() == InOldVarName)
+			{
+				BindPinToVariable(Pin, InNewVarName, true);
+			}
+		}
 	}
 
 	if (bUndo)
@@ -540,7 +560,7 @@ void URigVMController::RenameVariableNodes(const FName& InOldVarName, const FNam
 	}
 }
 
-void URigVMController::ChangeVariableNodesType(const FName& InVarName, const FString& InCPPType, UObject* InCPPTypeObject, bool bUndo)
+void URigVMController::OnExternalVariableTypeChanged(const FName& InVarName, const FString& InCPPType, UObject* InCPPTypeObject, bool bUndo)
 {
 	if (!IsValidGraph())
 	{
@@ -569,9 +589,21 @@ void URigVMController::ChangeVariableNodesType(const FName& InVarName, const FSt
 				if (VariablePin->GetDefaultValue() == VarNameStr)
 				{
 					RefreshVariableNode(Node->GetFName(), InVarName, InCPPType, InCPPTypeObject, bUndo);
+					continue;
 				}
 			}
 		}
+
+		TArray<URigVMPin*> AllPins = Node->GetAllPinsRecursively();
+		for (URigVMPin* Pin : AllPins)
+		{
+			if (Pin->GetBoundVariableName() == InVarName)
+			{
+				BindPinToVariable(Pin, NAME_None, true);
+				// try to bind it again - maybe it can be bound (due to cast rules etc)
+				BindPinToVariable(Pin, InVarName, true);
+			}
+		}
 	}
 
 	if (bUndo)
@@ -2940,6 +2972,88 @@ bool URigVMController::SetArrayPinSize(const FString& InArrayPinPath, int32 InSi
 	return RemovedPins > 0 || AddedPins > 0;
 }
 
+bool URigVMController::BindPinToVariable(const FString& InPinPath, const FName& InNewBoundVariableName, bool bUndo)
+{
+	if (!IsValidGraph())
+	{
+		return false;
+	}
+
+	URigVMPin* Pin = Graph->FindPin(InPinPath);
+	if (Pin == nullptr)
+	{
+		ReportErrorf(TEXT("Cannot find pin '%s'."), *InPinPath);
+		return false;
+	}
+
+	return BindPinToVariable(Pin, InNewBoundVariableName, bUndo);
+}
+
+bool URigVMController::UnbindPinFromVariable(const FString& InPinPath, bool bUndo)
+{
+	return BindPinToVariable(InPinPath, NAME_None, bUndo);
+}
+
+bool URigVMController::BindPinToVariable(URigVMPin* InPin, const FName& InNewBoundVariableName, bool bUndo)
+{
+	if (!IsValidPinForGraph(InPin))
+	{
+		return false;
+	}
+
+	if (InPin->GetBoundVariableName() == InNewBoundVariableName)
+	{
+		return false;
+	}
+
+	// check that the variable is compatible
+	if (!InNewBoundVariableName.IsNone())
+	{
+		FRigVMExternalVariable ExternalVariable = GetExternalVariableByName(InNewBoundVariableName);
+		if (!InPin->CanBeBoundToVariable(ExternalVariable))
+		{
+			return false;
+		}
+	}
+
+	FRigVMBaseAction Action;
+	if (bUndo)
+	{
+		if (InNewBoundVariableName.IsNone())
+		{
+			Action.Title = TEXT("Unbind pin from variable");
+		}
+		else
+		{
+			Action.Title = TEXT("Bind pin to variable");
+		}
+		ActionStack->BeginAction(Action);
+	}
+
+	if (!InPin->IsBoundToVariable() && bUndo)
+	{
+		// break all links on pin towards parent + children
+		BreakAllLinks(InPin, true, bUndo);
+		BreakAllLinksRecursive(InPin, true, true, bUndo);
+		BreakAllLinksRecursive(InPin, true, false, bUndo);
+	}
+
+	if (bUndo)
+	{
+		ActionStack->AddAction(FRigVMSetPinBoundVariableAction(InPin, InNewBoundVariableName));
+	}
+
+	InPin->BoundVariableName = InNewBoundVariableName;
+	Notify(ERigVMGraphNotifType::PinBoundVariableChanged, InPin);
+
+	if (bUndo)
+	{
+		ActionStack->EndAction(Action);
+	}
+
+	return true;
+}
+
 bool URigVMController::AddLink(const FString& InOutputPinPath, const FString& InInputPinPath, bool bUndo)
 {
 	if(!IsValidGraph())
@@ -2997,11 +3111,19 @@ bool URigVMController::AddLink(URigVMPin* OutputPin, URigVMPin* InputPin, bool b
 		return false;
 	}
 
-	FString FailureReason;
-	if (!Graph->CanLink(OutputPin, InputPin, &FailureReason))
 	{
-		ReportErrorf(TEXT("Cannot link '%s' to '%s': %s."), *OutputPin->GetPinPath(), *InputPin->GetPinPath(), *FailureReason);
-		return false;
+		// Temporarily remove the bound variable from the pin
+		// if it exists, so that link validation can work.
+		// BreakAllPins will remove the bound variable for real later.
+		TGuardValue<FName> OutputBoundVariableGuard(OutputPin->BoundVariableName, NAME_None);
+		TGuardValue<FName> InputBoundVariableGuard(InputPin->BoundVariableName, NAME_None);
+
+		FString FailureReason;
+		if (!Graph->CanLink(OutputPin, InputPin, &FailureReason))
+		{
+			ReportErrorf(TEXT("Cannot link '%s' to '%s': %s."), *OutputPin->GetPinPath(), *InputPin->GetPinPath(), *FailureReason);
+			return false;
+		}
 	}
 
 	ensure(!OutputPin->IsLinkedTo(InputPin));
@@ -3174,6 +3296,12 @@ bool URigVMController::BreakAllLinks(URigVMPin* Pin, bool bAsInput, bool bUndo)
 	}
 
 	int32 LinksBroken = 0;
+	if (Pin->IsBoundToVariable() && bAsInput && bUndo)
+	{
+		BindPinToVariable(Pin, NAME_None, bUndo);
+		LinksBroken++;
+	}
+
 	TArray<URigVMLink*> Links = Pin->GetLinks();
 	for (int32 LinkIndex = Links.Num() - 1; LinkIndex >= 0; LinkIndex--)
 	{
@@ -5064,4 +5192,26 @@ void URigVMController::DestroyObject(UObject* InObjectToDestroy)
 {
 	InObjectToDestroy->Rename(nullptr, GetTransientPackage(), REN_ForceNoResetLoaders | REN_DoNotDirty | REN_DontCreateRedirectors | REN_NonTransactional);
 	InObjectToDestroy->RemoveFromRoot();
-}
\ No newline at end of file
+}
+
+FRigVMExternalVariable URigVMController::GetExternalVariableByName(const FName& InExternalVariableName)
+{
+	TArray<FRigVMExternalVariable> ExternalVariables = GetExternalVariables();
+	for (const FRigVMExternalVariable& ExternalVariable : ExternalVariables)
+	{
+		if (ExternalVariable.Name == InExternalVariableName)
+		{
+			return ExternalVariable;
+		}
+	}
+	return FRigVMExternalVariable();
+}
+
+TArray<FRigVMExternalVariable> URigVMController::GetExternalVariables()
+{
+	if (GetExternalVariablesDelegate.IsBound())
+	{
+		return GetExternalVariablesDelegate.Execute();
+	}
+	return TArray<FRigVMExternalVariable>();
+}
diff --git a/Engine/Source/Developer/RigVMDeveloper/Private/RigVMModel/RigVMControllerActions.cpp b/Engine/Source/Developer/RigVMDeveloper/Private/RigVMModel/RigVMControllerActions.cpp
--- a/Engine/Source/Developer/RigVMDeveloper/Private/RigVMModel/RigVMControllerActions.cpp
+++ b/Engine/Source/Developer/RigVMDeveloper/Private/RigVMModel/RigVMControllerActions.cpp
@@ -1239,3 +1239,42 @@ bool FRigVMChangePinTypeAction::Redo(URigVMController* InController)
 	}
 	return FRigVMBaseAction::Redo(InController);
 }
+
+FRigVMSetPinBoundVariableAction::FRigVMSetPinBoundVariableAction(URigVMPin* InPin, const FName& InNewBoundVariableName)
+	: PinPath(InPin->GetPinPath())
+	, OldBoundVariableName(InPin->GetBoundVariableName())
+	, NewBoundVariableName(InNewBoundVariableName)
+{
+}
+
+bool FRigVMSetPinBoundVariableAction::Merge(const FRigVMBaseAction* Other)
+{
+	if (!FRigVMBaseAction::Merge(Other))
+	{
+		return false;
+	}
+
+	const FRigVMSetPinBoundVariableAction* Action = (const FRigVMSetPinBoundVariableAction*)Other;
+	if (PinPath != Action->PinPath)
+	{
+		return false;
+	}
+
+	NewBoundVariableName = Action->NewBoundVariableName;
+	return true;
+}
+
+bool FRigVMSetPinBoundVariableAction::Undo(URigVMController* InController)
+{
+	if (!FRigVMBaseAction::Undo(InController))
+	{
+		return false;
+	}
+	return InController->BindPinToVariable(PinPath, OldBoundVariableName, false);
+}
+
+bool FRigVMSetPinBoundVariableAction::Redo(URigVMController* InController)
+{
+	InController->BindPinToVariable(PinPath, NewBoundVariableName, false);
+	return FRigVMBaseAction::Redo(InController);
+}
diff --git a/Engine/Source/Developer/RigVMDeveloper/Private/RigVMModel/RigVMPin.cpp b/Engine/Source/Developer/RigVMDeveloper/Private/RigVMModel/RigVMPin.cpp
--- a/Engine/Source/Developer/RigVMDeveloper/Private/RigVMModel/RigVMPin.cpp
+++ b/Engine/Source/Developer/RigVMDeveloper/Private/RigVMModel/RigVMPin.cpp
@@ -82,6 +82,7 @@ URigVMPin::URigVMPin()
 	, CPPTypeObject(nullptr)
 	, CPPTypeObjectPath(NAME_None)
 	, DefaultValue(FString())
+	, BoundVariableName(NAME_None)
 {
 }
 
@@ -392,6 +393,49 @@ UObject* URigVMPin::FindObjectFromCPPTypeObjectPath(const FString& InObjectPath)
 	return FindObject<UObject>(ANY_PACKAGE, *InObjectPath);
 }
 
+bool URigVMPin::CanBeBoundToVariable(const FRigVMExternalVariable& InExternalVariable) const
+{
+	if (!InExternalVariable.IsValid())
+	{
+		return false;
+	}
+
+	if (bIsConstant)
+	{
+		return false;
+	}
+
+	// only allow to bind variables to input pins for now
+	if (Direction == ERigVMPinDirection::Output)
+	{
+		return false;
+	}
+
+	// check type validity
+	if (IsArray() != InExternalVariable.bIsArray)
+	{
+		return false;
+	}
+
+	if (GetCPPTypeObject() != nullptr)
+	{
+		if (GetCPPTypeObject() != InExternalVariable.TypeObject)
+		{
+			return false;
+		}
+	}
+	else
+	{
+		FString CPPBaseType = IsArray() ? GetArrayElementCppType() : GetCPPType();
+		if (CPPBaseType != InExternalVariable.TypeName.ToString())
+		{
+			return false;
+		}
+	}
+
+	return true;
+}
+
 FName URigVMPin::GetSliceContext(const FRigVMUserDataArray& InUserData)
 {
 	URigVMPin* RootPin = GetRootPin();
@@ -718,6 +762,15 @@ bool URigVMPin::CanLink(URigVMPin* InSourcePin, URigVMPin* InTargetPin, FString*
 		}
 		return false;
 	}
+	
+	if (InTargetPin->IsBoundToVariable())
+	{
+		if (OutFailureReason)
+		{
+			*OutFailureReason = TEXT("Target pin is bound to a variable.");
+		}
+		return false;
+	}
 
 	URigVMNode* SourceNode = InSourcePin->GetNode();
 	URigVMNode* TargetNode = InTargetPin->GetNode();
diff --git a/Engine/Source/Developer/RigVMDeveloper/Public/RigVMModel/RigVMController.h b/Engine/Source/Developer/RigVMDeveloper/Public/RigVMModel/RigVMController.h
--- a/Engine/Source/Developer/RigVMDeveloper/Public/RigVMModel/RigVMController.h
+++ b/Engine/Source/Developer/RigVMDeveloper/Public/RigVMModel/RigVMController.h
@@ -18,6 +18,7 @@
 class URigVMActionStack;
 
 DECLARE_DELEGATE_RetVal_OneParam(bool, FRigVMController_ShouldStructUnfoldDelegate, const UStruct*)
+DECLARE_DELEGATE_RetVal(TArray<FRigVMExternalVariable>, FRigVMController_GetExternalVariablesDelegate)
 
 /**
  * The Controller is the sole authority to perform changes
@@ -111,15 +112,15 @@ public:
 
 	// Removes all nodes related to a given variable
 	UFUNCTION(BlueprintCallable, Category = RigVMController)
-	void RemoveVariableNodes(const FName& InVarName, bool bUndo);
+	void OnExternalVariableRemoved(const FName& InVarName, bool bUndo);
 
 	// Renames the variable name in all relevant nodes
 	UFUNCTION(BlueprintCallable, Category = RigVMController)
-	void RenameVariableNodes(const FName& InOldVarName, const FName& InNewVarName, bool bUndo);
+	void OnExternalVariableRenamed(const FName& InOldVarName, const FName& InNewVarName, bool bUndo);
 
 	// Changes the data type of all nodes matching a given variable name
 	UFUNCTION(BlueprintCallable, Category = RigVMController)
-	void ChangeVariableNodesType(const FName& InVarName, const FString& InCPPType, UObject* InCPPTypeObject, bool bUndo);
+	void OnExternalVariableTypeChanged(const FName& InVarName, const FString& InCPPType, UObject* InCPPTypeObject, bool bUndo);
 
 	// Refreshes the variable node with the new data
 	UFUNCTION(BlueprintCallable, Category = RigVMController)
@@ -407,6 +408,16 @@ public:
 	UFUNCTION(BlueprintCallable, Category = RigVMController)
 	bool SetArrayPinSize(const FString& InArrayPinPath, int32 InSize, const FString& InDefaultValue = TEXT(""), bool bUndo = true);
 
+	// Binds a pin to a variable (or removes the binding given NAME_None)
+	// This causes a PinBoundVariableChanged modified event.
+	UFUNCTION(BlueprintCallable, Category = RigVMController)
+	bool BindPinToVariable(const FString& InPinPath, const FName& InNewBoundVariableName, bool bUndo = true);
+
+	// Removes the binging of a pin to a variable
+	// This causes a PinBoundVariableChanged modified event.
+	UFUNCTION(BlueprintCallable, Category = RigVMController)
+	bool UnbindPinFromVariable(const FString& InPinPath, bool bUndo = true);
+
 	// Adds a link to the graph.
 	// This causes a LinkAdded modified event.
 	UFUNCTION(BlueprintCallable, Category = RigVMController)
@@ -428,6 +439,9 @@ public:
 	// A delegate that can be set to change the struct unfolding behaviour
 	FRigVMController_ShouldStructUnfoldDelegate UnfoldStructDelegate;
 
+	// A delegate to retrieve the list of external variables
+	FRigVMController_GetExternalVariablesDelegate GetExternalVariablesDelegate;
+
 	int32 DetachLinksFromPinObjects();
 	int32 ReattachLinksToPinObjects(bool bFollowCoreRedirectors = false);
 	void AddPinRedirector(bool bInput, bool bOutput, const FString& OldPinPath, const FString& NewPinPath);
@@ -470,6 +484,7 @@ private:
 	URigVMPin* InsertArrayPin(URigVMPin* ArrayPin, int32 InIndex, const FString& InDefaultValue, bool bUndo);
 	bool RemovePin(URigVMPin* InPinToRemove, bool bUndo);
 	FProperty* FindPropertyForPin(const FString& InPinPath);
+	bool BindPinToVariable(URigVMPin* InPin, const FName& InNewBoundVariableName, bool bUndo);
 	bool AddLink(URigVMPin* OutputPin, URigVMPin* InputPin, bool bUndo);
 	bool BreakLink(URigVMPin* OutputPin, URigVMPin* InputPin, bool bUndo);
 	bool BreakAllLinks(URigVMPin* Pin, bool bAsInput, bool bUndo);
@@ -513,6 +528,9 @@ private:
 
 	void DestroyObject(UObject* InObjectToDestroy);
 
+	FRigVMExternalVariable GetExternalVariableByName(const FName& InExternalVariableName);
+	TArray<FRigVMExternalVariable> GetExternalVariables();
+
 	UPROPERTY(transient)
 	URigVMGraph* Graph;
 
diff --git a/Engine/Source/Developer/RigVMDeveloper/Public/RigVMModel/RigVMControllerActions.h b/Engine/Source/Developer/RigVMDeveloper/Public/RigVMModel/RigVMControllerActions.h
--- a/Engine/Source/Developer/RigVMDeveloper/Public/RigVMModel/RigVMControllerActions.h
+++ b/Engine/Source/Developer/RigVMDeveloper/Public/RigVMModel/RigVMControllerActions.h
@@ -1100,3 +1100,30 @@ public:
 	UPROPERTY()
 	FName NewCPPTypeObjectPath;
 };
+
+/**
+ * An action changing a pin's bound variable
+ */
+USTRUCT()
+struct FRigVMSetPinBoundVariableAction : public FRigVMBaseAction
+{
+	GENERATED_BODY()
+
+public:
+
+	FRigVMSetPinBoundVariableAction() {}
+	FRigVMSetPinBoundVariableAction(URigVMPin* InPin, const FName& InNewBoundVariableName);
+	virtual ~FRigVMSetPinBoundVariableAction() {};
+	virtual bool Merge(const FRigVMBaseAction* Other);
+	virtual bool Undo(URigVMController* InController) override;
+	virtual bool Redo(URigVMController* InController) override;
+
+	UPROPERTY()
+	FString PinPath;
+
+	UPROPERTY()
+	FName OldBoundVariableName;
+
+	UPROPERTY()
+	FName NewBoundVariableName;
+};
diff --git a/Engine/Source/Developer/RigVMDeveloper/Public/RigVMModel/RigVMNotifications.h b/Engine/Source/Developer/RigVMDeveloper/Public/RigVMModel/RigVMNotifications.h
--- a/Engine/Source/Developer/RigVMDeveloper/Public/RigVMModel/RigVMNotifications.h
+++ b/Engine/Source/Developer/RigVMDeveloper/Public/RigVMModel/RigVMNotifications.h
@@ -43,6 +43,7 @@ enum class ERigVMGraphNotifType : uint8
 	InteractionBracketOpened, // A bracket has been opened (Subject == nullptr)
 	InteractionBracketClosed, // A bracket has been opened (Subject == nullptr)
 	InteractionBracketCanceled, // A bracket has been canceled (Subject == nullptr)
+	PinBoundVariableChanged, // A pin has been bound or unbound to / from a variable (Subject == URigVMPin)
 	Invalid // The max for this enum (used for guarding)
 };
 
diff --git a/Engine/Source/Developer/RigVMDeveloper/Public/RigVMModel/RigVMPin.h b/Engine/Source/Developer/RigVMDeveloper/Public/RigVMModel/RigVMPin.h
--- a/Engine/Source/Developer/RigVMDeveloper/Public/RigVMModel/RigVMPin.h
+++ b/Engine/Source/Developer/RigVMDeveloper/Public/RigVMModel/RigVMPin.h
@@ -288,6 +288,15 @@ public:
 	// Returns the injected nodes this pin contains.
 	const TArray<URigVMInjectionInfo*> GetInjectedNodes() const { return InjectionInfos; }
 
+	// Returns the variable bound to this pin (or NAME_None)
+	const FName& GetBoundVariableName() const { return BoundVariableName; }
+
+	// Returns true if this pin is bound to a variable
+	bool IsBoundToVariable() const { return !GetBoundVariableName().IsNone(); }
+
+	// Returns true if the pin can be bound to a given variable
+	bool CanBeBoundToVariable(const FRigVMExternalVariable& InExternalVariable) const;
+
 	// helper function to retrieve an object from a path
 	static UObject* FindObjectFromCPPTypeObjectPath(const FString& InObjectPath);
 	template<class T>
@@ -352,6 +361,9 @@ private:
 	UPROPERTY()
 	TArray<URigVMInjectionInfo*> InjectionInfos;
 
+	UPROPERTY()
+	FName BoundVariableName;
+
 	friend class URigVMController;
 	friend class URigVMGraph;
 	friend class URigVMNode;
