Control Rig: Support for Variable pin binding

diff --git a/Engine/Plugins/Experimental/ControlRig/Source/ControlRig/Private/ControlRig.cpp b/Engine/Plugins/Experimental/ControlRig/Source/ControlRig/Private/ControlRig.cpp
--- a/Engine/Plugins/Experimental/ControlRig/Source/ControlRig/Private/ControlRig.cpp
+++ b/Engine/Plugins/Experimental/ControlRig/Source/ControlRig/Private/ControlRig.cpp
@@ -1963,10 +1963,12 @@ FRigVMExternalVariable UControlRig::GetExternalVariableFromPinType(const FName&
 	if (InPinType.PinCategory == UEdGraphSchema_K2::PC_Boolean)
 	{
 		ExternalVariable.TypeName = TEXT("bool");
+		ExternalVariable.Size = sizeof(bool);
 	}
 	else if (InPinType.PinCategory == UEdGraphSchema_K2::PC_Int)
 	{
 		ExternalVariable.TypeName = TEXT("int32");
+		ExternalVariable.Size = sizeof(int32);
 	}
 	else if (InPinType.PinCategory == UEdGraphSchema_K2::PC_Enum ||
 		InPinType.PinCategory == UEdGraphSchema_K2::PC_Byte)
@@ -1980,18 +1982,22 @@ FRigVMExternalVariable UControlRig::GetExternalVariableFromPinType(const FName&
 		{
 			ExternalVariable.TypeName = TEXT("uint8");
 		}
+		ExternalVariable.Size = sizeof(uint8);
 	}
 	else if (InPinType.PinCategory == UEdGraphSchema_K2::PC_Float)
 	{
 		ExternalVariable.TypeName = TEXT("float");
+		ExternalVariable.Size = sizeof(float);
 	}
 	else if (InPinType.PinCategory == UEdGraphSchema_K2::PC_Name)
 	{
 		ExternalVariable.TypeName = TEXT("FName");
+		ExternalVariable.Size = sizeof(FName);
 	}
 	else if (InPinType.PinCategory == UEdGraphSchema_K2::PC_String)
 	{
 		ExternalVariable.TypeName = TEXT("FString");
+		ExternalVariable.Size = sizeof(FString);
 	}
 	else if (InPinType.PinCategory == UEdGraphSchema_K2::PC_Struct)
 	{
@@ -1999,6 +2005,7 @@ FRigVMExternalVariable UControlRig::GetExternalVariableFromPinType(const FName&
 		{
 			ExternalVariable.TypeName = *Struct->GetStructCPPName();
 			ExternalVariable.TypeObject = Struct;
+			ExternalVariable.Size = Struct->GetStructureSize();
 		}
 	}
 
diff --git a/Engine/Plugins/Experimental/ControlRig/Source/ControlRig/Public/ControlRig.h b/Engine/Plugins/Experimental/ControlRig/Source/ControlRig/Public/ControlRig.h
--- a/Engine/Plugins/Experimental/ControlRig/Source/ControlRig/Public/ControlRig.h
+++ b/Engine/Plugins/Experimental/ControlRig/Source/ControlRig/Public/ControlRig.h
@@ -32,7 +32,7 @@ class UScriptStruct;
 struct FRigUnit;
 struct FRigControl;
 
-DECLARE_LOG_CATEGORY_EXTERN(LogControlRig, Log, All);
+CONTROLRIG_API DECLARE_LOG_CATEGORY_EXTERN(LogControlRig, Log, All);
 
 /** Runs logic for mapping input data to transforms (the "Rig") */
 UCLASS(Blueprintable, Abstract, editinlinenew)
diff --git a/Engine/Plugins/Experimental/ControlRig/Source/ControlRigDeveloper/Private/ControlRigBlueprint.cpp b/Engine/Plugins/Experimental/ControlRig/Source/ControlRigDeveloper/Private/ControlRigBlueprint.cpp
--- a/Engine/Plugins/Experimental/ControlRig/Source/ControlRigDeveloper/Private/ControlRigBlueprint.cpp
+++ b/Engine/Plugins/Experimental/ControlRig/Source/ControlRigDeveloper/Private/ControlRigBlueprint.cpp
@@ -25,7 +25,6 @@
 #if WITH_EDITOR
 #include "IControlRigEditorModule.h"
 #include "Kismet2/BlueprintEditorUtils.h"
-#include "Kismet2/Kismet2NameValidators.h"
 #include "ControlRigBlueprintUtils.h"
 #include "Settings/ControlRigSettings.h"
 #include "UnrealEdGlobals.h"
@@ -87,6 +86,46 @@ void UControlRigBlueprint::InitializeModelIfRequired(bool bRecompileVM)
 			return true;
 		});
 
+		TWeakObjectPtr<UControlRigBlueprint> WeakThis(this);
+
+		// this delegate is used by the controller to determine variable validity
+		// during a bind process. the controller itself doesn't own the variables,
+		// so we need a delegate to request them from the owning blueprint
+		Controller->GetExternalVariablesDelegate.BindLambda([WeakThis]() -> TArray<FRigVMExternalVariable> {
+			
+			if (WeakThis.IsValid())
+			{
+				if (UControlRigBlueprintGeneratedClass* RigClass = WeakThis->GetControlRigBlueprintGeneratedClass())
+				{
+					if (UControlRig* CDO = Cast<UControlRig>(RigClass->GetDefaultObject(true /* create if needed */)))
+					{
+						return CDO->GetExternalVariablesImpl(true /* rely on variables within blueprint */);
+					}
+				}
+			}
+			return TArray<FRigVMExternalVariable>();
+
+		});
+
+#if WITH_EDITOR
+
+		// this sets up three delegates:
+		// a) get external variables (mapped to Controller->GetExternalVariables)
+		// b) bind pin to variable (mapped to Controller->BindPinToVariable)
+		// c) create external variable (mapped to the passed in tfunction)
+		// the last one is defined within the blueprint since the controller
+		// doesn't own the variables and can't create one itself.
+		Controller->SetupDefaultStructNodeDelegates(TDelegate<FName(FRigVMExternalVariable)>::CreateLambda(
+			[WeakThis](FRigVMExternalVariable InVariableToCreate) -> FName {
+				if (WeakThis.IsValid())
+				{
+					return WeakThis->AddCRMemberVariableFromExternal(InVariableToCreate);
+				}
+				return NAME_None;
+			}
+		));
+#endif
+
 		Controller->RemoveStaleNodes();
 
 		for (int32 i = 0; i < UbergraphPages.Num(); ++i)
@@ -1183,6 +1222,7 @@ void UControlRigBlueprint::HandleModifiedEvent(ERigVMGraphNotifType InNotifType,
 			case ERigVMGraphNotifType::ParameterAdded:
 			case ERigVMGraphNotifType::ParameterRemoved:
 			case ERigVMGraphNotifType::ParameterRenamed:
+			case ERigVMGraphNotifType::PinBoundVariableChanged:
 			{
 				RequestAutoVMRecompilation();
 				MarkPackageDirty();
@@ -1264,74 +1304,6 @@ void UControlRigBlueprint::CreateMemberVariablesOnLoad()
 	{
 		TSharedPtr<FKismetNameValidator> NameValidator = MakeShareable(new FKismetNameValidator(this, NAME_None, nullptr));
 
-		struct Local
-		{
-			static FName FindUniqueName(TSharedPtr<FKismetNameValidator> InNameValidator, const FString& InBaseName)
-			{
-				FString KismetName = InBaseName;
-				if (InNameValidator->IsValid(KismetName) == EValidatorResult::ContainsInvalidCharacters)
-				{
-					for (TCHAR& TestChar : KismetName)
-					{
-						for (TCHAR BadChar : UE_BLUEPRINT_INVALID_NAME_CHARACTERS)
-						{
-							if (TestChar == BadChar)
-							{
-								TestChar = TEXT('_');
-								break;
-							}
-						}
-					}
-				}
-
-				int32 Suffix = 0;
-				while (InNameValidator->IsValid(KismetName) != EValidatorResult::Ok)
-				{
-					KismetName = FString::Printf(TEXT("%s_%d"), *InBaseName, Suffix);
-					Suffix++;
-				}
-
-
-				return *KismetName;
-			}
-
-			static int32 AddMemberVariable(UControlRigBlueprint* InBlueprint, const FName& InVarName, FEdGraphPinType InVarType, bool bIsPublic, bool bIsReadOnly)
-			{
-				FBPVariableDescription NewVar;
-
-				NewVar.VarName = InVarName;
-				NewVar.VarGuid = FGuid::NewGuid();
-				NewVar.FriendlyName = FName::NameToDisplayString(InVarName.ToString(), (InVarType.PinCategory == UEdGraphSchema_K2::PC_Boolean) ? true : false);
-				NewVar.VarType = InVarType;
-
-				NewVar.PropertyFlags |= (CPF_Edit | CPF_BlueprintVisible | CPF_DisableEditOnInstance);
-
-				if (bIsPublic)
-				{
-					NewVar.PropertyFlags &= ~CPF_DisableEditOnInstance;
-				}
-
-				if (bIsReadOnly)
-				{
-					NewVar.PropertyFlags |= CPF_BlueprintReadOnly;
-				}
-
-				NewVar.ReplicationCondition = COND_None;
-
-				NewVar.Category = UEdGraphSchema_K2::VR_DefaultCategory;
-
-				// user created variables should be none of these things
-				NewVar.VarType.bIsConst = false;
-				NewVar.VarType.bIsWeakPointer = false;
-				NewVar.VarType.bIsReference = false;
-
-				// Text variables, etc. should default to multiline
-				NewVar.SetMetaData(TEXT("MultiLine"), TEXT("true"));
-
-				return InBlueprint->NewVariables.Add(NewVar);
-			}
-		};
-
 		TArray<URigVMNode*> Nodes = Model->GetNodes();
 		for (URigVMNode* Node : Nodes)
 		{
@@ -1357,8 +1329,8 @@ void UControlRigBlueprint::CreateMemberVariablesOnLoad()
 					continue;
 				}
 
-				FName VarName = Local::FindUniqueName(NameValidator, Description.Name.ToString());
-				int32 VariableIndex = Local::AddMemberVariable(this, VarName, PinType, false, false);
+				FName VarName = FindCRMemberVariableUniqueName(NameValidator, Description.Name.ToString());
+				int32 VariableIndex = AddCRMemberVariable(this, VarName, PinType, false, false);
 				if (VariableIndex != INDEX_NONE)
 				{
 					AddedMemberVariableMap.Add(Description.Name, VariableIndex);
@@ -1388,8 +1360,8 @@ void UControlRigBlueprint::CreateMemberVariablesOnLoad()
 					continue;
 				}
 
-				FName VarName = Local::FindUniqueName(NameValidator, Description.Name.ToString());
-				int32 VariableIndex = Local::AddMemberVariable(this, VarName, PinType, true, !Description.bIsInput);
+				FName VarName = FindCRMemberVariableUniqueName(NameValidator, Description.Name.ToString());
+				int32 VariableIndex = AddCRMemberVariable(this, VarName, PinType, true, !Description.bIsInput);
 				if (VariableIndex != INDEX_NONE)
 				{
 					AddedMemberVariableMap.Add(Description.Name, VariableIndex);
@@ -1402,6 +1374,97 @@ void UControlRigBlueprint::CreateMemberVariablesOnLoad()
 #endif
 }
 
+#if WITH_EDITOR
+
+FName UControlRigBlueprint::FindCRMemberVariableUniqueName(TSharedPtr<FKismetNameValidator> InNameValidator, const FString& InBaseName)
+{
+	FString KismetName = InBaseName;
+	if (InNameValidator->IsValid(KismetName) == EValidatorResult::ContainsInvalidCharacters)
+	{
+		for (TCHAR& TestChar : KismetName)
+		{
+			for (TCHAR BadChar : UE_BLUEPRINT_INVALID_NAME_CHARACTERS)
+			{
+				if (TestChar == BadChar)
+				{
+					TestChar = TEXT('_');
+					break;
+				}
+			}
+		}
+	}
+
+	int32 Suffix = 0;
+	while (InNameValidator->IsValid(KismetName) != EValidatorResult::Ok)
+	{
+		KismetName = FString::Printf(TEXT("%s_%d"), *InBaseName, Suffix);
+		Suffix++;
+	}
+
+
+	return *KismetName;
+}
+
+int32 UControlRigBlueprint::AddCRMemberVariable(UControlRigBlueprint* InBlueprint, const FName& InVarName, FEdGraphPinType InVarType, bool bIsPublic, bool bIsReadOnly)
+{
+	FBPVariableDescription NewVar;
+
+	NewVar.VarName = InVarName;
+	NewVar.VarGuid = FGuid::NewGuid();
+	NewVar.FriendlyName = FName::NameToDisplayString(InVarName.ToString(), (InVarType.PinCategory == UEdGraphSchema_K2::PC_Boolean) ? true : false);
+	NewVar.VarType = InVarType;
+
+	NewVar.PropertyFlags |= (CPF_Edit | CPF_BlueprintVisible | CPF_DisableEditOnInstance);
+
+	if (bIsPublic)
+	{
+		NewVar.PropertyFlags &= ~CPF_DisableEditOnInstance;
+	}
+
+	if (bIsReadOnly)
+	{
+		NewVar.PropertyFlags |= CPF_BlueprintReadOnly;
+	}
+
+	NewVar.ReplicationCondition = COND_None;
+
+	NewVar.Category = UEdGraphSchema_K2::VR_DefaultCategory;
+
+	// user created variables should be none of these things
+	NewVar.VarType.bIsConst = false;
+	NewVar.VarType.bIsWeakPointer = false;
+	NewVar.VarType.bIsReference = false;
+
+	// Text variables, etc. should default to multiline
+	NewVar.SetMetaData(TEXT("MultiLine"), TEXT("true"));
+
+	return InBlueprint->NewVariables.Add(NewVar);
+}
+
+FName UControlRigBlueprint::AddCRMemberVariableFromExternal(FRigVMExternalVariable InVariableToCreate)
+{
+	FEdGraphPinType PinType = UControlRig::GetPinTypeFromExternalVariable(InVariableToCreate);
+	if (!PinType.PinCategory.IsValid())
+	{
+		return NAME_None;
+	}
+
+	Modify();
+
+	TSharedPtr<FKismetNameValidator> NameValidator = MakeShareable(new FKismetNameValidator(this, NAME_None, nullptr));
+	FName VarName = FindCRMemberVariableUniqueName(NameValidator, InVariableToCreate.Name.ToString());
+	int32 VariableIndex = AddCRMemberVariable(this, VarName, PinType, InVariableToCreate.bIsPublic, InVariableToCreate.bIsReadOnly);
+	if (VariableIndex != INDEX_NONE)
+	{
+		FBlueprintEditorUtils::MarkBlueprintAsStructurallyModified(this);
+		return VarName;
+	}
+
+	return NAME_None;
+}
+
+#endif
+
 void UControlRigBlueprint::PatchVariableNodesOnLoad()
 {
 #if WITH_EDITOR
diff --git a/Engine/Plugins/Experimental/ControlRig/Source/ControlRigDeveloper/Private/Graph/ControlRigGraphSchema.cpp b/Engine/Plugins/Experimental/ControlRig/Source/ControlRigDeveloper/Private/Graph/ControlRigGraphSchema.cpp
--- a/Engine/Plugins/Experimental/ControlRig/Source/ControlRigDeveloper/Private/Graph/ControlRigGraphSchema.cpp
+++ b/Engine/Plugins/Experimental/ControlRig/Source/ControlRigDeveloper/Private/Graph/ControlRigGraphSchema.cpp
@@ -457,6 +457,44 @@ bool UControlRigGraphSchema::RequestVariableDropOnPanel(UEdGraph* InGraph, FProp
 	return false;
 }
 
+bool UControlRigGraphSchema::RequestVariableDropOnPin(UEdGraph* InGraph, FProperty* InVariableToDrop, UEdGraphPin* InPin, const FVector2D& InDropPosition, const FVector2D& InScreenPosition)
+{
+#if WITH_EDITOR
+	if (CanVariableBeDropped(InGraph, InVariableToDrop))
+	{
+		UBlueprint* Blueprint = FBlueprintEditorUtils::FindBlueprintForGraph(InGraph);
+		UControlRigBlueprint* RigBlueprint = Cast<UControlRigBlueprint>(Blueprint);
+		if (RigBlueprint != nullptr)
+		{
+			if (URigVMPin* ModelPin = RigBlueprint->Model->FindPin(InPin->GetName()))
+			{
+				FRigVMExternalVariable ExternalVariable = FRigVMExternalVariable::Make(InVariableToDrop, nullptr);
+				if (ModelPin->CanBeBoundToVariable(ExternalVariable))
+				{
+					FModifierKeysState KeyState = FSlateApplication::Get().GetModifierKeys();
+					if (KeyState.IsAltDown())
+					{
+						return RigBlueprint->Controller->BindPinToVariable(ModelPin->GetPinPath(), InVariableToDrop->GetName());
+					}
+					else
+					{
+						RigBlueprint->Controller->OpenUndoBracket(TEXT("Bind Variable to Pin"));
+						if (URigVMVariableNode* VariableNode = RigBlueprint->Controller->AddVariableNode(ExternalVariable.Name, ExternalVariable.TypeName.ToString(), ExternalVariable.TypeObject, true, FString(), InDropPosition + FVector2D(0.f, -34.f)))
+						{
+							RigBlueprint->Controller->AddLink(VariableNode->FindPin(TEXT("Value"))->GetPinPath(), ModelPin->GetPinPath(), true);
+						}
+						RigBlueprint->Controller->CloseUndoBracket();
+						return true;
+					}
+				}
+			}
+		}
+	}
+#endif
+
+	return false;
+}
+
 void UControlRigGraphSchema::EndGraphNodeInteraction(UEdGraphNode* InNode) const
 {
 #if WITH_EDITOR
diff --git a/Engine/Plugins/Experimental/ControlRig/Source/ControlRigDeveloper/Public/ControlRigBlueprint.h b/Engine/Plugins/Experimental/ControlRig/Source/ControlRigDeveloper/Public/ControlRigBlueprint.h
--- a/Engine/Plugins/Experimental/ControlRig/Source/ControlRigDeveloper/Public/ControlRigBlueprint.h
+++ b/Engine/Plugins/Experimental/ControlRig/Source/ControlRigDeveloper/Public/ControlRigBlueprint.h
@@ -19,6 +19,11 @@
 #include "ControlRigHierarchyModifier.h"
 #include "ControlRigValidationPass.h"
 #include "Drawing/ControlRigDrawContainer.h"
+
+#if WITH_EDITOR
+#include "Kismet2/Kismet2NameValidators.h"
+#endif
+
 #include "ControlRigBlueprint.generated.h"
 
 class UControlRigBlueprintGeneratedClass;
@@ -238,6 +243,11 @@ private:
 
 	void CleanupBoneHierarchyDeprecated();
 	void CreateMemberVariablesOnLoad();
+#if WITH_EDITOR
+	static FName FindCRMemberVariableUniqueName(TSharedPtr<FKismetNameValidator> InNameValidator, const FString& InBaseName);
+	static int32 AddCRMemberVariable(UControlRigBlueprint* InBlueprint, const FName& InVarName, FEdGraphPinType InVarType, bool bIsPublic, bool bIsReadOnly);
+	FName AddCRMemberVariableFromExternal(FRigVMExternalVariable InVariableToCreate);
+#endif
 	void PatchVariableNodesOnLoad();
 
 	TMap<FName, int32> AddedMemberVariableMap;
diff --git a/Engine/Plugins/Experimental/ControlRig/Source/ControlRigDeveloper/Public/Graph/ControlRigGraphSchema.h b/Engine/Plugins/Experimental/ControlRig/Source/ControlRigDeveloper/Public/Graph/ControlRigGraphSchema.h
--- a/Engine/Plugins/Experimental/ControlRig/Source/ControlRigDeveloper/Public/Graph/ControlRigGraphSchema.h
+++ b/Engine/Plugins/Experimental/ControlRig/Source/ControlRigDeveloper/Public/Graph/ControlRigGraphSchema.h
@@ -74,6 +74,7 @@ public:
 	virtual bool SafeDeleteNodeFromGraph(UEdGraph* Graph, UEdGraphNode* Node) const override;
 	virtual bool CanVariableBeDropped(UEdGraph* InGraph, FProperty* InVariableToDrop) const override;
 	virtual bool RequestVariableDropOnPanel(UEdGraph* InGraph, FProperty* InVariableToDrop, const FVector2D& InDropPosition, const FVector2D& InScreenPosition) override;
+	virtual bool RequestVariableDropOnPin(UEdGraph* InGraph, FProperty* InVariableToDrop, UEdGraphPin* InPin, const FVector2D& InDropPosition, const FVector2D& InScreenPosition) override;
 	virtual bool IsStructEditable(UStruct* InStruct) const;
 
 	/** Create a graph node for a rig */
diff --git a/Engine/Plugins/Experimental/ControlRig/Source/ControlRigEditor/ControlRigEditor.Build.cs b/Engine/Plugins/Experimental/ControlRig/Source/ControlRigEditor/ControlRigEditor.Build.cs
--- a/Engine/Plugins/Experimental/ControlRig/Source/ControlRigEditor/ControlRigEditor.Build.cs
+++ b/Engine/Plugins/Experimental/ControlRig/Source/ControlRigEditor/ControlRigEditor.Build.cs
@@ -72,7 +72,8 @@ namespace UnrealBuildTool.Rules
                     "RigVM",
                     "RigVMDeveloper",
 					"AnimationEditor",
-					"MessageLog"
+					"MessageLog",
+					"PropertyAccessEditor",
 				}
             );
 
diff --git a/Engine/Plugins/Experimental/ControlRig/Source/ControlRigEditor/Private/ControlRigEditorModule.cpp b/Engine/Plugins/Experimental/ControlRig/Source/ControlRigEditor/Private/ControlRigEditorModule.cpp
--- a/Engine/Plugins/Experimental/ControlRig/Source/ControlRigEditor/Private/ControlRigEditorModule.cpp
+++ b/Engine/Plugins/Experimental/ControlRig/Source/ControlRigEditor/Private/ControlRigEditorModule.cpp
@@ -104,7 +104,7 @@
 #include "AnimSequenceLevelSequenceLink.h"
 #include "Rigs/FKControlRig.h"
 #include "SBakeToControlRigDialog.h"
-
+#include "Graph/SControlRigGraphPinVariableBinding.h"
 
 #define LOCTEXT_NAMESPACE "ControlRigEditorModule"
 
@@ -1132,8 +1132,62 @@ void FControlRigEditorModule::GetContextMenuActions(const UControlRigGraphSchema
 						));
 					}
 
+					if (ModelPin->GetDirection() == ERigVMPinDirection::Input &&
+							!ModelPin->IsExecuteContext())
+					{
+						if (ModelPin->IsBoundToVariable())
+						{
+							FVector2D NodePosition = FVector2D(Context->Node->NodePosX - 200.f, Context->Node->NodePosY);
+
+							FToolMenuSection& VariablesSection = Menu->AddSection("EdGraphSchemaVariables", LOCTEXT("Variables", "Variables"));
+							VariablesSection.AddMenuEntry(
+								"MakeVariableNodeFromBinding",
+								LOCTEXT("MakeVariableNodeFromBinding", "Make Variable Node"),
+								LOCTEXT("MakeVariableNodeFromBinding_Tooltip", "Turns the variable binding on the pin to a variable node"),
+								FSlateIcon(),
+								FUIAction(FExecuteAction::CreateLambda([RigBlueprint, ModelPin, NodePosition]() {
+									RigBlueprint->Controller->MakeVariableNodeFromBinding(ModelPin->GetPinPath(), NodePosition);
+								})
+							));
+						}
+						else
+						{
+							FVector2D NodePosition = FVector2D(Context->Node->NodePosX - 200.f, Context->Node->NodePosY);
+
+							FToolMenuSection& VariablesSection = Menu->AddSection("EdGraphSchemaVariables", LOCTEXT("Variables", "Variables"));
+							VariablesSection.AddMenuEntry(
+								"PromotePinToVariable",
+								LOCTEXT("PromotePinToVariable", "Promote Pin To Variable"),
+								LOCTEXT("PromotePinToVariable_Tooltip", "Turns the variable into a variable"),
+								FSlateIcon(),
+								FUIAction(FExecuteAction::CreateLambda([RigBlueprint, ModelPin, NodePosition]() {
+
+									FModifierKeysState KeyState = FSlateApplication::Get().GetModifierKeys();
+									bool bCreateVariableNode = !KeyState.IsAltDown();
+
+									RigBlueprint->Controller->PromotePinToVariable(ModelPin->GetPinPath(), bCreateVariableNode, NodePosition);
+								})
+							));
+						}
+					}
+
 					if(Cast<URigVMStructNode>(ModelPin->GetNode()))
 					{
+						if (ModelPin->GetDirection() == ERigVMPinDirection::Input && 
+							!ModelPin->IsExecuteContext() &&
+							!ModelPin->IsBoundToVariable())
+						{
+							FToolMenuSection& VariablesSection = Menu->FindOrAddSection(TEXT("Variables"));
+
+							TSharedRef<SControlRigVariableBinding> VariableBindingWidget =
+								SNew(SControlRigVariableBinding)
+								.Blueprint(RigBlueprint)
+								.ModelPin(ModelPin)
+								.CanRemoveBinding(false);
+
+							VariablesSection.AddEntry(FToolMenuEntry::InitWidget("BindPinToVariableWidget", VariableBindingWidget, FText(), true));
+						}
+
 						FToolMenuSection& Section = Menu->AddSection("EdGraphSchemaPinDefaults", LOCTEXT("PinDefaults", "Pin Defaults"));
 						Section.AddMenuEntry(
 							"ResetPinDefaultValue",
@@ -1144,7 +1198,7 @@ void FControlRigEditorModule::GetContextMenuActions(const UControlRigGraphSchema
 								RigBlueprint->Controller->ResetPinDefaultValue(ModelPin->GetPinPath());
 							})
 						));
-	}
+					}
 
 					if ((ModelPin->GetCPPType() == TEXT("FVector") ||
 						 ModelPin->GetCPPType() == TEXT("FQuat") ||
@@ -1598,6 +1652,20 @@ void FControlRigEditorModule::GetContextMenuActions(const UControlRigGraphSchema
 							})
 						));
 					}
+
+					if (URigVMVariableNode* VariableNode = Cast<URigVMVariableNode>(RigNode->GetModelNode()))
+					{
+						FToolMenuSection& VariablesSection = Menu->AddSection("EdGraphSchemaVariables", LOCTEXT("VariablesSettingsHeader", "Variables"));
+						VariablesSection.AddMenuEntry(
+							"MakePindingsFromVariableNode",
+							LOCTEXT("MakePindingsFromVariableNode", "Make Bindings From Node"),
+							LOCTEXT("MakePindingsFromVariableNode_Tooltip", "Turns the variable node into one ore more variable bindings on the pin(s)"),
+							FSlateIcon(),
+							FUIAction(FExecuteAction::CreateLambda([RigBlueprint, VariableNode]() {
+								RigBlueprint->Controller->MakeBindingsFromVariableNode(VariableNode->GetFName());
+							})
+						));
+					}
 				}
 			}
 		}
diff --git a/Engine/Plugins/Experimental/ControlRig/Source/ControlRigEditor/Private/Graph/ControlRigGraphPanelPinFactory.cpp b/Engine/Plugins/Experimental/ControlRig/Source/ControlRigEditor/Private/Graph/ControlRigGraphPanelPinFactory.cpp
--- a/Engine/Plugins/Experimental/ControlRig/Source/ControlRigEditor/Private/Graph/ControlRigGraphPanelPinFactory.cpp
+++ b/Engine/Plugins/Experimental/ControlRig/Source/ControlRigEditor/Private/Graph/ControlRigGraphPanelPinFactory.cpp
@@ -8,6 +8,7 @@
 #include "Graph/SControlRigGraphPinCurveFloat.h"
 #include "Graph/SControlRigGraphPinVariableName.h"
 #include "Graph/SControlRigGraphPinParameterName.h"
+#include "Graph/SControlRigGraphPinVariableBinding.h"
 #include "KismetPins/SGraphPinExec.h"
 #include "ControlRig.h"
 #include "NodeFactory.h"
@@ -15,6 +16,7 @@
 #include "Curves/CurveFloat.h"
 #include "RigVMModel/Nodes/RigVMStructNode.h"
 #include "RigVMCore/RigVMExecuteContext.h"
+#include "IPropertyAccessEditor.h"
 
 TSharedPtr<SGraphPin> FControlRigGraphPanelPinFactory::CreatePin(UEdGraphPin* InPin) const
 {
@@ -26,6 +28,16 @@ TSharedPtr<SGraphPin> FControlRigGraphPanelPinFactory::CreatePin(UEdGraphPin* In
 
 			if (URigVMPin* ModelPin = RigNode->GetModelPinFromPinPath(InPin->GetName()))
 			{
+				if (ModelPin->IsBoundToVariable())
+				{
+					if (UControlRigBlueprint* Blueprint = Cast<UControlRigBlueprint>(RigGraph->GetOuter()))
+					{
+						return SNew(SControlRigGraphPinVariableBinding, InPin)
+							.ModelPin(ModelPin)
+							.Blueprint(Blueprint);
+					}
+				}
+
 				FName CustomWidgetName = ModelPin->GetCustomWidgetName();
 				if (CustomWidgetName == TEXT("BoneName"))
 				{
diff --git a/Engine/Plugins/Experimental/ControlRig/Source/ControlRigEditor/Private/Graph/NodeSpawners/ControlRigUnitNodeSpawner.cpp b/Engine/Plugins/Experimental/ControlRig/Source/ControlRigEditor/Private/Graph/NodeSpawners/ControlRigUnitNodeSpawner.cpp
--- a/Engine/Plugins/Experimental/ControlRig/Source/ControlRigEditor/Private/Graph/NodeSpawners/ControlRigUnitNodeSpawner.cpp
+++ b/Engine/Plugins/Experimental/ControlRig/Source/ControlRigEditor/Private/Graph/NodeSpawners/ControlRigUnitNodeSpawner.cpp
@@ -141,6 +141,9 @@ UControlRigGraphNode* UControlRigUnitNodeSpawner::SpawnNode(UEdGraph* ParentGrap
 			Controller->OpenUndoBracket(FString::Printf(TEXT("Add '%s' Node"), *Name.ToString()));
 		}
 
+		FRigVMStructNodeCreatedContext& StructNodeCreatedContext = Controller->GetStructNodeCreatedContext();
+		FRigVMStructNodeCreatedContext::FScope ReasonScope(StructNodeCreatedContext, ERigVMNodeCreatedReason::NodeSpawner);
+
 		if (URigVMStructNode* ModelNode = Controller->AddStructNode(StructTemplate, TEXT("Execute"), Location, Name.ToString(), bUndo))
 		{
 			NewNode = Cast<UControlRigGraphNode>(RigGraph->FindNodeForModelNodeName(ModelNode->GetFName()));
diff --git a/Engine/Plugins/Experimental/ControlRig/Source/ControlRigEditor/Private/Graph/SControlRigGraphNode.cpp b/Engine/Plugins/Experimental/ControlRig/Source/ControlRigEditor/Private/Graph/SControlRigGraphNode.cpp
--- a/Engine/Plugins/Experimental/ControlRig/Source/ControlRigEditor/Private/Graph/SControlRigGraphNode.cpp
+++ b/Engine/Plugins/Experimental/ControlRig/Source/ControlRigEditor/Private/Graph/SControlRigGraphNode.cpp
@@ -602,7 +602,7 @@ TSharedRef<ITableRow> SControlRigGraphNode::MakeTableRowWidget(URigVMPin* InItem
 			}
 
 			// Only leaf pins have value widgets, but not containers
-			if(bLeaf && !bIsContainer && bIsPlainOrEditableStruct)
+			if(((bLeaf && bIsPlainOrEditableStruct) || InItem->IsBoundToVariable()) && !bIsContainer)
 			{
 				InputPinValueWidget = (*InputGraphPinPtr)->GetValueWidget();
 			}
diff --git a/Engine/Plugins/Experimental/ControlRig/Source/ControlRigEditor/Private/Graph/SControlRigGraphPinVariableBinding.cpp b/Engine/Plugins/Experimental/ControlRig/Source/ControlRigEditor/Private/Graph/SControlRigGraphPinVariableBinding.cpp
new file mode 100644
--- /dev/null
+++ b/Engine/Plugins/Experimental/ControlRig/Source/ControlRigEditor/Private/Graph/SControlRigGraphPinVariableBinding.cpp
@@ -0,0 +1,146 @@
+// Copyright Epic Games, Inc. All Rights Reserved.
+
+
+#include "Graph/SControlRigGraphPinVariableBinding.h"
+#include "Widgets/Layout/SBox.h"
+#include "EdGraphSchema_K2.h"
+#include "DetailLayoutBuilder.h"
+#include "RigVMModel/RigVMController.h"
+#include "ControlRigBlueprintGeneratedClass.h"
+
+#define LOCTEXT_NAMESPACE "SControlRigGraphPinVariableBinding"
+
+void SControlRigVariableBinding::Construct(const FArguments& InArgs)
+{
+	this->ModelPin = InArgs._ModelPin;
+	this->Blueprint = InArgs._Blueprint;
+	this->bCanRemoveBinding = InArgs._CanRemoveBinding;
+
+	IPropertyAccessEditor& PropertyAccessEditor = IModularFeatures::Get().GetModularFeature<IPropertyAccessEditor>("PropertyAccessEditor");
+
+	BindingArgs.CurrentBindingText.BindRaw(this, &SControlRigVariableBinding::GetBindingText);
+	BindingArgs.CurrentBindingImage.BindRaw(this, &SControlRigVariableBinding::GetBindingImage);
+	BindingArgs.CurrentBindingColor.BindRaw(this, &SControlRigVariableBinding::GetBindingColor);
+
+	BindingArgs.OnCanBindProperty.BindSP(this, &SControlRigVariableBinding::OnCanBindProperty);
+	BindingArgs.OnCanBindToClass.BindSP(this, &SControlRigVariableBinding::OnCanBindToClass);
+
+	BindingArgs.OnAddBinding.BindSP(this, &SControlRigVariableBinding::OnAddBinding);
+	BindingArgs.OnCanRemoveBinding.BindSP(this, &SControlRigVariableBinding::OnCanRemoveBinding);
+	BindingArgs.OnRemoveBinding.BindSP(this, &SControlRigVariableBinding::OnRemoveBinding);
+
+	BindingArgs.bGeneratePureBindings = true;
+	BindingArgs.bAllowNewBindings = true;
+	BindingArgs.bAllowArrayElementBindings = false;
+	BindingArgs.bAllowUObjectFunctions = false;
+
+	this->ChildSlot
+	[
+		PropertyAccessEditor.MakePropertyBindingWidget(Blueprint, BindingArgs)
+	];
+}
+
+FText SControlRigVariableBinding::GetBindingText() const
+{
+	if (ModelPin)
+	{
+		FString VariablePath = ModelPin->GetBoundVariablePath();
+		return FText::FromString(VariablePath);
+	}
+
+	return FText();
+}
+
+const FSlateBrush* SControlRigVariableBinding::GetBindingImage() const
+{
+	static FName PropertyIcon(TEXT("Kismet.Tabs.Variables"));
+	return FEditorStyle::GetBrush(PropertyIcon);
+}
+
+FLinearColor SControlRigVariableBinding::GetBindingColor() const
+{
+	if (Blueprint && ModelPin)
+	{
+		const UEdGraphSchema_K2* Schema = GetDefault<UEdGraphSchema_K2>();
+
+		for (const FBPVariableDescription& VariableDescription : Blueprint->NewVariables)
+		{
+			if (VariableDescription.VarName.ToString() == ModelPin->GetBoundVariableName())
+			{
+				return Schema->GetPinTypeColor(VariableDescription.VarType);
+			}
+		}
+	}
+	return FLinearColor::White;
+}
+
+bool SControlRigVariableBinding::OnCanBindProperty(FProperty* InProperty) const
+{
+	if (InProperty == BindingArgs.Property)
+	{
+		return true;
+	}
+
+	if (InProperty && ModelPin)
+	{
+		FRigVMExternalVariable ExternalVariable = FRigVMExternalVariable::Make(InProperty, nullptr);
+		return ModelPin->CanBeBoundToVariable(ExternalVariable);
+	}
+
+	return false;
+}
+
+bool SControlRigVariableBinding::OnCanBindToClass(UClass* InClass) const
+{
+	if (InClass)
+	{
+		return InClass->ClassGeneratedBy == Blueprint;
+	}
+	return true;
+}
+
+void SControlRigVariableBinding::OnAddBinding(FName InPropertyName, const TArray<FBindingChainElement>& InBindingChain)
+{
+	if (Blueprint && ModelPin)
+	{
+		TArray<FString> Parts;
+		for (const FBindingChainElement& ChainElement : InBindingChain)
+		{
+			ensure(ChainElement.Field);
+			Parts.Add(ChainElement.Field.GetName());
+		}
+		Blueprint->Controller->BindPinToVariable(ModelPin->GetPinPath(), FString::Join(Parts, TEXT(".")), true /* undo */);
+	}
+}
+
+bool SControlRigVariableBinding::OnCanRemoveBinding(FName InPropertyName)
+{
+	return bCanRemoveBinding;
+}
+
+void SControlRigVariableBinding::OnRemoveBinding(FName InPropertyName)
+{
+	if (Blueprint && ModelPin)
+	{
+		Blueprint->Controller->UnbindPinFromVariable(ModelPin->GetPinPath(), true /* undo */);
+	}
+}
+
+/////////////////////////////////////////////////////////////////////////////////////////////////////////////
+
+void SControlRigGraphPinVariableBinding::Construct(const FArguments& InArgs, UEdGraphPin* InGraphPinObj)
+{
+	this->ModelPin = InArgs._ModelPin;
+	this->Blueprint = InArgs._Blueprint;
+
+	SGraphPin::Construct(SGraphPin::FArguments(), InGraphPinObj);
+}
+
+TSharedRef<SWidget>	SControlRigGraphPinVariableBinding::GetDefaultValueWidget()
+{
+	return SNew(SControlRigVariableBinding)
+		.Blueprint(Blueprint)
+		.ModelPin(ModelPin);
+}
+
+#undef LOCTEXT_NAMESPACE
\ No newline at end of file
diff --git a/Engine/Plugins/Experimental/ControlRig/Source/ControlRigEditor/Private/Graph/SControlRigGraphPinVariableBinding.h b/Engine/Plugins/Experimental/ControlRig/Source/ControlRigEditor/Private/Graph/SControlRigGraphPinVariableBinding.h
new file mode 100644
--- /dev/null
+++ b/Engine/Plugins/Experimental/ControlRig/Source/ControlRigEditor/Private/Graph/SControlRigGraphPinVariableBinding.h
@@ -0,0 +1,68 @@
+// Copyright Epic Games, Inc. All Rights Reserved.
+
+#pragma once
+
+#include "CoreMinimal.h"
+#include "Widgets/DeclarativeSyntaxSupport.h"
+#include "Widgets/SWidget.h"
+#include "SGraphPin.h"
+#include "RigVMModel/RigVMPin.h"
+#include "IPropertyAccessEditor.h"
+#include "ControlRigBlueprint.h"
+#include "SControlRigGraphPinVariableBinding.h"
+
+class SControlRigVariableBinding : public SCompoundWidget
+{
+public:
+
+	SLATE_BEGIN_ARGS(SControlRigVariableBinding)
+	: _CanRemoveBinding(true)
+	{}
+
+		SLATE_ARGUMENT(URigVMPin*, ModelPin)
+		SLATE_ARGUMENT(UControlRigBlueprint*, Blueprint)
+		SLATE_ARGUMENT(bool, CanRemoveBinding)
+
+	SLATE_END_ARGS()
+
+	void Construct(const FArguments& InArgs);
+
+protected:
+
+	FText GetBindingText() const;
+	const FSlateBrush* GetBindingImage() const;
+	FLinearColor GetBindingColor() const;
+	bool OnCanBindProperty(FProperty* InProperty) const;
+	bool OnCanBindToClass(UClass* InClass) const;
+	void OnAddBinding(FName InPropertyName, const TArray<FBindingChainElement>& InBindingChain);
+	bool OnCanRemoveBinding(FName InPropertyName);
+	void OnRemoveBinding(FName InPropertyName);
+
+	URigVMPin* ModelPin;
+	UControlRigBlueprint* Blueprint;
+	FPropertyBindingWidgetArgs BindingArgs;
+	bool bCanRemoveBinding;
+};
+
+class SControlRigGraphPinVariableBinding : public SGraphPin
+{
+public:
+
+	SLATE_BEGIN_ARGS(SControlRigGraphPinVariableBinding){}
+
+		SLATE_ARGUMENT(URigVMPin*, ModelPin)
+		SLATE_ARGUMENT(UControlRigBlueprint*, Blueprint)
+
+	SLATE_END_ARGS()
+
+	void Construct(const FArguments& InArgs, UEdGraphPin* InGraphPinObj);
+
+protected:
+
+	//~ Begin SGraphPin Interface
+	virtual TSharedRef<SWidget>	GetDefaultValueWidget() override;
+	//~ End SGraphPin Interface
+
+	URigVMPin* ModelPin;
+	UControlRigBlueprint* Blueprint;
+};
diff --git a/Engine/Source/Developer/RigVMDeveloper/Private/RigVMCompiler/RigVMAST.cpp b/Engine/Source/Developer/RigVMDeveloper/Private/RigVMCompiler/RigVMAST.cpp
--- a/Engine/Source/Developer/RigVMDeveloper/Private/RigVMCompiler/RigVMAST.cpp
+++ b/Engine/Source/Developer/RigVMDeveloper/Private/RigVMCompiler/RigVMAST.cpp
@@ -51,6 +51,10 @@ FName FRigVMExprAST::GetTypeName() const
 		{
 			return TEXT("[Literal]");
 		}
+		case EType::ExternalVar:
+		{
+			return TEXT("[ExtVar.]");
+		}
 		case EType::Assign:
 		{
 			return TEXT("[.Assign]");
@@ -329,6 +333,11 @@ FString FRigVMExprAST::DumpDot(TArray<bool>& OutExpressionDefined, const FString
 				Label = FString::Printf(TEXT("%s(Literal)"), *To<FRigVMLiteralExprAST>()->GetPin()->GetName());
 				break;
 			}
+			case EType::ExternalVar:
+			{
+				Label = FString::Printf(TEXT("%s(ExternalVar)"), *To<FRigVMExternalVarExprAST>()->GetPin()->GetBoundVariableName());
+				break;
+			}
 			case EType::Var:
 			{
 				if (To<FRigVMVarExprAST>()->IsGraphParameter())
@@ -543,6 +552,15 @@ bool FRigVMNodeExprAST::IsConstant() const
 		{
 			return false;
 		}
+
+		TArray<URigVMPin*> AllPins = CurrentNode->GetAllPinsRecursively();
+		for (URigVMPin* Pin : AllPins)
+		{
+			if (Pin->IsBoundToVariable())
+			{
+				return false;
+			}
+		}
 	}
 	return FRigVMExprAST::IsConstant();
 }
@@ -1077,6 +1095,24 @@ FRigVMExprAST* FRigVMParserAST::TraversePin(URigVMPin* InPin, FRigVMExprAST* InP
 
 	TArray<URigVMLink*> SourceLinks = InPin->GetSourceLinks(true);
 
+	struct Local
+	{
+		static void LookForPinsBoundToVariables(URigVMPin* InPin, TArray<URigVMPin*>& SubPinsBoundToVariables)
+		{
+			for (URigVMPin* SubPin : InPin->GetSubPins())
+			{
+				if (SubPin->IsBoundToVariable())
+				{
+					SubPinsBoundToVariables.Add(SubPin);
+				}
+				LookForPinsBoundToVariables(SubPin, SubPinsBoundToVariables);
+			}
+		}
+	};
+
+	TArray<URigVMPin*> SubPinsBoundToVariables;
+	Local::LookForPinsBoundToVariables(InPin, SubPinsBoundToVariables);
+
 	FRigVMExprAST* PinExpr = nullptr;
 
 	if (Cast<URigVMVariableNode>(InPin->GetNode()))
@@ -1097,7 +1133,8 @@ FRigVMExprAST* FRigVMParserAST::TraversePin(URigVMPin* InPin, FRigVMExprAST* InP
 
 	if ((InPin->GetDirection() == ERigVMPinDirection::Input ||
 		InPin->GetDirection() == ERigVMPinDirection::Visible) &&
-		SourceLinks.Num() == 0)
+		SourceLinks.Num() == 0 &&
+		SubPinsBoundToVariables.Num() == 0)
 	{
 		if (Cast<URigVMParameterNode>(InPin->GetNode()) ||
 			Cast<URigVMVariableNode>(InPin->GetNode()))
@@ -1109,6 +1146,10 @@ FRigVMExprAST* FRigVMParserAST::TraversePin(URigVMPin* InPin, FRigVMExprAST* InP
 			PinCopyExpr->AddParent(PinExpr);
 			PinLiteralExpr->AddParent(PinCopyExpr);
 		}
+		else if (InPin->IsBoundToVariable())
+		{
+			PinExpr = MakeExpr<FRigVMExternalVarExprAST>(InPin);
+		}
 		else
 		{
 			PinExpr = MakeExpr<FRigVMLiteralExprAST>(InPin);
@@ -1132,6 +1173,11 @@ FRigVMExprAST* FRigVMParserAST::TraversePin(URigVMPin* InPin, FRigVMExprAST* InP
 		return PinExpr;
 	}
 
+	if (PinExpr->IsA(FRigVMExprAST::ExternalVar))
+	{
+		return PinExpr;
+	}
+
 	if ((InPin->GetDirection() == ERigVMPinDirection::IO ||
 		InPin->GetDirection() == ERigVMPinDirection::Input) 
 		&& !InPin->IsExecuteContext())
@@ -1149,7 +1195,7 @@ FRigVMExprAST* FRigVMParserAST::TraversePin(URigVMPin* InPin, FRigVMExprAST* InP
 
 		if (!bHasSourceLinkToRoot && 
 			InPin->GetSourceLinks(false).Num() == 0 &&
-			(InPin->GetDirection() == ERigVMPinDirection::IO || SourceLinks.Num() > 0))
+			(InPin->GetDirection() == ERigVMPinDirection::IO || SourceLinks.Num() > 0 || SubPinsBoundToVariables.Num() > 0))
 		{
 			FRigVMLiteralExprAST* LiteralExpr = MakeExpr<FRigVMLiteralExprAST>(InPin);
 			FRigVMCopyExprAST* LiteralCopyExpr = MakeExpr<FRigVMCopyExprAST>(InPin, InPin);
@@ -1160,6 +1206,14 @@ FRigVMExprAST* FRigVMParserAST::TraversePin(URigVMPin* InPin, FRigVMExprAST* InP
 
 			SubjectToExpression[InPin] = LiteralExpr;
 		}
+		else
+		{
+			SubPinsBoundToVariables.Reset();
+		}
+	}
+	else
+	{
+		SubPinsBoundToVariables.Reset();
 	}
 
 	FRigVMExprAST* ParentExprForLinks = PinExpr;
@@ -1179,6 +1233,16 @@ FRigVMExprAST* FRigVMParserAST::TraversePin(URigVMPin* InPin, FRigVMExprAST* InP
 		TraverseLink(SourceLink, ParentExprForLinks);
 	}
 
+	for (URigVMPin* SubPinBoundToVariable : SubPinsBoundToVariables)
+	{
+		FRigVMExternalVarExprAST* ExternalVarExpr = MakeExpr<FRigVMExternalVarExprAST>(SubPinBoundToVariable);
+		FRigVMCopyExprAST* ExternalVarExprCopyExpr = MakeExpr<FRigVMCopyExprAST>(SubPinBoundToVariable, SubPinBoundToVariable);
+		ExternalVarExprCopyExpr->Name = *FString::Printf(TEXT("%s -> %s"), *SubPinBoundToVariable->GetPinPath(), *SubPinBoundToVariable->GetPinPath());
+		ExternalVarExprCopyExpr->AddParent(PinExpr);
+		ExternalVarExpr->AddParent(ExternalVarExprCopyExpr);
+		ExternalVarExpr->Name = *SubPinBoundToVariable->GetPinPath();
+	}
+
 	return PinExpr;
 }
 
diff --git a/Engine/Source/Developer/RigVMDeveloper/Private/RigVMCompiler/RigVMCompiler.cpp b/Engine/Source/Developer/RigVMDeveloper/Private/RigVMCompiler/RigVMCompiler.cpp
--- a/Engine/Source/Developer/RigVMDeveloper/Private/RigVMCompiler/RigVMCompiler.cpp
+++ b/Engine/Source/Developer/RigVMDeveloper/Private/RigVMCompiler/RigVMCompiler.cpp
@@ -345,6 +345,11 @@ void URigVMCompiler::TraverseExpression(const FRigVMExprAST* InExpr, FRigVMCompi
 			TraverseLiteral(InExpr->To<FRigVMLiteralExprAST>(), WorkData);
 			break;
 		}
+		case FRigVMExprAST::EType::ExternalVar:
+		{
+			TraverseExternalVar(InExpr->To<FRigVMExternalVarExprAST>(), WorkData);
+			break;
+		}
 		case FRigVMExprAST::EType::Assign:
 		{
 			TraverseAssign(InExpr->To<FRigVMAssignExprAST>(), WorkData);
@@ -644,6 +649,11 @@ void URigVMCompiler::TraverseLiteral(const FRigVMVarExprAST* InExpr, FRigVMCompi
 	TraverseVar(InExpr, WorkData);
 }
 
+void URigVMCompiler::TraverseExternalVar(const FRigVMExternalVarExprAST* InExpr, FRigVMCompilerWorkData& WorkData)
+{
+	TraverseVar(InExpr, WorkData);
+}
+
 void URigVMCompiler::TraverseAssign(const FRigVMAssignExprAST* InExpr, FRigVMCompilerWorkData& WorkData)
 {
 	TraverseChildren(InExpr, WorkData);
@@ -708,7 +718,7 @@ void URigVMCompiler::TraverseAssign(const FRigVMAssignExprAST* InExpr, FRigVMCom
 		{
 			struct Local
 			{
-				static void SetupRegisterOffset(URigVM* VM, URigVMPin* Pin, FRigVMOperand& Operand)
+				static void SetupRegisterOffset(URigVM* VM, URigVMPin* Pin, FRigVMOperand& Operand, bool bSource)
 				{
 					URigVMPin* RootPin = Pin->GetRootPin();
 					if (Pin == RootPin)
@@ -716,6 +726,11 @@ void URigVMCompiler::TraverseAssign(const FRigVMAssignExprAST* InExpr, FRigVMCom
 						return;
 					}
 
+					if (bSource && Pin->IsBoundToVariable())
+					{
+						return;
+					}
+
 					if (URigVMSelectNode* SelectNode = Cast<URigVMSelectNode>(RootPin->GetNode()))
 					{
 						if (Pin->GetParentPin() == RootPin && RootPin->GetName() == URigVMSelectNode::ValueName)
@@ -758,8 +773,8 @@ void URigVMCompiler::TraverseAssign(const FRigVMAssignExprAST* InExpr, FRigVMCom
 				}
 			};
 
-			Local::SetupRegisterOffset(WorkData.VM, InExpr->GetSourcePin(), Source);
-			Local::SetupRegisterOffset(WorkData.VM, InExpr->GetTargetPin(), Target);
+			Local::SetupRegisterOffset(WorkData.VM, InExpr->GetSourcePin(), Source, true);
+			Local::SetupRegisterOffset(WorkData.VM, InExpr->GetTargetPin(), Target, false);
 		}
 
 		WorkData.VM->GetByteCode().AddCopyOp(Source, Target);
@@ -1087,11 +1102,17 @@ FString URigVMCompiler::GetPinHash(URigVMPin* InPin, const FRigVMVarExprAST* InV
 
 	if (InVarExpr != nullptr && !bIsDebugValue)
 	{
+		if (InVarExpr->IsA(FRigVMExprAST::ExternalVar) && !bIsDebugValue)
+		{
+			FString VariablePath = InPin->GetBoundVariablePath();
+			return FString::Printf(TEXT("%sVariable::%s%s"), *Prefix, *VariablePath, *Suffix);
+		}
+
 		bIsExecutePin = InPin->IsExecuteContext();
 		bIsLiteral = InVarExpr->GetType() == FRigVMExprAST::EType::Literal;
 
 		bIsParameter = Cast<URigVMParameterNode>(Node) != nullptr;
-		bIsVariable = Cast<URigVMVariableNode>(Node) != nullptr;
+		bIsVariable = Cast<URigVMVariableNode>(Node) != nullptr || InVarExpr->IsA(FRigVMExprAST::ExternalVar);
 
 		// determine if this is an initialization for an IO pin
 		if (!bIsLiteral && !bIsParameter && !bIsVariable &&
@@ -1253,6 +1274,30 @@ FRigVMOperand URigVMCompiler::FindOrAddRegister(const FRigVMVarExprAST* InVarExp
 		return *ExistingOperand;
 	}
 
+	// check if this is a variable with a segment path
+	if (!Operand.IsValid())
+	{
+		if (Pin->IsBoundToVariable() && !bIsDebugValue)
+		{
+			FString VariablePath = Pin->GetBoundVariablePath();
+			FString VariableName = VariablePath, SegmentPath;
+			VariablePath.Split(TEXT("."), &VariableName, &SegmentPath);
+
+			ExistingOperand = WorkData.PinPathToOperand->Find(FString::Printf(TEXT("Variable::%s"), *VariableName));
+			if (ExistingOperand)
+			{
+				Operand = *ExistingOperand;
+				Operand.RegisterOffset = INDEX_NONE;
+				if (!SegmentPath.IsEmpty())
+				{
+					const FRigVMExternalVariable& ExternalVariable = WorkData.VM->GetExternalVariables()[Operand.GetRegisterIndex()];
+					UScriptStruct* ScriptStruct = CastChecked<UScriptStruct>(ExternalVariable.TypeObject);
+					Operand.RegisterOffset = WorkData.VM->GetWorkMemory().GetOrAddRegisterOffset(Operand.GetRegisterIndex(), ScriptStruct, SegmentPath, 0 /*ArrayIndex */);
+				}
+			}
+		}
+	}
+
 	// create remaining operands / registers
 	if (!Operand.IsValid())
 	{
diff --git a/Engine/Source/Developer/RigVMDeveloper/Private/RigVMModel/RigVMController.cpp b/Engine/Source/Developer/RigVMDeveloper/Private/RigVMModel/RigVMController.cpp
--- a/Engine/Source/Developer/RigVMDeveloper/Private/RigVMModel/RigVMController.cpp
+++ b/Engine/Source/Developer/RigVMDeveloper/Private/RigVMModel/RigVMController.cpp
@@ -497,9 +497,9 @@ void URigVMController::OnExternalVariableRemoved(const FName& InVarName, bool bU
 		TArray<URigVMPin*> AllPins = Node->GetAllPinsRecursively();
 		for (URigVMPin* Pin : AllPins)
 		{
-			if (Pin->GetBoundVariableName() == InVarName)
+			if (Pin->GetBoundVariableName() == InVarName.ToString())
 			{
-				BindPinToVariable(Pin, NAME_None, true);
+				BindPinToVariable(Pin, FString(), true);
 			}
 		}
 	}
@@ -547,9 +547,11 @@ void URigVMController::OnExternalVariableRenamed(const FName& InOldVarName, cons
 		TArray<URigVMPin*> AllPins = Node->GetAllPinsRecursively();
 		for (URigVMPin* Pin : AllPins)
 		{
-			if (Pin->GetBoundVariableName() == InOldVarName)
+			if (Pin->GetBoundVariableName() == InOldVarName.ToString())
 			{
-				BindPinToVariable(Pin, InNewVarName, true);
+				FString OldVariablePath = Pin->GetBoundVariablePath();
+				FString NewVariablePath = OldVariablePath.Replace(*InOldVarName.ToString(), *InNewVarName.ToString());
+				BindPinToVariable(Pin, NewVariablePath, true);
 			}
 		}
 	}
@@ -597,11 +599,12 @@ void URigVMController::OnExternalVariableTypeChanged(const FName& InVarName, con
 		TArray<URigVMPin*> AllPins = Node->GetAllPinsRecursively();
 		for (URigVMPin* Pin : AllPins)
 		{
-			if (Pin->GetBoundVariableName() == InVarName)
+			if (Pin->GetBoundVariableName() == InVarName.ToString())
 			{
-				BindPinToVariable(Pin, NAME_None, true);
+				FString BoundVariablePath = Pin->GetBoundVariablePath();
+				BindPinToVariable(Pin, FString(), true);
 				// try to bind it again - maybe it can be bound (due to cast rules etc)
-				BindPinToVariable(Pin, InVarName, true);
+				BindPinToVariable(Pin, BoundVariablePath, true);
 			}
 		}
 	}
@@ -1545,6 +1548,7 @@ TArray<FName> URigVMController::ImportNodesFromText(const FString& InText, bool
 		ActionStack->BeginAction(AddNodesAction);
 	}
 
+	FRigVMStructNodeCreatedContext::FScope StructNodeCreatedScope(StructNodeCreatedContext, ERigVMNodeCreatedReason::Paste);
 	for (URigVMNode* CreatedNode : Factory.CreatedNodes)
 	{
 		Graph->Nodes.Add(CreatedNode);
@@ -1553,6 +1557,20 @@ TArray<FName> URigVMController::ImportNodesFromText(const FString& InText, bool
 		{
 			ActionStack->AddAction(FRigVMRemoveNodeAction(CreatedNode));
 		}
+
+		if (URigVMStructNode* StructNode = Cast<URigVMStructNode>(CreatedNode))
+		{
+			if (StructNodeCreatedContext.IsValid())
+			{
+				if (TSharedPtr<FStructOnScope> StructScope = StructNode->ConstructStructInstance())
+				{
+					TGuardValue<FName> NodeNameScope(StructNodeCreatedContext.NodeName, StructNode->GetFName());
+					FRigVMStruct* StructInstance = (FRigVMStruct*)StructScope->GetStructMemory();
+					StructInstance->OnStructNodeCreated(StructNodeCreatedContext);
+				}
+			}
+		}
+
 		Notify(ERigVMGraphNotifType::NodeAdded, CreatedNode);
 
 		NodeNames.Add(CreatedNode->GetFName());
@@ -2972,7 +2990,7 @@ bool URigVMController::SetArrayPinSize(const FString& InArrayPinPath, int32 InSi
 	return RemovedPins > 0 || AddedPins > 0;
 }
 
-bool URigVMController::BindPinToVariable(const FString& InPinPath, const FName& InNewBoundVariableName, bool bUndo)
+bool URigVMController::BindPinToVariable(const FString& InPinPath, const FString& InNewBoundVariablePath, bool bUndo)
 {
 	if (!IsValidGraph())
 	{
@@ -2986,31 +3004,46 @@ bool URigVMController::BindPinToVariable(const FString& InPinPath, const FName&
 		return false;
 	}
 
-	return BindPinToVariable(Pin, InNewBoundVariableName, bUndo);
+	return BindPinToVariable(Pin, InNewBoundVariablePath, bUndo);
 }
 
 bool URigVMController::UnbindPinFromVariable(const FString& InPinPath, bool bUndo)
 {
-	return BindPinToVariable(InPinPath, NAME_None, bUndo);
+	return BindPinToVariable(InPinPath, FString(), bUndo);
 }
 
-bool URigVMController::BindPinToVariable(URigVMPin* InPin, const FName& InNewBoundVariableName, bool bUndo)
+bool URigVMController::BindPinToVariable(URigVMPin* InPin, const FString& InNewBoundVariablePath, bool bUndo)
 {
 	if (!IsValidPinForGraph(InPin))
 	{
 		return false;
 	}
 
-	if (InPin->GetBoundVariableName() == InNewBoundVariableName)
+	if (InPin->GetBoundVariablePath() == InNewBoundVariablePath)
 	{
 		return false;
 	}
 
 	// check that the variable is compatible
-	if (!InNewBoundVariableName.IsNone())
+	if (!InNewBoundVariablePath.IsEmpty())
 	{
-		FRigVMExternalVariable ExternalVariable = GetExternalVariableByName(InNewBoundVariableName);
-		if (!InPin->CanBeBoundToVariable(ExternalVariable))
+		FString VariableName = InNewBoundVariablePath, SegmentPath;
+		InNewBoundVariablePath.Split(TEXT("."), &VariableName, &SegmentPath);
+
+		FRigVMExternalVariable ExternalVariable = GetExternalVariableByName(*VariableName);
+		if (ExternalVariable.IsValid(true))
+		{
+			FRigVMRegisterOffset RegisterOffset;
+			if (!SegmentPath.IsEmpty())
+			{
+				RegisterOffset = FRigVMRegisterOffset(Cast<UScriptStruct>(ExternalVariable.TypeObject), SegmentPath);
+			}
+			if (!InPin->CanBeBoundToVariable(ExternalVariable, RegisterOffset))
+			{
+				return false;
+			}
+		}
+		else
 		{
 			return false;
 		}
@@ -3019,7 +3052,7 @@ bool URigVMController::BindPinToVariable(URigVMPin* InPin, const FName& InNewBou
 	FRigVMBaseAction Action;
 	if (bUndo)
 	{
-		if (InNewBoundVariableName.IsNone())
+		if (InNewBoundVariablePath.IsEmpty())
 		{
 			Action.Title = TEXT("Unbind pin from variable");
 		}
@@ -3040,10 +3073,10 @@ bool URigVMController::BindPinToVariable(URigVMPin* InPin, const FName& InNewBou
 
 	if (bUndo)
 	{
-		ActionStack->AddAction(FRigVMSetPinBoundVariableAction(InPin, InNewBoundVariableName));
+		ActionStack->AddAction(FRigVMSetPinBoundVariableAction(InPin, InNewBoundVariablePath));
 	}
 
-	InPin->BoundVariableName = InNewBoundVariableName;
+	InPin->BoundVariablePath = InNewBoundVariablePath;
 	Notify(ERigVMGraphNotifType::PinBoundVariableChanged, InPin);
 
 	if (bUndo)
@@ -3054,6 +3087,159 @@ bool URigVMController::BindPinToVariable(URigVMPin* InPin, const FName& InNewBou
 	return true;
 }
 
+bool URigVMController::MakeBindingsFromVariableNode(const FName& InNodeName, bool bUndo)
+{
+	if (!IsValidGraph())
+	{
+		return false;
+	}
+
+	if (URigVMVariableNode* VariableNode = Cast<URigVMVariableNode>(Graph->FindNodeByName(InNodeName)))
+	{
+		return MakeBindingsFromVariableNode(VariableNode, bUndo);
+	}
+
+	return false;
+}
+
+bool URigVMController::MakeBindingsFromVariableNode(URigVMVariableNode* InNode, bool bUndo)
+{
+	check(InNode);
+
+	TArray<URigVMLink*> Links;
+	if (URigVMPin* ValuePin = InNode->FindPin(URigVMVariableNode::ValueName))
+	{
+		Links = ValuePin->GetTargetLinks(true);
+	}
+
+	FName VariableName = InNode->GetVariableName();
+	FRigVMExternalVariable ExternalVariable = GetExternalVariableByName(VariableName);
+	if (!ExternalVariable.IsValid(true /* allow nullptr */))
+	{
+		return false;
+	}
+
+	if (Links.Num() > 0)
+	{
+		if (bUndo)
+		{
+			OpenUndoBracket(TEXT("Turn Variable Node into Bindings"));
+		}
+
+		for (URigVMLink* Link : Links)
+		{
+			URigVMPin* SourcePin = Link->GetSourcePin();
+			URigVMPin* TargetPin = Link->GetTargetPin();
+			FString SegmentPath = SourcePin->GetSegmentPath();
+			FString VariablePathToBind = VariableName.ToString();
+			if (!SegmentPath.IsEmpty())
+			{
+				VariablePathToBind = FString::Printf(TEXT("%s.%s"), *VariablePathToBind, *SegmentPath);
+			}
+
+			if (!BindPinToVariable(TargetPin, VariablePathToBind, bUndo))
+			{
+				CancelUndoBracket();
+			}
+		}
+
+		RemoveNode(InNode, bUndo, true);
+
+		if (bUndo)
+		{
+			CloseUndoBracket();
+		}
+		return true;
+	}
+
+	return false;
+
+}
+
+bool URigVMController::MakeVariableNodeFromBinding(const FString& InPinPath, const FVector2D& InNodePosition, bool bUndo)
+{
+	return PromotePinToVariable(InPinPath, true, InNodePosition, bUndo);
+}
+
+bool URigVMController::PromotePinToVariable(const FString& InPinPath, bool bCreateVariableNode, const FVector2D& InNodePosition, bool bUndo)
+{
+	if (!IsValidGraph())
+	{
+		return false;
+	}
+
+	URigVMPin* Pin = Graph->FindPin(InPinPath);
+	if (Pin == nullptr)
+	{
+		ReportErrorf(TEXT("Cannot find pin '%s'."), *InPinPath);
+		return false;
+	}
+
+	return PromotePinToVariable(Pin, bCreateVariableNode, InNodePosition, bUndo);
+}
+
+bool URigVMController::PromotePinToVariable(URigVMPin* InPin, bool bCreateVariableNode, const FVector2D& InNodePosition, bool bUndo)
+{
+	check(InPin);
+
+	if (InPin->GetDirection() != ERigVMPinDirection::Input)
+	{
+		return false;
+	}
+
+	FRigVMExternalVariable VariableForPin;
+	if (InPin->IsBoundToVariable())
+	{
+		VariableForPin = GetExternalVariableByName(*InPin->GetBoundVariableName());
+		check(VariableForPin.IsValid(true /* allow nullptr */));
+	}
+	else
+	{
+		if (!StructNodeCreatedContext.GetCreateExternalVariableDelegate().IsBound())
+		{
+			return false;
+		}
+
+		VariableForPin = InPin->ToExternalVariable();
+		FName VariableName = StructNodeCreatedContext.GetCreateExternalVariableDelegate().Execute(VariableForPin);
+		if (VariableName.IsNone())
+		{
+			return false;
+		}
+
+		VariableForPin = GetExternalVariableByName(VariableName);
+		if (!VariableForPin.IsValid(true /* allow nullptr*/))
+		{
+			return false;
+		}
+	}
+
+	if (bCreateVariableNode)
+	{
+		if (URigVMVariableNode* VariableNode = AddVariableNode(
+			VariableForPin.Name,
+			VariableForPin.TypeName.ToString(),
+			VariableForPin.TypeObject,
+			true,
+			FString(),
+			InNodePosition,
+			FString(),
+			bUndo))
+		{
+			if (URigVMPin* ValuePin = VariableNode->FindPin(URigVMVariableNode::ValueName))
+			{
+				return AddLink(ValuePin, InPin, bUndo);
+			}
+		}
+	}
+	else
+	{
+		return BindPinToVariable(InPin, VariableForPin.Name.ToString(), bUndo);
+	}
+
+	return false;
+}
+
 bool URigVMController::AddLink(const FString& InOutputPinPath, const FString& InInputPinPath, bool bUndo)
 {
 	if(!IsValidGraph())
@@ -3115,8 +3301,8 @@ bool URigVMController::AddLink(URigVMPin* OutputPin, URigVMPin* InputPin, bool b
 		// Temporarily remove the bound variable from the pin
 		// if it exists, so that link validation can work.
 		// BreakAllPins will remove the bound variable for real later.
-		TGuardValue<FName> OutputBoundVariableGuard(OutputPin->BoundVariableName, NAME_None);
-		TGuardValue<FName> InputBoundVariableGuard(InputPin->BoundVariableName, NAME_None);
+		TGuardValue<FString> OutputBoundVariableGuard(OutputPin->BoundVariablePath, FString());
+		TGuardValue<FString> InputBoundVariableGuard(InputPin->BoundVariablePath, FString());
 
 		FString FailureReason;
 		if (!Graph->CanLink(OutputPin, InputPin, &FailureReason))
@@ -3298,7 +3484,7 @@ bool URigVMController::BreakAllLinks(URigVMPin* Pin, bool bAsInput, bool bUndo)
 	int32 LinksBroken = 0;
 	if (Pin->IsBoundToVariable() && bAsInput && bUndo)
 	{
-		BindPinToVariable(Pin, NAME_None, bUndo);
+		BindPinToVariable(Pin, FString(), bUndo);
 		LinksBroken++;
 	}
 
@@ -4693,6 +4879,14 @@ void URigVMController::RepopulatePinsOnNode(URigVMNode* InNode)
 		}
 	}
 
+	TMap<FString, FString> BoundVariables;
+	for (URigVMPin* Pin : AllPins)
+	{
+		FString PinPath, NodeName;
+		URigVMPin::SplitPinPathAtStart(Pin->GetPinPath(), NodeName, PinPath);
+		BoundVariables.Add(PinPath, Pin->GetBoundVariablePath());
+	}
+
 	// also in case this node is part of an injection
 	FName InjectionInputPinName = NAME_None;
 	FName InjectionOutputPinName = NAME_None;
@@ -4878,6 +5072,20 @@ void URigVMController::RepopulatePinsOnNode(URigVMNode* InNode)
 		}
 	}
 
+	for (TPair<FString, FString> BoundVariablePair: BoundVariables)
+	{
+		FString PinPath = BoundVariablePair.Key;
+		if (RedirectedPinPaths.Contains(PinPath))
+		{
+			PinPath = RedirectedPinPaths.FindChecked(PinPath);
+		}
+
+		if (URigVMPin* Pin = InNode->FindPin(PinPath))
+		{
+			BindPinToVariable(Pin, BoundVariablePair.Value, false);
+		}
+	}
+
 	if (URigVMInjectionInfo* InjectionInfo = InNode->GetInjectionInfo())
 	{
 		InjectionInfo->InputPin = InNode->FindPin(InjectionInputPinName.ToString());
@@ -4888,6 +5096,36 @@ void URigVMController::RepopulatePinsOnNode(URigVMNode* InNode)
 
 #endif
 
+void URigVMController::SetupDefaultStructNodeDelegates(TDelegate<FName(FRigVMExternalVariable)> InCreateExternalVariableDelegate)
+{
+	TWeakObjectPtr<URigVMController> WeakThis(this);
+
+	StructNodeCreatedContext.GetAllExternalVariablesDelegate().BindLambda([WeakThis]() -> TArray<FRigVMExternalVariable> {
+		if (WeakThis.IsValid())
+		{
+			return WeakThis->GetExternalVariables();
+		}
+		return TArray<FRigVMExternalVariable>();
+	});
+
+	StructNodeCreatedContext.GetBindPinToExternalVariableDelegate().BindLambda([WeakThis](FString InPinPath, FString InVariablePath) -> bool {
+		if (WeakThis.IsValid())
+		{
+			return WeakThis->BindPinToVariable(InPinPath, InVariablePath, true);
+		}
+		return false;
+	});
+
+	StructNodeCreatedContext.GetCreateExternalVariableDelegate() = InCreateExternalVariableDelegate;
+}
+
+void URigVMController::ResetStructNodeDelegates()
+{
+	StructNodeCreatedContext.GetAllExternalVariablesDelegate().Unbind();
+	StructNodeCreatedContext.GetBindPinToExternalVariableDelegate().Unbind();
+	StructNodeCreatedContext.GetCreateExternalVariableDelegate().Unbind();
+}
+
 FLinearColor URigVMController::GetColorFromMetadata(const FString& InMetadata)
 {
 	FLinearColor Color = FLinearColor::Black;
diff --git a/Engine/Source/Developer/RigVMDeveloper/Private/RigVMModel/RigVMControllerActions.cpp b/Engine/Source/Developer/RigVMDeveloper/Private/RigVMModel/RigVMControllerActions.cpp
--- a/Engine/Source/Developer/RigVMDeveloper/Private/RigVMModel/RigVMControllerActions.cpp
+++ b/Engine/Source/Developer/RigVMDeveloper/Private/RigVMModel/RigVMControllerActions.cpp
@@ -1240,10 +1240,10 @@ bool FRigVMChangePinTypeAction::Redo(URigVMController* InController)
 	return FRigVMBaseAction::Redo(InController);
 }
 
-FRigVMSetPinBoundVariableAction::FRigVMSetPinBoundVariableAction(URigVMPin* InPin, const FName& InNewBoundVariableName)
+FRigVMSetPinBoundVariableAction::FRigVMSetPinBoundVariableAction(URigVMPin* InPin, const FString& InNewBoundVariablePath)
 	: PinPath(InPin->GetPinPath())
-	, OldBoundVariableName(InPin->GetBoundVariableName())
-	, NewBoundVariableName(InNewBoundVariableName)
+	, OldBoundVariablePath(InPin->GetBoundVariablePath())
+	, NewBoundVariablePath(InNewBoundVariablePath)
 {
 }
 
@@ -1260,7 +1260,7 @@ bool FRigVMSetPinBoundVariableAction::Merge(const FRigVMBaseAction* Other)
 		return false;
 	}
 
-	NewBoundVariableName = Action->NewBoundVariableName;
+	NewBoundVariablePath = Action->NewBoundVariablePath;
 	return true;
 }
 
@@ -1270,11 +1270,11 @@ bool FRigVMSetPinBoundVariableAction::Undo(URigVMController* InController)
 	{
 		return false;
 	}
-	return InController->BindPinToVariable(PinPath, OldBoundVariableName, false);
+	return InController->BindPinToVariable(PinPath, OldBoundVariablePath, false);
 }
 
 bool FRigVMSetPinBoundVariableAction::Redo(URigVMController* InController)
 {
-	InController->BindPinToVariable(PinPath, NewBoundVariableName, false);
+	InController->BindPinToVariable(PinPath, NewBoundVariablePath, false);
 	return FRigVMBaseAction::Redo(InController);
 }
diff --git a/Engine/Source/Developer/RigVMDeveloper/Private/RigVMModel/RigVMPin.cpp b/Engine/Source/Developer/RigVMDeveloper/Private/RigVMModel/RigVMPin.cpp
--- a/Engine/Source/Developer/RigVMDeveloper/Private/RigVMModel/RigVMPin.cpp
+++ b/Engine/Source/Developer/RigVMDeveloper/Private/RigVMModel/RigVMPin.cpp
@@ -82,7 +82,7 @@ URigVMPin::URigVMPin()
 	, CPPTypeObject(nullptr)
 	, CPPTypeObjectPath(NAME_None)
 	, DefaultValue(FString())
-	, BoundVariableName(NAME_None)
+	, BoundVariablePath()
 {
 }
 
@@ -393,9 +393,16 @@ UObject* URigVMPin::FindObjectFromCPPTypeObjectPath(const FString& InObjectPath)
 	return FindObject<UObject>(ANY_PACKAGE, *InObjectPath);
 }
 
-bool URigVMPin::CanBeBoundToVariable(const FRigVMExternalVariable& InExternalVariable) const
+FString URigVMPin::GetBoundVariableName() const
 {
-	if (!InExternalVariable.IsValid())
+	FString VariableName = BoundVariablePath;
+	BoundVariablePath.Split(TEXT("."), &VariableName, nullptr);
+	return VariableName;
+}
+
+bool URigVMPin::CanBeBoundToVariable(const FRigVMExternalVariable& InExternalVariable, const FRigVMRegisterOffset& InOffset) const
+{
+	if (!InExternalVariable.IsValid(true))
 	{
 		return false;
 	}
@@ -412,14 +419,28 @@ bool URigVMPin::CanBeBoundToVariable(const FRigVMExternalVariable& InExternalVar
 	}
 
 	// check type validity
+	// in the future we need to allow arrays as well
+	if (IsArray() && InOffset.IsValid())
+	{
+		return false;
+	}
 	if (IsArray() != InExternalVariable.bIsArray)
 	{
 		return false;
 	}
 
+	FString ExternalCPPType = InExternalVariable.TypeName.ToString();
+	UObject* ExternalCPPTypeObject = InExternalVariable.TypeObject;
+
+	if (InOffset.IsValid())
+	{
+		ExternalCPPType = InOffset.GetCPPType().ToString();
+		ExternalCPPTypeObject = InOffset.GetScriptStruct();
+	}
+
 	if (GetCPPTypeObject() != nullptr)
 	{
-		if (GetCPPTypeObject() != InExternalVariable.TypeObject)
+		if (GetCPPTypeObject() != ExternalCPPTypeObject)
 		{
 			return false;
 		}
@@ -427,7 +448,7 @@ bool URigVMPin::CanBeBoundToVariable(const FRigVMExternalVariable& InExternalVar
 	else
 	{
 		FString CPPBaseType = IsArray() ? GetArrayElementCppType() : GetCPPType();
-		if (CPPBaseType != InExternalVariable.TypeName.ToString())
+		if (CPPBaseType != ExternalCPPType)
 		{
 			return false;
 		}
@@ -481,6 +502,44 @@ bool URigVMPin::ShowInDetailsPanelOnly() const
 	return false;
 }
 
+// Returns nullptr external variable matching this description
+FRigVMExternalVariable URigVMPin::ToExternalVariable() const
+{
+	FRigVMExternalVariable ExternalVariable;
+
+	FString VariableName = GetBoundVariableName();
+	if (VariableName.IsEmpty())
+	{
+		FString NodeName, PinPath;
+		if (!SplitPinPathAtStart(GetPinPath(), NodeName, VariableName))
+		{
+			return ExternalVariable;
+		}
+
+		VariableName = VariableName.Replace(TEXT("."), TEXT("_"));
+	}
+	ExternalVariable.Name = *VariableName;
+
+	if (CPPType.StartsWith(TEXT("TArray<")))
+	{
+		ExternalVariable.bIsArray = true;
+		ExternalVariable.TypeName = *CPPType.Mid(7, CPPType.Len() - 8);
+		ExternalVariable.TypeObject = CPPTypeObject;
+	}
+	else
+	{
+		ExternalVariable.bIsArray = false;
+		ExternalVariable.TypeName = *CPPType;
+		ExternalVariable.TypeObject = CPPTypeObject;
+	}
+
+	ExternalVariable.bIsPublic = false;
+	ExternalVariable.bIsReadOnly = false;
+	ExternalVariable.Memory = nullptr;
+
+	return ExternalVariable;
+}
+
 void URigVMPin::UpdateCPPTypeObjectIfRequired() const
 {
 	if (CPPTypeObject == nullptr)
@@ -763,15 +822,6 @@ bool URigVMPin::CanLink(URigVMPin* InSourcePin, URigVMPin* InTargetPin, FString*
 		return false;
 	}
 	
-	if (InTargetPin->IsBoundToVariable())
-	{
-		if (OutFailureReason)
-		{
-			*OutFailureReason = TEXT("Target pin is bound to a variable.");
-		}
-		return false;
-	}
-
 	URigVMNode* SourceNode = InSourcePin->GetNode();
 	URigVMNode* TargetNode = InTargetPin->GetNode();
 	if (SourceNode == TargetNode)
diff --git a/Engine/Source/Developer/RigVMDeveloper/Public/RigVMCompiler/RigVMAST.h b/Engine/Source/Developer/RigVMDeveloper/Public/RigVMCompiler/RigVMAST.h
--- a/Engine/Source/Developer/RigVMDeveloper/Public/RigVMCompiler/RigVMAST.h
+++ b/Engine/Source/Developer/RigVMDeveloper/Public/RigVMCompiler/RigVMAST.h
@@ -15,6 +15,7 @@ class FRigVMCallExternExprAST;
 class FRigVMNoOpExprAST;
 class FRigVMVarExprAST;
 class FRigVMLiteralExprAST;
+class FRigVMExternalVarExprAST;
 class FRigVMAssignExprAST;
 class FRigVMCopyExprAST;
 class FRigVMCachedValueExprAST;
@@ -51,6 +52,7 @@ public:
 		NoOp,
 		Var,
 		Literal,
+		ExternalVar,
 		Assign,
 		Copy,
 		CachedValue,
@@ -327,6 +329,17 @@ FORCEINLINE const FRigVMLiteralExprAST* FRigVMExprAST::To() const
 	return (const FRigVMLiteralExprAST*)this;
 }
 
+// specialized cast for type checking
+// for a External Variable / FRigVMExternalVarExprAST expression
+// will raise if types are not compatible
+// @return this expression cast to FRigVMExternalVarExprAST
+template<>
+FORCEINLINE const FRigVMExternalVarExprAST* FRigVMExprAST::To() const
+{
+	ensure(IsA(EType::ExternalVar));
+	return (const FRigVMExternalVarExprAST*)this;
+}
+
 // specialized cast for type checking
 // for a Assign / FRigVMAssignExprAST expression
 // will raise if types are not compatible
@@ -730,6 +743,42 @@ private:
 	friend class FRigVMParserAST;
 };
 
+/*
+ * An abstract syntax tree external variable expression represents the reference 
+ * of unowned / external memory.
+ */
+class RIGVMDEVELOPER_API FRigVMExternalVarExprAST : public FRigVMVarExprAST
+{
+public:
+
+	// virtual destructor
+	virtual ~FRigVMExternalVarExprAST() {}
+
+	// disable copy constructor
+	FRigVMExternalVarExprAST(const FRigVMExternalVarExprAST&) = delete;
+
+	// overload of the type checking mechanism
+	virtual bool IsA(EType InType) const override
+	{
+		if (FRigVMVarExprAST::IsA(InType))
+		{
+			return true;
+		}
+		return InType == EType::ExternalVar;
+	};
+
+protected:
+
+	// default constructor (protected so that only parser can access it)
+	FRigVMExternalVarExprAST(UObject* InSubject = nullptr)
+		: FRigVMVarExprAST(EType::ExternalVar, InSubject)
+	{}
+
+private:
+
+	friend class FRigVMParserAST;
+};
+
 /*
  * An abstract syntax tree assign expression represents the assignment of one
  * expression to another. This can result in referencing memory from a to b or
diff --git a/Engine/Source/Developer/RigVMDeveloper/Public/RigVMCompiler/RigVMCompiler.h b/Engine/Source/Developer/RigVMDeveloper/Public/RigVMCompiler/RigVMCompiler.h
--- a/Engine/Source/Developer/RigVMDeveloper/Public/RigVMCompiler/RigVMCompiler.h
+++ b/Engine/Source/Developer/RigVMDeveloper/Public/RigVMCompiler/RigVMCompiler.h
@@ -124,6 +124,7 @@ private:
 	void TraverseNoOp(const FRigVMNoOpExprAST* InExpr, FRigVMCompilerWorkData& WorkData);
 	void TraverseVar(const FRigVMVarExprAST* InExpr, FRigVMCompilerWorkData& WorkData);
 	void TraverseLiteral(const FRigVMVarExprAST* InExpr, FRigVMCompilerWorkData& WorkData);
+	void TraverseExternalVar(const FRigVMExternalVarExprAST* InExpr, FRigVMCompilerWorkData& WorkData);
 	void TraverseAssign(const FRigVMAssignExprAST* InExpr, FRigVMCompilerWorkData& WorkData);
 	void TraverseCopy(const FRigVMCopyExprAST* InExpr, FRigVMCompilerWorkData& WorkData);
 	void TraverseCachedValue(const FRigVMCachedValueExprAST* InExpr, FRigVMCompilerWorkData& WorkData);
diff --git a/Engine/Source/Developer/RigVMDeveloper/Public/RigVMModel/RigVMController.h b/Engine/Source/Developer/RigVMDeveloper/Public/RigVMModel/RigVMController.h
--- a/Engine/Source/Developer/RigVMDeveloper/Public/RigVMModel/RigVMController.h
+++ b/Engine/Source/Developer/RigVMDeveloper/Public/RigVMModel/RigVMController.h
@@ -411,13 +411,25 @@ public:
 	// Binds a pin to a variable (or removes the binding given NAME_None)
 	// This causes a PinBoundVariableChanged modified event.
 	UFUNCTION(BlueprintCallable, Category = RigVMController)
-	bool BindPinToVariable(const FString& InPinPath, const FName& InNewBoundVariableName, bool bUndo = true);
+	bool BindPinToVariable(const FString& InPinPath, const FString& InNewBoundVariablePath, bool bUndo = true);
 
 	// Removes the binging of a pin to a variable
 	// This causes a PinBoundVariableChanged modified event.
 	UFUNCTION(BlueprintCallable, Category = RigVMController)
 	bool UnbindPinFromVariable(const FString& InPinPath, bool bUndo = true);
 
+	// Turns a variable node into one or more bindings
+	UFUNCTION(BlueprintCallable, Category = RigVMController)
+	bool MakeBindingsFromVariableNode(const FName& InNodeName, bool bUndo = true);
+
+	// Turns a binding to a variable node
+	UFUNCTION(BlueprintCallable, Category = RigVMController)
+	bool MakeVariableNodeFromBinding(const FString& InPinPath, const FVector2D& InNodePosition = FVector2D::ZeroVector, bool bUndo = true);
+
+	// Promotes a pin to a variable
+	UFUNCTION(BlueprintCallable, Category = RigVMController)
+	bool PromotePinToVariable(const FString& InPinPath, bool bCreateVariableNode, const FVector2D& InNodePosition = FVector2D::ZeroVector, bool bUndo = true);
+
 	// Adds a link to the graph.
 	// This causes a LinkAdded modified event.
 	UFUNCTION(BlueprintCallable, Category = RigVMController)
@@ -459,6 +471,11 @@ public:
 
 	FRigVMStructNodeCreatedContext& GetStructNodeCreatedContext() { return StructNodeCreatedContext; }
 
+	// Wires the struct node delegates to the default controller delegates.
+	// this is used only within the Control Rig Editor currently.
+	void SetupDefaultStructNodeDelegates(TDelegate<FName(FRigVMExternalVariable)> InCreateExternalVariableDelegate);
+	void ResetStructNodeDelegates();
+
 private:
 
 	UPROPERTY(BlueprintReadOnly, Category = RigVMController, meta = (ScriptName = "ModifiedEvent", AllowPrivateAccess = "true"))
@@ -484,7 +501,9 @@ private:
 	URigVMPin* InsertArrayPin(URigVMPin* ArrayPin, int32 InIndex, const FString& InDefaultValue, bool bUndo);
 	bool RemovePin(URigVMPin* InPinToRemove, bool bUndo);
 	FProperty* FindPropertyForPin(const FString& InPinPath);
-	bool BindPinToVariable(URigVMPin* InPin, const FName& InNewBoundVariableName, bool bUndo);
+	bool BindPinToVariable(URigVMPin* InPin, const FString& InNewBoundVariablePath, bool bUndo);
+	bool MakeBindingsFromVariableNode(URigVMVariableNode* InNode, bool bUndo);
+	bool PromotePinToVariable(URigVMPin* InPin, bool bCreateVariableNode, const FVector2D& InNodePosition, bool bUndo);
 	bool AddLink(URigVMPin* OutputPin, URigVMPin* InputPin, bool bUndo);
 	bool BreakLink(URigVMPin* OutputPin, URigVMPin* InputPin, bool bUndo);
 	bool BreakAllLinks(URigVMPin* Pin, bool bAsInput, bool bUndo);
diff --git a/Engine/Source/Developer/RigVMDeveloper/Public/RigVMModel/RigVMControllerActions.h b/Engine/Source/Developer/RigVMDeveloper/Public/RigVMModel/RigVMControllerActions.h
--- a/Engine/Source/Developer/RigVMDeveloper/Public/RigVMModel/RigVMControllerActions.h
+++ b/Engine/Source/Developer/RigVMDeveloper/Public/RigVMModel/RigVMControllerActions.h
@@ -1112,7 +1112,7 @@ struct FRigVMSetPinBoundVariableAction : public FRigVMBaseAction
 public:
 
 	FRigVMSetPinBoundVariableAction() {}
-	FRigVMSetPinBoundVariableAction(URigVMPin* InPin, const FName& InNewBoundVariableName);
+	FRigVMSetPinBoundVariableAction(URigVMPin* InPin, const FString& InNewBoundVariablePath);
 	virtual ~FRigVMSetPinBoundVariableAction() {};
 	virtual bool Merge(const FRigVMBaseAction* Other);
 	virtual bool Undo(URigVMController* InController) override;
@@ -1122,8 +1122,8 @@ public:
 	FString PinPath;
 
 	UPROPERTY()
-	FName OldBoundVariableName;
+	FString OldBoundVariablePath;
 
 	UPROPERTY()
-	FName NewBoundVariableName;
+	FString NewBoundVariablePath;
 };
diff --git a/Engine/Source/Developer/RigVMDeveloper/Public/RigVMModel/RigVMPin.h b/Engine/Source/Developer/RigVMDeveloper/Public/RigVMModel/RigVMPin.h
--- a/Engine/Source/Developer/RigVMDeveloper/Public/RigVMModel/RigVMPin.h
+++ b/Engine/Source/Developer/RigVMDeveloper/Public/RigVMModel/RigVMPin.h
@@ -289,13 +289,16 @@ public:
 	const TArray<URigVMInjectionInfo*> GetInjectedNodes() const { return InjectionInfos; }
 
 	// Returns the variable bound to this pin (or NAME_None)
-	const FName& GetBoundVariableName() const { return BoundVariableName; }
+	const FString& GetBoundVariablePath() const { return BoundVariablePath; }
+
+	// Returns the variable bound to this pin (or NAME_None)
+	FString GetBoundVariableName() const;
 
 	// Returns true if this pin is bound to a variable
-	bool IsBoundToVariable() const { return !GetBoundVariableName().IsNone(); }
+	bool IsBoundToVariable() const { return !BoundVariablePath.IsEmpty(); }
 
 	// Returns true if the pin can be bound to a given variable
-	bool CanBeBoundToVariable(const FRigVMExternalVariable& InExternalVariable) const;
+	bool CanBeBoundToVariable(const FRigVMExternalVariable& InExternalVariable, const FRigVMRegisterOffset& InOffset = FRigVMRegisterOffset()) const;
 
 	// helper function to retrieve an object from a path
 	static UObject* FindObjectFromCPPTypeObjectPath(const FString& InObjectPath);
@@ -314,6 +317,9 @@ public:
 	// Returns true if the pin should not show up on a node, but in the details panel
 	bool ShowInDetailsPanelOnly() const;
 
+	// Returns an external variable matching this pin's type
+	FRigVMExternalVariable ToExternalVariable() const;
+
 private:
 
 	void UpdateCPPTypeObjectIfRequired() const;
@@ -362,7 +368,7 @@ private:
 	TArray<URigVMInjectionInfo*> InjectionInfos;
 
 	UPROPERTY()
-	FName BoundVariableName;
+	FString BoundVariablePath;
 
 	friend class URigVMController;
 	friend class URigVMGraph;
diff --git a/Engine/Source/Editor/GraphEditor/Private/SGraphPin.cpp b/Engine/Source/Editor/GraphEditor/Private/SGraphPin.cpp
--- a/Engine/Source/Editor/GraphEditor/Private/SGraphPin.cpp
+++ b/Engine/Source/Editor/GraphEditor/Private/SGraphPin.cpp
@@ -771,6 +771,13 @@ FReply SGraphPin::OnDrop( const FGeometry& MyGeometry, const FDragDropEvent& Dra
 			//Line it up vertically with pin
 			NodeAddPosition.Y += MyGeometry.Size.Y;
 
+			// if the pin widget is nested into another compound
+			if (MyGeometry.Position == FVector2D::ZeroVector)
+			{
+				FVector2D PinOffsetPosition = MyGeometry.AbsolutePosition - NodeWidget->GetTickSpaceGeometry().AbsolutePosition;
+				NodeAddPosition = OwnerNode->GetPosition() + PinOffsetPosition;
+			}
+
 			if(GetDirection() == EEdGraphPinDirection::EGPD_Input)
 			{
 				//left side just offset by fixed amount
diff --git a/Engine/Source/Runtime/RigVM/Private/RigVMCore/RigVM.cpp b/Engine/Source/Runtime/RigVM/Private/RigVMCore/RigVM.cpp
--- a/Engine/Source/Runtime/RigVM/Private/RigVMCore/RigVM.cpp
+++ b/Engine/Source/Runtime/RigVM/Private/RigVMCore/RigVM.cpp
@@ -1699,8 +1699,21 @@ void URigVM::CacheSingleMemoryHandle(const FRigVMOperand& InArg, bool bForExecut
 		ensure(ExternalVariables.IsValidIndex(InArg.GetRegisterIndex()));
 
 		FRigVMExternalVariable& ExternalVariable = ExternalVariables[InArg.GetRegisterIndex()];
+		FRigVMRegisterOffset RegisterOffset = GetWorkMemory().GetRegisterOffsetForOperand(InArg);
 		check(ExternalVariable.Memory);
-		CachedMemoryHandles.Add(ExternalVariable.GetHandle(InArg.GetRegisterOffset()));
+
+		if (RegisterOffset.IsValid())
+		{
+			CachedMemoryHandles.Add(FRigVMMemoryHandle(
+				RegisterOffset.GetData(ExternalVariable.Memory),
+				RegisterOffset.GetElementSize(),
+				FRigVMMemoryHandle::FType::Plain
+			));
+		}
+		else
+		{
+			CachedMemoryHandles.Add(ExternalVariable.GetHandle(InArg.GetRegisterOffset()));
+		}
 		return;
 	}
 
diff --git a/Engine/Source/Runtime/RigVM/Private/RigVMCore/RigVMMemory.cpp b/Engine/Source/Runtime/RigVM/Private/RigVMCore/RigVMMemory.cpp
--- a/Engine/Source/Runtime/RigVM/Private/RigVMCore/RigVMMemory.cpp
+++ b/Engine/Source/Runtime/RigVM/Private/RigVMCore/RigVMMemory.cpp
@@ -1633,9 +1633,7 @@ int32 FRigVMMemoryContainer::GetOrAddRegisterOffset(int32 InRegisterIndex, UScri
 		return INDEX_NONE;
 	}
 
-	ensure(Registers.IsValidIndex(InRegisterIndex));
-
-	if (InElementSize == 0)
+	if (InElementSize == 0 && Registers.IsValidIndex(InRegisterIndex))
 	{
 		InElementSize = (int32)Registers[InRegisterIndex].ElementSize;
 	}
@@ -1810,7 +1808,8 @@ TArray<FString> FRigVMMemoryContainer::GetRegisterValueAsString(const FRigVMOper
 
 const FRigVMRegisterOffset FRigVMMemoryContainer::GetRegisterOffsetForOperand(const FRigVMOperand& InOperand) const
 {
-	ensure(InOperand.GetMemoryType() == MemoryType);
+	ensure(InOperand.GetMemoryType() == MemoryType || 
+		(InOperand.GetMemoryType() == ERigVMMemoryType::External && MemoryType == ERigVMMemoryType::Work));
 
 	int32 RegisterOffsetIndex = InOperand.GetRegisterOffset();
 	if (RegisterOffsets.IsValidIndex(RegisterOffsetIndex))
diff --git a/Engine/Source/Runtime/RigVM/Private/RigVMCore/RigVMStruct.cpp b/Engine/Source/Runtime/RigVM/Private/RigVMCore/RigVMStruct.cpp
--- a/Engine/Source/Runtime/RigVM/Private/RigVMCore/RigVMStruct.cpp
+++ b/Engine/Source/Runtime/RigVM/Private/RigVMCore/RigVMStruct.cpp
@@ -34,11 +34,12 @@ FName FRigVMStructNodeCreatedContext::AddExternalVariable(const FRigVMExternalVa
 	return NAME_None;
 }
 
-bool FRigVMStructNodeCreatedContext::BindPinToExternalVariable(FString InPinPath, FName InVariableName)
+bool FRigVMStructNodeCreatedContext::BindPinToExternalVariable(FString InPinPath, FString InVariablePath)
 {
 	if (BindPinToExternalVariableDelegate.IsBound())
 	{
-		return BindPinToExternalVariableDelegate.Execute(InPinPath, InVariableName);
+		const FString NodePinPath = FString::Printf(TEXT("%s.%s"), *NodeName.ToString(), *InPinPath);
+		return BindPinToExternalVariableDelegate.Execute(NodePinPath, InVariablePath);
 	}
 	return false;
 }
diff --git a/Engine/Source/Runtime/RigVM/Public/RigVMCore/RigVMMemory.h b/Engine/Source/Runtime/RigVM/Public/RigVMCore/RigVMMemory.h
--- a/Engine/Source/Runtime/RigVM/Public/RigVMCore/RigVMMemory.h
+++ b/Engine/Source/Runtime/RigVM/Public/RigVMCore/RigVMMemory.h
@@ -111,6 +111,8 @@ private:
 	
 	UPROPERTY()
 	uint16 RegisterOffset;
+
+	friend class URigVMCompiler;
 };
 
 typedef FRigVMFixedArray<FRigVMOperand> FRigVMOperandArray;
diff --git a/Engine/Source/Runtime/RigVM/Public/RigVMCore/RigVMStruct.h b/Engine/Source/Runtime/RigVM/Public/RigVMCore/RigVMStruct.h
--- a/Engine/Source/Runtime/RigVM/Public/RigVMCore/RigVMStruct.h
+++ b/Engine/Source/Runtime/RigVM/Public/RigVMCore/RigVMStruct.h
@@ -11,7 +11,7 @@
 // delegates used for variable introspection / creation
 DECLARE_DELEGATE_RetVal(TArray<FRigVMExternalVariable>, FRigVMGetExternalVariablesDelegate)
 DECLARE_DELEGATE_RetVal_OneParam(FName, FRigVMCreateExternalVariableDelegate, FRigVMExternalVariable)
-DECLARE_DELEGATE_RetVal_TwoParams(bool, FRigVMBindPinToExternalVariableDelegate, FString, FName)
+DECLARE_DELEGATE_RetVal_TwoParams(bool, FRigVMBindPinToExternalVariableDelegate, FString, FString)
 
 /** Context as of why the node was created */
 enum class ERigVMNodeCreatedReason : uint8
@@ -66,7 +66,7 @@ public:
 	FName AddExternalVariable(const FRigVMExternalVariable& InVariableToCreate);
 
 	/** Binds a pin to an external variable on the created node */
-	bool BindPinToExternalVariable(FString InPinPath, FName InVariableName);
+	bool BindPinToExternalVariable(FString InPinPath, FString InVariablePath);
 
 	/** Returns a variable given a name (or a non-valid variable if not found) */
 	FRigVMExternalVariable FindVariable(FName InVariableName) const;
