UE-93573 Cycle check can fail on rewire

Since a new wire can disconnect an old wire - we need to ignore certain wires when building the map for link validation and cycle checking. With this fix we pass a list of links to ignore (wires to ignore in the graph) to the AST that's going to get build, and then perform cycle checking only with that.

diff --git a/Engine/Source/Developer/RigVMDeveloper/Private/RigVMCompiler/RigVMAST.cpp b/Engine/Source/Developer/RigVMDeveloper/Private/RigVMCompiler/RigVMAST.cpp
--- a/Engine/Source/Developer/RigVMDeveloper/Private/RigVMCompiler/RigVMAST.cpp
+++ b/Engine/Source/Developer/RigVMDeveloper/Private/RigVMCompiler/RigVMAST.cpp
@@ -823,6 +823,7 @@ FRigVMParserAST::FRigVMParserAST(URigVMGraph* InGraph, URigVMController* InContr
 
 	ObsoleteBlock = nullptr;
 	LastCycleCheckExpr = nullptr;
+	LinksToSkip = InSettings.LinksToSkip;
 
 	const TArray<URigVMNode*> Nodes = InGraph->GetNodes();
 	for (URigVMNode* Node : Nodes)
@@ -977,9 +978,14 @@ FRigVMExprAST* FRigVMParserAST::TraverseMutableNode(URigVMNode* InNode, FRigVMEx
 						}
 					}
 
-					TArray<URigVMPin*> TargetPins = SourcePin->GetLinkedTargetPins();
-					for (URigVMPin* TargetPin : TargetPins)
+					const TArray<URigVMLink*>& Links = SourcePin->GetLinks();
+					for(URigVMLink* Link : Links)
 					{
+						if (LinksToSkip.Contains(Link))
+						{
+							continue;
+						}
+						URigVMPin* TargetPin = Link->GetTargetPin();
 						TraverseMutableNode(TargetPin->GetNode(), ParentExpr);
 					}
 				}
@@ -1094,6 +1100,12 @@ FRigVMExprAST* FRigVMParserAST::TraversePin(URigVMPin* InPin, FRigVMExprAST* InP
 	ensure(!SubjectToExpression.Contains(InPin));
 
 	TArray<URigVMLink*> SourceLinks = InPin->GetSourceLinks(true);
+	
+	if (LinksToSkip.Num() > 0)
+	{
+		const TArray<URigVMLink*>& LocalLinksToSkip = LinksToSkip;
+		SourceLinks.RemoveAll([LocalLinksToSkip](URigVMLink* LinkToCheck) { return LocalLinksToSkip.Contains(LinkToCheck); });
+	}
 
 	struct Local
 	{
diff --git a/Engine/Source/Developer/RigVMDeveloper/Private/RigVMModel/RigVMGraph.cpp b/Engine/Source/Developer/RigVMDeveloper/Private/RigVMModel/RigVMGraph.cpp
--- a/Engine/Source/Developer/RigVMDeveloper/Private/RigVMModel/RigVMGraph.cpp
+++ b/Engine/Source/Developer/RigVMDeveloper/Private/RigVMModel/RigVMGraph.cpp
@@ -119,11 +119,12 @@ void URigVMGraph::Notify(ERigVMGraphNotifType InNotifType, UObject* InSubject)
 	ModifiedEvent.Broadcast(InNotifType, this, InSubject);
 }
 
-TSharedPtr<FRigVMParserAST> URigVMGraph::GetDiagnosticsAST(bool bForceRefresh)
+TSharedPtr<FRigVMParserAST> URigVMGraph::GetDiagnosticsAST(bool bForceRefresh, TArray<URigVMLink*> InLinksToSkip)
 {
 	if (DiagnosticsAST == nullptr || bForceRefresh)
 	{
 		FRigVMParserASTSettings Settings = FRigVMParserASTSettings::Fast();
+		Settings.LinksToSkip = InLinksToSkip;
 		DiagnosticsAST = MakeShareable(new FRigVMParserAST(this, nullptr, Settings));
 	}
 	return DiagnosticsAST;
@@ -164,7 +165,12 @@ bool URigVMGraph::IsNameAvailable(const FString& InName)
 
 void URigVMGraph::PrepareCycleChecking(URigVMPin* InPin, bool bAsInput)
 {
-	GetDiagnosticsAST()->PrepareCycleChecking(InPin);
+	TArray<URigVMLink*> LinksToSkip;
+	if (InPin)
+	{
+		LinksToSkip = InPin->GetLinks();
+	}
+	GetDiagnosticsAST(LinksToSkip.Num() > 0, LinksToSkip)->PrepareCycleChecking(InPin);
 }
 
 bool URigVMGraph::CanLink(URigVMPin* InSourcePin, URigVMPin* InTargetPin, FString* OutFailureReason)
diff --git a/Engine/Source/Developer/RigVMDeveloper/Public/RigVMCompiler/RigVMAST.h b/Engine/Source/Developer/RigVMDeveloper/Public/RigVMCompiler/RigVMAST.h
--- a/Engine/Source/Developer/RigVMDeveloper/Public/RigVMCompiler/RigVMAST.h
+++ b/Engine/Source/Developer/RigVMDeveloper/Public/RigVMCompiler/RigVMAST.h
@@ -1103,6 +1103,10 @@ struct RIGVMDEVELOPER_API FRigVMParserASTSettings
 	UPROPERTY(EditAnywhere, Category = "AST")
 	bool bFoldConstantBranches;
 
+	// links to be ignored during the parse
+	UPROPERTY()
+	TArray<URigVMLink*> LinksToSkip;
+
 	// static method to provide fast AST parse settings
 	static FRigVMParserASTSettings Fast()
 	{
@@ -1334,6 +1338,7 @@ private:
 
 	const FRigVMExprAST* LastCycleCheckExpr;
 	TArray<ETraverseRelationShip> CycleCheckFlags;
+	TArray<URigVMLink*> LinksToSkip;
 
 	friend class FRigVMExprAST;
 	friend class URigVMCompiler;
diff --git a/Engine/Source/Developer/RigVMDeveloper/Public/RigVMModel/RigVMGraph.h b/Engine/Source/Developer/RigVMDeveloper/Public/RigVMModel/RigVMGraph.h
--- a/Engine/Source/Developer/RigVMDeveloper/Public/RigVMModel/RigVMGraph.h
+++ b/Engine/Source/Developer/RigVMDeveloper/Public/RigVMModel/RigVMGraph.h
@@ -79,7 +79,7 @@ public:
 	void PrepareCycleChecking(URigVMPin* InPin, bool bAsInput);
 	virtual bool CanLink(URigVMPin* InSourcePin, URigVMPin* InTargetPin, FString* OutFailureReason = nullptr);
 	
-	TSharedPtr<FRigVMParserAST> GetDiagnosticsAST(bool bForceRefresh = false);
+	TSharedPtr<FRigVMParserAST> GetDiagnosticsAST(bool bForceRefresh = false, TArray<URigVMLink*> InLinksToSkip = TArray<URigVMLink*>());
 	TSharedPtr<FRigVMParserAST> GetRuntimeAST(const FRigVMParserASTSettings& InSettings = FRigVMParserASTSettings::Optimized(), bool bForceRefresh = false);
 	void ClearAST(bool bClearDiagnostics = true, bool bClearRuntime = true);
 
