Control Rig: Rename bUndo to bSetupUndoRedo

diff --git a/Engine/Source/Developer/RigVMDeveloper/Private/RigVMModel/RigVMController.cpp b/Engine/Source/Developer/RigVMDeveloper/Private/RigVMModel/RigVMController.cpp
--- a/Engine/Source/Developer/RigVMDeveloper/Private/RigVMModel/RigVMController.cpp
+++ b/Engine/Source/Developer/RigVMDeveloper/Private/RigVMModel/RigVMController.cpp
@@ -159,7 +159,7 @@ void URigVMController::HandleModifiedEvent(ERigVMGraphNotifType InNotifType, URi
 
 #if WITH_EDITOR
 
-URigVMStructNode* URigVMController::AddStructNode(UScriptStruct* InScriptStruct, const FName& InMethodName, const FVector2D& InPosition, const FString& InNodeName, bool bUndo)
+URigVMStructNode* URigVMController::AddStructNode(UScriptStruct* InScriptStruct, const FName& InMethodName, const FVector2D& InPosition, const FString& InNodeName, bool bSetupUndoRedo)
 {
 	if(!IsValidGraph())
 	{
@@ -216,7 +216,7 @@ URigVMStructNode* URigVMController::AddStructNode(UScriptStruct* InScriptStruct,
 	}
 
 	FRigVMAddStructNodeAction Action;
-	if (bUndo)
+	if (bSetupUndoRedo)
 	{
 		Action = FRigVMAddStructNodeAction(Node);
 		Action.Title = FString::Printf(TEXT("Add %s Node"), *Node->GetNodeTitle());
@@ -235,7 +235,7 @@ URigVMStructNode* URigVMController::AddStructNode(UScriptStruct* InScriptStruct,
 		}
 	}
 
-	if (bUndo)
+	if (bSetupUndoRedo)
 	{
 		ActionStack->EndAction(Action);
 	}
@@ -243,7 +243,7 @@ URigVMStructNode* URigVMController::AddStructNode(UScriptStruct* InScriptStruct,
 	return Node;
 }
 
-URigVMStructNode* URigVMController::AddStructNodeFromStructPath(const FString& InScriptStructPath, const FName& InMethodName, const FVector2D& InPosition, const FString& InNodeName, bool bUndo)
+URigVMStructNode* URigVMController::AddStructNodeFromStructPath(const FString& InScriptStructPath, const FName& InMethodName, const FVector2D& InPosition, const FString& InNodeName, bool bSetupUndoRedo)
 {
 	if(!IsValidGraph())
 	{
@@ -257,10 +257,10 @@ URigVMStructNode* URigVMController::AddStructNodeFromStructPath(const FString& I
 		return nullptr;
 	}
 
-	return AddStructNode(ScriptStruct, InMethodName, InPosition, InNodeName, bUndo);
+	return AddStructNode(ScriptStruct, InMethodName, InPosition, InNodeName, bSetupUndoRedo);
 }
 
-URigVMVariableNode* URigVMController::AddVariableNode(const FName& InVariableName, const FString& InCPPType, UObject* InCPPTypeObject, bool bIsGetter, const FString& InDefaultValue, const FVector2D& InPosition, const FString& InNodeName, bool bUndo)
+URigVMVariableNode* URigVMController::AddVariableNode(const FName& InVariableName, const FString& InCPPType, UObject* InCPPTypeObject, bool bIsGetter, const FString& InDefaultValue, const FVector2D& InPosition, const FString& InNodeName, bool bSetupUndoRedo)
 {
 	if (!IsValidGraph())
 	{
@@ -343,7 +343,7 @@ URigVMVariableNode* URigVMController::AddVariableNode(const FName& InVariableNam
 	}
 
 	FRigVMAddVariableNodeAction Action;
-	if (bUndo)
+	if (bSetupUndoRedo)
 	{
 		Action = FRigVMAddVariableNodeAction(Node);
 		Action.Title = FString::Printf(TEXT("Add %s Variable"), *InVariableName.ToString());
@@ -353,7 +353,7 @@ URigVMVariableNode* URigVMController::AddVariableNode(const FName& InVariableNam
 	Notify(ERigVMGraphNotifType::NodeAdded, Node);
 	Notify(ERigVMGraphNotifType::VariableAdded, Node);
 
-	if (bUndo)
+	if (bSetupUndoRedo)
 	{
 		ActionStack->EndAction(Action);
 	}
@@ -361,7 +361,7 @@ URigVMVariableNode* URigVMController::AddVariableNode(const FName& InVariableNam
 	return Node;
 }
 
-URigVMVariableNode* URigVMController::AddVariableNodeFromObjectPath(const FName& InVariableName, const FString& InCPPType, const FString& InCPPTypeObjectPath, bool bIsGetter, const FString& InDefaultValue, const FVector2D& InPosition, const FString& InNodeName, bool bUndo)
+URigVMVariableNode* URigVMController::AddVariableNodeFromObjectPath(const FName& InVariableName, const FString& InCPPType, const FString& InCPPTypeObjectPath, bool bIsGetter, const FString& InDefaultValue, const FVector2D& InPosition, const FString& InNodeName, bool bSetupUndoRedo)
 {
 	if (!IsValidGraph())
 	{
@@ -379,10 +379,10 @@ URigVMVariableNode* URigVMController::AddVariableNodeFromObjectPath(const FName&
 		}
 	}
 
-	return AddVariableNode(InVariableName, InCPPType, CPPTypeObject, bIsGetter, InDefaultValue, InPosition, InNodeName, bUndo);
+	return AddVariableNode(InVariableName, InCPPType, CPPTypeObject, bIsGetter, InDefaultValue, InPosition, InNodeName, bSetupUndoRedo);
 }
 
-void URigVMController::RefreshVariableNode(const FName& InNodeName, const FName& InVariableName, const FString& InCPPType, UObject* InCPPTypeObject, bool bUndo)
+void URigVMController::RefreshVariableNode(const FName& InNodeName, const FName& InVariableName, const FString& InCPPType, UObject* InCPPTypeObject, bool bSetupUndoRedo)
 {
 	if (!IsValidGraph())
 	{
@@ -395,7 +395,7 @@ void URigVMController::RefreshVariableNode(const FName& InNodeName, const FName&
 		{
 			if (VariablePin->Direction == ERigVMPinDirection::Visible)
 			{
-				if (bUndo)
+				if (bSetupUndoRedo)
 				{
 					VariablePin->Modify();
 				}
@@ -405,7 +405,7 @@ void URigVMController::RefreshVariableNode(const FName& InNodeName, const FName&
 
 			if (InVariableName.IsValid() && VariablePin->DefaultValue != InVariableName.ToString())
 			{
-				if (bUndo)
+				if (bSetupUndoRedo)
 				{
 					VariablePin->Modify();
 				}
@@ -420,18 +420,18 @@ void URigVMController::RefreshVariableNode(const FName& InNodeName, const FName&
 				{
 					if (ValuePin->CPPType != InCPPType)
 					{
-						if (bUndo)
+						if (bSetupUndoRedo)
 						{
 							ValuePin->Modify();
 						}
 
-						BreakAllLinks(ValuePin, ValuePin->GetDirection() == ERigVMPinDirection::Input, bUndo);
-						BreakAllLinksRecursive(ValuePin, ValuePin->GetDirection() == ERigVMPinDirection::Input, false, bUndo);
+						BreakAllLinks(ValuePin, ValuePin->GetDirection() == ERigVMPinDirection::Input, bSetupUndoRedo);
+						BreakAllLinksRecursive(ValuePin, ValuePin->GetDirection() == ERigVMPinDirection::Input, false, bSetupUndoRedo);
 
 						// if this is an unsupported datatype...
 						if (InCPPType == FName(NAME_None).ToString())
 						{
-							RemoveNode(VariableNode, bUndo);
+							RemoveNode(VariableNode, bSetupUndoRedo);
 							return;
 						}
 
@@ -460,7 +460,7 @@ void URigVMController::RefreshVariableNode(const FName& InNodeName, const FName&
 	}
 }
 
-void URigVMController::OnExternalVariableRemoved(const FName& InVarName, bool bUndo)
+void URigVMController::OnExternalVariableRemoved(const FName& InVarName, bool bSetupUndoRedo)
 {
 	if (!IsValidGraph())
 	{
@@ -474,7 +474,7 @@ void URigVMController::OnExternalVariableRemoved(const FName& InVarName, bool bU
 
 	FString VarNameStr = InVarName.ToString();
 
-	if (bUndo)
+	if (bSetupUndoRedo)
 	{
 		OpenUndoBracket(TEXT("Remove Variable Nodes"));
 	}
@@ -488,7 +488,7 @@ void URigVMController::OnExternalVariableRemoved(const FName& InVarName, bool bU
 			{
 				if (VariablePin->GetDefaultValue() == VarNameStr)
 				{
-					RemoveNode(Node, bUndo, true);
+					RemoveNode(Node, bSetupUndoRedo, true);
 					continue;
 				}
 			}
@@ -504,13 +504,13 @@ void URigVMController::OnExternalVariableRemoved(const FName& InVarName, bool bU
 		}
 	}
 
-	if (bUndo)
+	if (bSetupUndoRedo)
 	{
 		CloseUndoBracket();
 	}
 }
 
-void URigVMController::OnExternalVariableRenamed(const FName& InOldVarName, const FName& InNewVarName, bool bUndo)
+void URigVMController::OnExternalVariableRenamed(const FName& InOldVarName, const FName& InNewVarName, bool bSetupUndoRedo)
 {
 	if (!IsValidGraph())
 	{
@@ -524,7 +524,7 @@ void URigVMController::OnExternalVariableRenamed(const FName& InOldVarName, cons
 
 	FString VarNameStr = InOldVarName.ToString();
 
-	if (bUndo)
+	if (bSetupUndoRedo)
 	{
 		OpenUndoBracket(TEXT("Rename Variable Nodes"));
 	}
@@ -538,7 +538,7 @@ void URigVMController::OnExternalVariableRenamed(const FName& InOldVarName, cons
 			{
 				if (VariablePin->GetDefaultValue() == VarNameStr)
 				{
-					RefreshVariableNode(Node->GetFName(), InNewVarName, FString(), nullptr, bUndo);
+					RefreshVariableNode(Node->GetFName(), InNewVarName, FString(), nullptr, bSetupUndoRedo);
 					continue;
 				}
 			}
@@ -556,13 +556,13 @@ void URigVMController::OnExternalVariableRenamed(const FName& InOldVarName, cons
 		}
 	}
 
-	if (bUndo)
+	if (bSetupUndoRedo)
 	{
 		CloseUndoBracket();
 	}
 }
 
-void URigVMController::OnExternalVariableTypeChanged(const FName& InVarName, const FString& InCPPType, UObject* InCPPTypeObject, bool bUndo)
+void URigVMController::OnExternalVariableTypeChanged(const FName& InVarName, const FString& InCPPType, UObject* InCPPTypeObject, bool bSetupUndoRedo)
 {
 	if (!IsValidGraph())
 	{
@@ -576,7 +576,7 @@ void URigVMController::OnExternalVariableTypeChanged(const FName& InVarName, con
 
 	FString VarNameStr = InVarName.ToString();
 
-	if (bUndo)
+	if (bSetupUndoRedo)
 	{
 		OpenUndoBracket(TEXT("Change Variable Nodes Type"));
 	}
@@ -590,7 +590,7 @@ void URigVMController::OnExternalVariableTypeChanged(const FName& InVarName, con
 			{
 				if (VariablePin->GetDefaultValue() == VarNameStr)
 				{
-					RefreshVariableNode(Node->GetFName(), InVarName, InCPPType, InCPPTypeObject, bUndo);
+					RefreshVariableNode(Node->GetFName(), InVarName, InCPPType, InCPPTypeObject, bSetupUndoRedo);
 					continue;
 				}
 			}
@@ -609,13 +609,13 @@ void URigVMController::OnExternalVariableTypeChanged(const FName& InVarName, con
 		}
 	}
 
-	if (bUndo)
+	if (bSetupUndoRedo)
 	{
 		CloseUndoBracket();
 	}
 }
 
-URigVMVariableNode* URigVMController::ReplaceParameterNodeWithVariable(const FName& InNodeName, const FName& InVariableName, const FString& InCPPType, UObject* InCPPTypeObject, bool bUndo)
+URigVMVariableNode* URigVMController::ReplaceParameterNodeWithVariable(const FName& InNodeName, const FName& InVariableName, const FString& InCPPType, UObject* InCPPTypeObject, bool bSetupUndoRedo)
 {
 	if (!IsValidGraph())
 	{
@@ -637,7 +637,7 @@ URigVMVariableNode* URigVMController::ReplaceParameterNodeWithVariable(const FNa
 			ParameterValuePin->GetDefaultValue(),
 			ParameterNode->GetPosition(),
 			FString(),
-			bUndo);
+			bSetupUndoRedo);
 
 		if (VariableNode)
 		{
@@ -647,10 +647,10 @@ URigVMVariableNode* URigVMController::ReplaceParameterNodeWithVariable(const FNa
 				ParameterValuePin,
 				VariableValuePin,
 				ParameterValuePin->GetDirection() == ERigVMPinDirection::Input,
-				bUndo
+				bSetupUndoRedo
 			);
 
-			RemoveNode(ParameterNode, bUndo, true);
+			RemoveNode(ParameterNode, bSetupUndoRedo, true);
 
 			return VariableNode;
 		}
@@ -659,7 +659,7 @@ URigVMVariableNode* URigVMController::ReplaceParameterNodeWithVariable(const FNa
 	return nullptr;
 }
 
-URigVMParameterNode* URigVMController::AddParameterNode(const FName& InParameterName, const FString& InCPPType, UObject* InCPPTypeObject, bool bIsInput, const FString& InDefaultValue, const FVector2D& InPosition, const FString& InNodeName, bool bUndo)
+URigVMParameterNode* URigVMController::AddParameterNode(const FName& InParameterName, const FString& InCPPType, UObject* InCPPTypeObject, bool bIsInput, const FString& InDefaultValue, const FVector2D& InPosition, const FString& InNodeName, bool bSetupUndoRedo)
 {
 	if (!IsValidGraph())
 	{
@@ -797,7 +797,7 @@ URigVMParameterNode* URigVMController::AddParameterNode(const FName& InParameter
 	}
 
 	FRigVMAddParameterNodeAction Action;
-	if (bUndo)
+	if (bSetupUndoRedo)
 	{
 		Action = FRigVMAddParameterNodeAction(Node);
 		Action.Title = FString::Printf(TEXT("Add %s Parameter"), *InParameterName.ToString());
@@ -807,7 +807,7 @@ URigVMParameterNode* URigVMController::AddParameterNode(const FName& InParameter
 	Notify(ERigVMGraphNotifType::NodeAdded, Node);
 	Notify(ERigVMGraphNotifType::ParameterAdded, Node);
 
-	if (bUndo)
+	if (bSetupUndoRedo)
 	{
 		ActionStack->EndAction(Action);
 	}
@@ -815,7 +815,7 @@ URigVMParameterNode* URigVMController::AddParameterNode(const FName& InParameter
 	return Node;
 }
 
-URigVMParameterNode* URigVMController::AddParameterNodeFromObjectPath(const FName& InParameterName, const FString& InCPPType, const FString& InCPPTypeObjectPath, bool bIsInput, const FString& InDefaultValue, const FVector2D& InPosition, const FString& InNodeName, bool bUndo)
+URigVMParameterNode* URigVMController::AddParameterNodeFromObjectPath(const FName& InParameterName, const FString& InCPPType, const FString& InCPPTypeObjectPath, bool bIsInput, const FString& InDefaultValue, const FVector2D& InPosition, const FString& InNodeName, bool bSetupUndoRedo)
 {
 	if (!IsValidGraph())
 	{
@@ -833,10 +833,10 @@ URigVMParameterNode* URigVMController::AddParameterNodeFromObjectPath(const FNam
 		}
 	}
 
-	return AddParameterNode(InParameterName, InCPPType, CPPTypeObject, bIsInput, InDefaultValue, InPosition, InNodeName, bUndo);
+	return AddParameterNode(InParameterName, InCPPType, CPPTypeObject, bIsInput, InDefaultValue, InPosition, InNodeName, bSetupUndoRedo);
 }
 
-URigVMCommentNode* URigVMController::AddCommentNode(const FString& InCommentText, const FVector2D& InPosition, const FVector2D& InSize, const FLinearColor& InColor, const FString& InNodeName, bool bUndo)
+URigVMCommentNode* URigVMController::AddCommentNode(const FString& InCommentText, const FVector2D& InPosition, const FVector2D& InSize, const FLinearColor& InColor, const FString& InNodeName, bool bSetupUndoRedo)
 {
 	if (!IsValidGraph())
 	{
@@ -857,7 +857,7 @@ URigVMCommentNode* URigVMController::AddCommentNode(const FString& InCommentText
 	}
 
 	FRigVMAddCommentNodeAction Action;
-	if (bUndo)
+	if (bSetupUndoRedo)
 	{
 		Action = FRigVMAddCommentNodeAction(Node);
 		Action.Title = FString::Printf(TEXT("Add Comment"));
@@ -866,7 +866,7 @@ URigVMCommentNode* URigVMController::AddCommentNode(const FString& InCommentText
 
 	Notify(ERigVMGraphNotifType::NodeAdded, Node);
 
-	if (bUndo)
+	if (bSetupUndoRedo)
 	{
 		ActionStack->EndAction(Action);
 	}
@@ -874,7 +874,7 @@ URigVMCommentNode* URigVMController::AddCommentNode(const FString& InCommentText
 	return Node;
 }
 
-URigVMRerouteNode* URigVMController::AddRerouteNodeOnLink(URigVMLink* InLink, bool bShowAsFullNode, const FVector2D& InPosition, const FString& InNodeName, bool bUndo)
+URigVMRerouteNode* URigVMController::AddRerouteNodeOnLink(URigVMLink* InLink, bool bShowAsFullNode, const FVector2D& InPosition, const FString& InNodeName, bool bSetupUndoRedo)
 {
 	if(!IsValidLinkForGraph(InLink))
 	{
@@ -887,16 +887,16 @@ URigVMRerouteNode* URigVMController::AddRerouteNodeOnLink(URigVMLink* InLink, bo
 	TGuardValue<bool> GuardCompactness(bIgnoreRerouteCompactnessChanges, true);
 
 	FRigVMBaseAction Action;
-	if (bUndo)
+	if (bSetupUndoRedo)
 	{
 		Action.Title = FString::Printf(TEXT("Add Reroute"));
 		ActionStack->BeginAction(Action);
 	}
 
-	URigVMRerouteNode* Node = AddRerouteNodeOnPin(TargetPin->GetPinPath(), true, bShowAsFullNode, InPosition, InNodeName, bUndo);
+	URigVMRerouteNode* Node = AddRerouteNodeOnPin(TargetPin->GetPinPath(), true, bShowAsFullNode, InPosition, InNodeName, bSetupUndoRedo);
 	if (Node == nullptr)
 	{
-		if (bUndo)
+		if (bSetupUndoRedo)
 		{
 			ActionStack->CancelAction(Action);
 		}
@@ -904,9 +904,9 @@ URigVMRerouteNode* URigVMController::AddRerouteNodeOnLink(URigVMLink* InLink, bo
 	}
 
 	URigVMPin* ValuePin = Node->Pins[0];
-	AddLink(SourcePin, ValuePin, bUndo);
+	AddLink(SourcePin, ValuePin, bSetupUndoRedo);
 
-	if (bUndo)
+	if (bSetupUndoRedo)
 	{
 		ActionStack->EndAction(Action);
 	}
@@ -914,7 +914,7 @@ URigVMRerouteNode* URigVMController::AddRerouteNodeOnLink(URigVMLink* InLink, bo
 	return Node;
 }
 
-URigVMRerouteNode* URigVMController::AddRerouteNodeOnLinkPath(const FString& InLinkPinPathRepresentation, bool bShowAsFullNode, const FVector2D& InPosition, const FString& InNodeName, bool bUndo)
+URigVMRerouteNode* URigVMController::AddRerouteNodeOnLinkPath(const FString& InLinkPinPathRepresentation, bool bShowAsFullNode, const FVector2D& InPosition, const FString& InNodeName, bool bSetupUndoRedo)
 {
 	if (!IsValidGraph())
 	{
@@ -922,10 +922,10 @@ URigVMRerouteNode* URigVMController::AddRerouteNodeOnLinkPath(const FString& InL
 	}
 
 	URigVMLink* Link = Graph->FindLink(InLinkPinPathRepresentation);
-	return AddRerouteNodeOnLink(Link, bShowAsFullNode, InPosition, InNodeName, bUndo);
+	return AddRerouteNodeOnLink(Link, bShowAsFullNode, InPosition, InNodeName, bSetupUndoRedo);
 }
 
-URigVMRerouteNode* URigVMController::AddRerouteNodeOnPin(const FString& InPinPath, bool bAsInput, bool bShowAsFullNode, const FVector2D& InPosition, const FString& InNodeName, bool bUndo)
+URigVMRerouteNode* URigVMController::AddRerouteNodeOnPin(const FString& InPinPath, bool bAsInput, bool bShowAsFullNode, const FVector2D& InPosition, const FString& InNodeName, bool bSetupUndoRedo)
 {
 	if (!IsValidGraph())
 	{
@@ -941,13 +941,13 @@ URigVMRerouteNode* URigVMController::AddRerouteNodeOnPin(const FString& InPinPat
 	TGuardValue<bool> GuardCompactness(bIgnoreRerouteCompactnessChanges, true);
 
 	FRigVMBaseAction Action;
-	if (bUndo)
+	if (bSetupUndoRedo)
 	{
 		Action.Title = FString::Printf(TEXT("Add Reroute"));
 		ActionStack->BeginAction(Action);
 	}
 
-	BreakAllLinks(Pin, bAsInput, bUndo);
+	BreakAllLinks(Pin, bAsInput, bSetupUndoRedo);
 
 	FString Name = GetValidNodeName(InNodeName.IsEmpty() ? FString(TEXT("RerouteNode")) : InNodeName);
 	URigVMRerouteNode* Node = NewObject<URigVMRerouteNode>(Graph, *Name);
@@ -984,7 +984,7 @@ URigVMRerouteNode* URigVMController::AddRerouteNodeOnPin(const FString& InPinPat
 		Graph->MarkPackageDirty();
 	}
 
-	if (bUndo)
+	if (bSetupUndoRedo)
 	{
 		ActionStack->AddAction(FRigVMAddRerouteNodeAction(Node));
 	}
@@ -993,14 +993,14 @@ URigVMRerouteNode* URigVMController::AddRerouteNodeOnPin(const FString& InPinPat
 
 	if (bAsInput)
 	{
-		AddLink(ValuePin, Pin, bUndo);
+		AddLink(ValuePin, Pin, bSetupUndoRedo);
 	}
 	else
 	{
-		AddLink(Pin, ValuePin, bUndo);
+		AddLink(Pin, ValuePin, bSetupUndoRedo);
 	}
 
-	if (bUndo)
+	if (bSetupUndoRedo)
 	{
 		ActionStack->EndAction(Action);
 	}
@@ -1008,7 +1008,7 @@ URigVMRerouteNode* URigVMController::AddRerouteNodeOnPin(const FString& InPinPat
 	return Node;
 }
 
-URigVMInjectionInfo* URigVMController::AddInjectedNode(const FString& InPinPath, bool bAsInput, UScriptStruct* InScriptStruct, const FName& InMethodName, const FName& InInputPinName, const FName& InOutputPinName, const FString& InNodeName, bool bUndo)
+URigVMInjectionInfo* URigVMController::AddInjectedNode(const FString& InPinPath, bool bAsInput, UScriptStruct* InScriptStruct, const FName& InMethodName, const FName& InInputPinName, const FName& InOutputPinName, const FString& InNodeName, bool bSetupUndoRedo)
 {
 	if (!IsValidGraph())
 	{
@@ -1074,7 +1074,7 @@ URigVMInjectionInfo* URigVMController::AddInjectedNode(const FString& InPinPath,
 	}
 
 	FRigVMBaseAction Action;
-	if (bUndo)
+	if (bSetupUndoRedo)
 	{
 		Action.Title = FString::Printf(TEXT("Add Injected Node"));
 		ActionStack->BeginAction(Action);
@@ -1087,7 +1087,7 @@ URigVMInjectionInfo* URigVMController::AddInjectedNode(const FString& InPinPath,
 	}
 	if (StructNode == nullptr)
 	{
-		if (bUndo)
+		if (bSetupUndoRedo)
 		{
 			ActionStack->CancelAction(Action);
 		}
@@ -1097,7 +1097,7 @@ URigVMInjectionInfo* URigVMController::AddInjectedNode(const FString& InPinPath,
 	{
 		ReportErrorf(TEXT("Injected node %s is mutable."), *InScriptStruct->GetName());
 		RemoveNode(StructNode, false);
-		if (bUndo)
+		if (bSetupUndoRedo)
 		{
 			ActionStack->CancelAction(Action);
 		}
@@ -1114,7 +1114,7 @@ URigVMInjectionInfo* URigVMController::AddInjectedNode(const FString& InPinPath,
 	{
 		ReportErrorf(TEXT("Injected node %s is using incompatible input and output pins."), *InScriptStruct->GetName());
 		RemoveNode(StructNode, false);
-		if (bUndo)
+		if (bSetupUndoRedo)
 		{
 			ActionStack->CancelAction(Action);
 		}
@@ -1126,7 +1126,7 @@ URigVMInjectionInfo* URigVMController::AddInjectedNode(const FString& InPinPath,
 	{
 		ReportErrorf(TEXT("Injected node %s is using incompatible pin."), *InScriptStruct->GetName());
 		RemoveNode(StructNode, false);
-		if (bUndo)
+		if (bSetupUndoRedo)
 		{
 			ActionStack->CancelAction(Action);
 		}
@@ -1143,7 +1143,7 @@ URigVMInjectionInfo* URigVMController::AddInjectedNode(const FString& InPinPath,
 	InjectionInfo->InputPin = InputPin;
 	InjectionInfo->OutputPin = OutputPin;
 
-	if (bUndo)
+	if (bSetupUndoRedo)
 	{
 		ActionStack->AddAction(FRigVMAddInjectedNodeAction(InjectionInfo));
 	}
@@ -1186,7 +1186,7 @@ URigVMInjectionInfo* URigVMController::AddInjectedNode(const FString& InPinPath,
 		}
 	}
 
-	if (bUndo)
+	if (bSetupUndoRedo)
 	{
 		ActionStack->EndAction(Action);
 	}
@@ -1195,7 +1195,7 @@ URigVMInjectionInfo* URigVMController::AddInjectedNode(const FString& InPinPath,
 
 }
 
-URigVMInjectionInfo* URigVMController::AddInjectedNodeFromStructPath(const FString& InPinPath, bool bAsInput, const FString& InScriptStructPath, const FName& InMethodName, const FName& InInputPinName, const FName& InOutputPinName, const FString& InNodeName, bool bUndo)
+URigVMInjectionInfo* URigVMController::AddInjectedNodeFromStructPath(const FString& InPinPath, bool bAsInput, const FString& InScriptStructPath, const FName& InMethodName, const FName& InInputPinName, const FName& InOutputPinName, const FString& InNodeName, bool bSetupUndoRedo)
 {
 	if (!IsValidGraph())
 	{
@@ -1209,10 +1209,10 @@ URigVMInjectionInfo* URigVMController::AddInjectedNodeFromStructPath(const FStri
 		return nullptr;
 	}
 
-	return AddInjectedNode(InPinPath, bAsInput, ScriptStruct, InMethodName, InInputPinName, InOutputPinName, InNodeName, bUndo);
+	return AddInjectedNode(InPinPath, bAsInput, ScriptStruct, InMethodName, InInputPinName, InOutputPinName, InNodeName, bSetupUndoRedo);
 }
 
-URigVMNode* URigVMController::EjectNodeFromPin(const FString& InPinPath, bool bUndo)
+URigVMNode* URigVMController::EjectNodeFromPin(const FString& InPinPath, bool bSetupUndoRedo)
 {
 	if (!IsValidGraph())
 	{
@@ -1254,7 +1254,7 @@ URigVMNode* URigVMController::EjectNodeFromPin(const FString& InPinPath, bool bU
 	}
 
 	FRigVMBaseAction Action;
-	if (bUndo)
+	if (bSetupUndoRedo)
 	{
 		Action.Title = FString::Printf(TEXT("Eject Node"));
 		ActionStack->BeginAction(Action);
@@ -1270,7 +1270,7 @@ URigVMNode* URigVMController::EjectNodeFromPin(const FString& InPinPath, bool bU
 		Position -= FVector2D(250.f, 0.f);
 	}
 
-	URigVMNode* EjectedNode = AddStructNode(ScriptStruct, MethodName, Position, FString(), bUndo);
+	URigVMNode* EjectedNode = AddStructNode(ScriptStruct, MethodName, Position, FString(), bSetupUndoRedo);
 
 	for (const TPair<FName, FString>& Pair : DefaultValues)
 	{
@@ -1280,7 +1280,7 @@ URigVMNode* URigVMController::EjectNodeFromPin(const FString& InPinPath, bool bU
 		}
 		if (URigVMPin* PinOnNode = EjectedNode->FindPin(Pair.Key.ToString()))
 		{
-			SetPinDefaultValue(PinOnNode, Pair.Value, true, bUndo, false);
+			SetPinDefaultValue(PinOnNode, Pair.Value, true, bSetupUndoRedo, false);
 		}
 	}
 
@@ -1292,7 +1292,7 @@ URigVMNode* URigVMController::EjectNodeFromPin(const FString& InPinPath, bool bU
 		PreviousLink->SourcePin = PreviousLink->TargetPin = nullptr;
 	}
 
-	RemoveNode(Injection->StructNode, bUndo);
+	RemoveNode(Injection->StructNode, bSetupUndoRedo);
 
 	FString OldNodeNamePrefix = StructNodeName.ToString() + TEXT(".");
 	FString NewNodeNamePrefix = EjectedNode->GetName() + TEXT(".");
@@ -1312,14 +1312,14 @@ URigVMNode* URigVMController::EjectNodeFromPin(const FString& InPinPath, bool bU
 
 		URigVMPin* SourcePin = Graph->FindPin(SourcePinPath);
 		URigVMPin* TargetPin = Graph->FindPin(TargetPinPath);
-		AddLink(SourcePin, TargetPin, bUndo);
+		AddLink(SourcePin, TargetPin, bSetupUndoRedo);
 	}
 
 	TArray<FName> NodeNamesToSelect;
 	NodeNamesToSelect.Add(EjectedNode->GetFName());
-	SetNodeSelection(NodeNamesToSelect, bUndo);
+	SetNodeSelection(NodeNamesToSelect, bSetupUndoRedo);
 
-	if (bUndo)
+	if (bSetupUndoRedo)
 	{
 		ActionStack->EndAction(Action);
 	}
@@ -1521,7 +1521,7 @@ bool URigVMController::CanImportNodesFromText(const FString& InText)
 	return Factory.CanCreateObjectsFromText(InText);
 }
 
-TArray<FName> URigVMController::ImportNodesFromText(const FString& InText, bool bUndo)
+TArray<FName> URigVMController::ImportNodesFromText(const FString& InText, bool bSetupUndoRedo)
 {
 	TArray<FName> NodeNames;
 	if (!IsValidGraph())
@@ -1537,13 +1537,13 @@ TArray<FName> URigVMController::ImportNodesFromText(const FString& InText, bool
 		return NodeNames;
 	}
 
-	if (bUndo)
+	if (bSetupUndoRedo)
 	{
 		OpenUndoBracket(TEXT("Importing Nodes from Text"));
 	}
 
 	FRigVMInverseAction AddNodesAction;
-	if (bUndo)
+	if (bSetupUndoRedo)
 	{
 		ActionStack->BeginAction(AddNodesAction);
 	}
@@ -1553,7 +1553,7 @@ TArray<FName> URigVMController::ImportNodesFromText(const FString& InText, bool
 	{
 		Graph->Nodes.Add(CreatedNode);
 
-		if (bUndo)
+		if (bSetupUndoRedo)
 		{
 			ActionStack->AddAction(FRigVMRemoveNodeAction(CreatedNode));
 		}
@@ -1576,7 +1576,7 @@ TArray<FName> URigVMController::ImportNodesFromText(const FString& InText, bool
 		NodeNames.Add(CreatedNode->GetFName());
 	}
 
-	if (bUndo)
+	if (bSetupUndoRedo)
 	{
 		ActionStack->EndAction(AddNodesAction);
 	}
@@ -1584,7 +1584,7 @@ TArray<FName> URigVMController::ImportNodesFromText(const FString& InText, bool
 	if (Factory.CreatedLinks.Num() > 0)
 	{
 		FRigVMBaseAction AddLinksAction;
-		if (bUndo)
+		if (bSetupUndoRedo)
 		{
 			ActionStack->BeginAction(AddLinksAction);
 		}
@@ -1609,7 +1609,7 @@ TArray<FName> URigVMController::ImportNodesFromText(const FString& InText, bool
 						SourcePin->Links.Add(CreatedLink);
 						TargetPin->Links.Add(CreatedLink);
 
-						if (bUndo)
+						if (bSetupUndoRedo)
 						{
 							ActionStack->AddAction(FRigVMAddLinkAction(SourcePin, TargetPin));
 						}
@@ -1623,13 +1623,13 @@ TArray<FName> URigVMController::ImportNodesFromText(const FString& InText, bool
 			DestroyObject(CreatedLink);
 		}
 
-		if (bUndo)
+		if (bSetupUndoRedo)
 		{
 			ActionStack->EndAction(AddLinksAction);
 		}
 	}
 
-	if (bUndo)
+	if (bSetupUndoRedo)
 	{
 		CloseUndoBracket();
 	}
@@ -1652,7 +1652,7 @@ FName URigVMController::GetUniqueName(const FName& InName, TFunction<bool(const
 
 #endif
 
-bool URigVMController::RemoveNode(URigVMNode* InNode, bool bUndo, bool bRecursive)
+bool URigVMController::RemoveNode(URigVMNode* InNode, bool bSetupUndoRedo, bool bRecursive)
 {
 	if (!IsValidNodeForGraph(InNode))
 	{
@@ -1662,7 +1662,7 @@ bool URigVMController::RemoveNode(URigVMNode* InNode, bool bUndo, bool bRecursiv
 	TGuardValue<bool> GuardCompactness(bIgnoreRerouteCompactnessChanges, true);
 
 	FRigVMBaseAction Action;
-	if (bUndo)
+	if (bSetupUndoRedo)
 	{
 		Action = FRigVMBaseAction();
 		Action.Title = FString::Printf(TEXT("Remove %s Node"), *InNode->GetNodeTitle());
@@ -1688,26 +1688,26 @@ bool URigVMController::RemoveNode(URigVMNode* InNode, bool bUndo, bool bRecursiv
 		}
 	}
 
-	if (bUndo || bRecursive)
+	if (bSetupUndoRedo || bRecursive)
 	{
-		SelectNode(InNode, false, bUndo);
+		SelectNode(InNode, false, bSetupUndoRedo);
 
 		for (URigVMPin* Pin : InNode->Pins)
 		{
 			TArray<URigVMInjectionInfo*> InjectedNodes = Pin->GetInjectedNodes();
 			for (URigVMInjectionInfo* InjectedNode : InjectedNodes)
 			{
-				RemoveNode(InjectedNode->StructNode, bUndo);
+				RemoveNode(InjectedNode->StructNode, bSetupUndoRedo);
 			}
 
-			BreakAllLinks(Pin, true, bUndo);
-			BreakAllLinks(Pin, false, bUndo);
-			BreakAllLinksRecursive(Pin, true, false, bUndo);
-			BreakAllLinksRecursive(Pin, false, false, bUndo);
+			BreakAllLinks(Pin, true, bSetupUndoRedo);
+			BreakAllLinks(Pin, false, bSetupUndoRedo);
+			BreakAllLinksRecursive(Pin, true, false, bSetupUndoRedo);
+			BreakAllLinksRecursive(Pin, false, false, bSetupUndoRedo);
 		}
 	}
 
-	if (bUndo)
+	if (bSetupUndoRedo)
 	{
 		ActionStack->AddAction(FRigVMRemoveNodeAction(InNode));
 	}
@@ -1736,7 +1736,7 @@ bool URigVMController::RemoveNode(URigVMNode* InNode, bool bUndo, bool bRecursiv
 
 	DestroyObject(InNode);
 
-	if (bUndo)
+	if (bSetupUndoRedo)
 	{
 		ActionStack->EndAction(Action);
 	}
@@ -1744,16 +1744,16 @@ bool URigVMController::RemoveNode(URigVMNode* InNode, bool bUndo, bool bRecursiv
 	return true;
 }
 
-bool URigVMController::RemoveNodeByName(const FName& InNodeName, bool bUndo, bool bRecursive)
+bool URigVMController::RemoveNodeByName(const FName& InNodeName, bool bSetupUndoRedo, bool bRecursive)
 {
 	if (!IsValidGraph())
 	{
 		return false;
 	}
-	return RemoveNode(Graph->FindNodeByName(InNodeName), bUndo, bRecursive);
+	return RemoveNode(Graph->FindNodeByName(InNodeName), bSetupUndoRedo, bRecursive);
 }
 
-bool URigVMController::SelectNode(URigVMNode* InNode, bool bSelect, bool bUndo)
+bool URigVMController::SelectNode(URigVMNode* InNode, bool bSelect, bool bSetupUndoRedo)
 {
 	if (!IsValidNodeForGraph(InNode))
 	{
@@ -1775,29 +1775,29 @@ bool URigVMController::SelectNode(URigVMNode* InNode, bool bSelect, bool bUndo)
 		NewSelection.Remove(InNode->GetFName());
 	}
 
-	return SetNodeSelection(NewSelection, bUndo);
+	return SetNodeSelection(NewSelection, bSetupUndoRedo);
 }
 
-bool URigVMController::SelectNodeByName(const FName& InNodeName, bool bSelect, bool bUndo)
+bool URigVMController::SelectNodeByName(const FName& InNodeName, bool bSelect, bool bSetupUndoRedo)
 {
 	if (!IsValidGraph())
 	{
 		return false;
 	}
-	return SelectNode(Graph->FindNodeByName(InNodeName), bSelect, bUndo);
+	return SelectNode(Graph->FindNodeByName(InNodeName), bSelect, bSetupUndoRedo);
 }
 
-bool URigVMController::ClearNodeSelection(bool bUndo)
+bool URigVMController::ClearNodeSelection(bool bSetupUndoRedo)
 {
 	if (!IsValidGraph())
 	{
 		return false;
 	}
 
-	return SetNodeSelection(TArray<FName>(), bUndo);
+	return SetNodeSelection(TArray<FName>(), bSetupUndoRedo);
 }
 
-bool URigVMController::SetNodeSelection(const TArray<FName>& InNodeNames, bool bUndo)
+bool URigVMController::SetNodeSelection(const TArray<FName>& InNodeNames, bool bSetupUndoRedo)
 {
 	if (!IsValidGraph())
 	{
@@ -1805,7 +1805,7 @@ bool URigVMController::SetNodeSelection(const TArray<FName>& InNodeNames, bool b
 	}
 
 	FRigVMSetNodeSelectionAction Action;
-	if (bUndo)
+	if (bSetupUndoRedo)
 	{
 		Action = FRigVMSetNodeSelectionAction(Graph, InNodeNames);
 		ActionStack->BeginAction(Action);
@@ -1840,7 +1840,7 @@ bool URigVMController::SetNodeSelection(const TArray<FName>& InNodeNames, bool b
 		}
 	}
 
-	if (bUndo)
+	if (bSetupUndoRedo)
 	{
 		if (bSelectedSomething)
 		{
@@ -1876,7 +1876,7 @@ bool URigVMController::SetNodeSelection(const TArray<FName>& InNodeNames, bool b
 	return bSelectedSomething;
 }
 
-bool URigVMController::SetNodePosition(URigVMNode* InNode, const FVector2D& InPosition, bool bUndo, bool bMergeUndoAction)
+bool URigVMController::SetNodePosition(URigVMNode* InNode, const FVector2D& InPosition, bool bSetupUndoRedo, bool bMergeUndoAction)
 {
 	if (!IsValidNodeForGraph(InNode))
 	{
@@ -1889,7 +1889,7 @@ bool URigVMController::SetNodePosition(URigVMNode* InNode, const FVector2D& InPo
 	}
 
 	FRigVMSetNodePositionAction Action;
-	if (bUndo)
+	if (bSetupUndoRedo)
 	{
 		Action = FRigVMSetNodePositionAction(InNode, InPosition);
 		Action.Title = FString::Printf(TEXT("Set Node Position"));
@@ -1899,7 +1899,7 @@ bool URigVMController::SetNodePosition(URigVMNode* InNode, const FVector2D& InPo
 	InNode->Position = InPosition;
 	Notify(ERigVMGraphNotifType::NodePositionChanged, InNode);
 
-	if (bUndo)
+	if (bSetupUndoRedo)
 	{
 		ActionStack->EndAction(Action, bMergeUndoAction);
 	}
@@ -1907,7 +1907,7 @@ bool URigVMController::SetNodePosition(URigVMNode* InNode, const FVector2D& InPo
 	return true;
 }
 
-bool URigVMController::SetNodePositionByName(const FName& InNodeName, const FVector2D& InPosition, bool bUndo, bool bMergeUndoAction)
+bool URigVMController::SetNodePositionByName(const FName& InNodeName, const FVector2D& InPosition, bool bSetupUndoRedo, bool bMergeUndoAction)
 {
 	if(!IsValidGraph())
 	{
@@ -1915,10 +1915,10 @@ bool URigVMController::SetNodePositionByName(const FName& InNodeName, const FVec
 	}
 
 	URigVMNode* Node = Graph->FindNodeByName(InNodeName);
-	return SetNodePosition(Node, InPosition, bUndo, bMergeUndoAction);
+	return SetNodePosition(Node, InPosition, bSetupUndoRedo, bMergeUndoAction);
 }
 
-bool URigVMController::SetNodeSize(URigVMNode* InNode, const FVector2D& InSize, bool bUndo, bool bMergeUndoAction)
+bool URigVMController::SetNodeSize(URigVMNode* InNode, const FVector2D& InSize, bool bSetupUndoRedo, bool bMergeUndoAction)
 {
 	if (!IsValidNodeForGraph(InNode))
 	{
@@ -1931,7 +1931,7 @@ bool URigVMController::SetNodeSize(URigVMNode* InNode, const FVector2D& InSize,
 	}
 
 	FRigVMSetNodeSizeAction Action;
-	if (bUndo)
+	if (bSetupUndoRedo)
 	{
 		Action = FRigVMSetNodeSizeAction(InNode, InSize);
 		Action.Title = FString::Printf(TEXT("Set Node Size"));
@@ -1941,7 +1941,7 @@ bool URigVMController::SetNodeSize(URigVMNode* InNode, const FVector2D& InSize,
 	InNode->Size = InSize;
 	Notify(ERigVMGraphNotifType::NodeSizeChanged, InNode);
 
-	if (bUndo)
+	if (bSetupUndoRedo)
 	{
 		ActionStack->EndAction(Action, bMergeUndoAction);
 	}
@@ -1949,7 +1949,7 @@ bool URigVMController::SetNodeSize(URigVMNode* InNode, const FVector2D& InSize,
 	return true;
 }
 
-bool URigVMController::SetNodeSizeByName(const FName& InNodeName, const FVector2D& InSize, bool bUndo, bool bMergeUndoAction)
+bool URigVMController::SetNodeSizeByName(const FName& InNodeName, const FVector2D& InSize, bool bSetupUndoRedo, bool bMergeUndoAction)
 {
 	if(!IsValidGraph())
 	{
@@ -1957,10 +1957,10 @@ bool URigVMController::SetNodeSizeByName(const FName& InNodeName, const FVector2
 	}
 
 	URigVMNode* Node = Graph->FindNodeByName(InNodeName);
-	return SetNodeSize(Node, InSize, bUndo, bMergeUndoAction);
+	return SetNodeSize(Node, InSize, bSetupUndoRedo, bMergeUndoAction);
 }
 
-bool URigVMController::SetNodeColor(URigVMNode* InNode, const FLinearColor& InColor, bool bUndo, bool bMergeUndoAction)
+bool URigVMController::SetNodeColor(URigVMNode* InNode, const FLinearColor& InColor, bool bSetupUndoRedo, bool bMergeUndoAction)
 {
 	if (!IsValidNodeForGraph(InNode))
 	{
@@ -1973,7 +1973,7 @@ bool URigVMController::SetNodeColor(URigVMNode* InNode, const FLinearColor& InCo
 	}
 
 	FRigVMSetNodeColorAction Action;
-	if (bUndo)
+	if (bSetupUndoRedo)
 	{
 		Action = FRigVMSetNodeColorAction(InNode, InColor);
 		Action.Title = FString::Printf(TEXT("Set Node Color"));
@@ -1983,7 +1983,7 @@ bool URigVMController::SetNodeColor(URigVMNode* InNode, const FLinearColor& InCo
 	InNode->NodeColor = InColor;
 	Notify(ERigVMGraphNotifType::NodeColorChanged, InNode);
 
-	if (bUndo)
+	if (bSetupUndoRedo)
 	{
 		ActionStack->EndAction(Action, bMergeUndoAction);
 	}
@@ -1991,7 +1991,7 @@ bool URigVMController::SetNodeColor(URigVMNode* InNode, const FLinearColor& InCo
 	return true;
 }
 
-bool URigVMController::SetNodeColorByName(const FName& InNodeName, const FLinearColor& InColor, bool bUndo, bool bMergeUndoAction)
+bool URigVMController::SetNodeColorByName(const FName& InNodeName, const FLinearColor& InColor, bool bSetupUndoRedo, bool bMergeUndoAction)
 {
 	if(!IsValidGraph())
 	{
@@ -1999,10 +1999,10 @@ bool URigVMController::SetNodeColorByName(const FName& InNodeName, const FLinear
 	}
 
 	URigVMNode* Node = Graph->FindNodeByName(InNodeName);
-	return SetNodeColor(Node, InColor, bUndo, bMergeUndoAction);
+	return SetNodeColor(Node, InColor, bSetupUndoRedo, bMergeUndoAction);
 }
 
-bool URigVMController::SetCommentText(URigVMNode* InNode, const FString& InCommentText, bool bUndo)
+bool URigVMController::SetCommentText(URigVMNode* InNode, const FString& InCommentText, bool bSetupUndoRedo)
 {
 	if (!IsValidNodeForGraph(InNode))
 	{
@@ -2017,7 +2017,7 @@ bool URigVMController::SetCommentText(URigVMNode* InNode, const FString& InComme
 		}
 
 		FRigVMSetCommentTextAction Action;
-		if (bUndo)
+		if (bSetupUndoRedo)
 		{
 			Action = FRigVMSetCommentTextAction(CommentNode, InCommentText);
 			Action.Title = FString::Printf(TEXT("Set Comment Text"));
@@ -2027,7 +2027,7 @@ bool URigVMController::SetCommentText(URigVMNode* InNode, const FString& InComme
 		CommentNode->CommentText = InCommentText;
 		Notify(ERigVMGraphNotifType::CommentTextChanged, InNode);
 
-		if (bUndo)
+		if (bSetupUndoRedo)
 		{
 			ActionStack->EndAction(Action);
 		}
@@ -2038,7 +2038,7 @@ bool URigVMController::SetCommentText(URigVMNode* InNode, const FString& InComme
 	return false;
 }
 
-bool URigVMController::SetCommentTextByName(const FName& InNodeName, const FString& InCommentText, bool bUndo)
+bool URigVMController::SetCommentTextByName(const FName& InNodeName, const FString& InCommentText, bool bSetupUndoRedo)
 {
 	if(!IsValidGraph())
 	{
@@ -2046,10 +2046,10 @@ bool URigVMController::SetCommentTextByName(const FName& InNodeName, const FStri
 	}
 
 	URigVMNode* Node = Graph->FindNodeByName(InNodeName);
-	return SetCommentText(Node, InCommentText, bUndo);
+	return SetCommentText(Node, InCommentText, bSetupUndoRedo);
 }
 
-bool URigVMController::SetRerouteCompactness(URigVMNode* InNode, bool bShowAsFullNode, bool bUndo)
+bool URigVMController::SetRerouteCompactness(URigVMNode* InNode, bool bShowAsFullNode, bool bSetupUndoRedo)
 {
 	if (!IsValidNodeForGraph(InNode))
 	{
@@ -2064,7 +2064,7 @@ bool URigVMController::SetRerouteCompactness(URigVMNode* InNode, bool bShowAsFul
 		}
 
 		FRigVMSetRerouteCompactnessAction Action;
-		if (bUndo)
+		if (bSetupUndoRedo)
 		{
 			Action = FRigVMSetRerouteCompactnessAction(RerouteNode, bShowAsFullNode);
 			Action.Title = FString::Printf(TEXT("Set Reroute Size"));
@@ -2074,7 +2074,7 @@ bool URigVMController::SetRerouteCompactness(URigVMNode* InNode, bool bShowAsFul
 		RerouteNode->bShowAsFullNode = bShowAsFullNode;
 		Notify(ERigVMGraphNotifType::RerouteCompactnessChanged, InNode);
 
-		if (bUndo)
+		if (bSetupUndoRedo)
 		{
 			ActionStack->EndAction(Action);
 		}
@@ -2085,7 +2085,7 @@ bool URigVMController::SetRerouteCompactness(URigVMNode* InNode, bool bShowAsFul
 	return false;
 }
 
-bool URigVMController::SetRerouteCompactnessByName(const FName& InNodeName, bool bShowAsFullNode, bool bUndo)
+bool URigVMController::SetRerouteCompactnessByName(const FName& InNodeName, bool bShowAsFullNode, bool bSetupUndoRedo)
 {
 	if(!IsValidGraph())
 	{
@@ -2093,10 +2093,10 @@ bool URigVMController::SetRerouteCompactnessByName(const FName& InNodeName, bool
 	}
 
 	URigVMNode* Node = Graph->FindNodeByName(InNodeName);
-	return SetRerouteCompactness(Node, bShowAsFullNode, bUndo);
+	return SetRerouteCompactness(Node, bShowAsFullNode, bSetupUndoRedo);
 }
 
-bool URigVMController::RenameVariable(const FName& InOldName, const FName& InNewName, bool bUndo)
+bool URigVMController::RenameVariable(const FName& InOldName, const FName& InNewName, bool bSetupUndoRedo)
 {
 	if(!IsValidGraph())
 	{
@@ -2120,7 +2120,7 @@ bool URigVMController::RenameVariable(const FName& InOldName, const FName& InNew
 	}
 
 	FRigVMRenameVariableAction Action;
-	if (bUndo)
+	if (bSetupUndoRedo)
 	{
 		Action = FRigVMRenameVariableAction(InOldName, InNewName);
 		Action.Title = FString::Printf(TEXT("Rename Variable"));
@@ -2149,7 +2149,7 @@ bool URigVMController::RenameVariable(const FName& InOldName, const FName& InNew
 		}
 	}
 
-	if (bUndo)
+	if (bSetupUndoRedo)
 	{
 		if (RenamedNodes.Num() > 0)
 		{
@@ -2164,7 +2164,7 @@ bool URigVMController::RenameVariable(const FName& InOldName, const FName& InNew
 	return RenamedNodes.Num() > 0;
 }
 
-bool URigVMController::RenameParameter(const FName& InOldName, const FName& InNewName, bool bUndo)
+bool URigVMController::RenameParameter(const FName& InOldName, const FName& InNewName, bool bSetupUndoRedo)
 {
 	if(!IsValidGraph())
 	{
@@ -2188,7 +2188,7 @@ bool URigVMController::RenameParameter(const FName& InOldName, const FName& InNe
 	}
 
 	FRigVMRenameParameterAction Action;
-	if (bUndo)
+	if (bSetupUndoRedo)
 	{
 		Action = FRigVMRenameParameterAction(InOldName, InNewName);
 		Action.Title = FString::Printf(TEXT("Rename Parameter"));
@@ -2217,7 +2217,7 @@ bool URigVMController::RenameParameter(const FName& InOldName, const FName& InNe
 		}
 	}
 
-	if (bUndo)
+	if (bSetupUndoRedo)
 	{
 		if (RenamedNodes.Num() > 0)
 		{
@@ -2232,7 +2232,7 @@ bool URigVMController::RenameParameter(const FName& InOldName, const FName& InNe
 	return RenamedNodes.Num() > 0;
 }
 
-void URigVMController::UpdateRerouteNodeAfterChangingLinks(URigVMPin* PinChanged, bool bUndo)
+void URigVMController::UpdateRerouteNodeAfterChangingLinks(URigVMPin* PinChanged, bool bSetupUndoRedo)
 {
 	if (bIgnoreRerouteCompactnessChanges)
 	{
@@ -2259,10 +2259,10 @@ void URigVMController::UpdateRerouteNodeAfterChangingLinks(URigVMPin* PinChanged
 	bool bOnlyConnectionsOnOneSide = (NbTotalSources == 0) || (NbTotalTargets == 0);
 	bool bShowAsFullNode = (!bJustTopLevelConnections) || bOnlyConnectionsOnOneSide;
 
-	SetRerouteCompactness(Node, bShowAsFullNode, bUndo);
+	SetRerouteCompactness(Node, bShowAsFullNode, bSetupUndoRedo);
 }
 
-bool URigVMController::SetPinExpansion(const FString& InPinPath, bool bIsExpanded, bool bUndo)
+bool URigVMController::SetPinExpansion(const FString& InPinPath, bool bIsExpanded, bool bSetupUndoRedo)
 {
 	if (!IsValidGraph())
 	{
@@ -2276,10 +2276,10 @@ bool URigVMController::SetPinExpansion(const FString& InPinPath, bool bIsExpande
 		return false;
 	}
 
-	return SetPinExpansion(Pin, bIsExpanded, bUndo);
+	return SetPinExpansion(Pin, bIsExpanded, bSetupUndoRedo);
 }
 
-bool URigVMController::SetPinExpansion(URigVMPin* InPin, bool bIsExpanded, bool bUndo)
+bool URigVMController::SetPinExpansion(URigVMPin* InPin, bool bIsExpanded, bool bSetupUndoRedo)
 {
 	if (InPin->GetSubPins().Num() == 0)
 	{
@@ -2292,7 +2292,7 @@ bool URigVMController::SetPinExpansion(URigVMPin* InPin, bool bIsExpanded, bool
 	}
 
 	FRigVMSetPinExpansionAction Action;
-	if (bUndo)
+	if (bSetupUndoRedo)
 	{
 		Action = FRigVMSetPinExpansionAction(InPin, bIsExpanded);
 		Action.Title = bIsExpanded ? TEXT("Expand Pin") : TEXT("Collapse Pin");
@@ -2307,7 +2307,7 @@ bool URigVMController::SetPinExpansion(URigVMPin* InPin, bool bIsExpanded, bool
 		Graph->MarkPackageDirty();
 	}
 
-	if (bUndo)
+	if (bSetupUndoRedo)
 	{
 		ActionStack->EndAction(Action);
 	}
@@ -2315,7 +2315,7 @@ bool URigVMController::SetPinExpansion(URigVMPin* InPin, bool bIsExpanded, bool
 	return true;
 }
 
-bool URigVMController::SetPinIsWatched(const FString& InPinPath, bool bIsWatched, bool bUndo)
+bool URigVMController::SetPinIsWatched(const FString& InPinPath, bool bIsWatched, bool bSetupUndoRedo)
 {
 	if (!IsValidGraph())
 	{
@@ -2329,10 +2329,10 @@ bool URigVMController::SetPinIsWatched(const FString& InPinPath, bool bIsWatched
 		return false;
 	}
 
-	return SetPinIsWatched(Pin, bIsWatched, bUndo);
+	return SetPinIsWatched(Pin, bIsWatched, bSetupUndoRedo);
 }
 
-bool URigVMController::SetPinIsWatched(URigVMPin* InPin, bool bIsWatched, bool bUndo)
+bool URigVMController::SetPinIsWatched(URigVMPin* InPin, bool bIsWatched, bool bSetupUndoRedo)
 {
 	if (!IsValidPinForGraph(InPin))
 	{
@@ -2350,7 +2350,7 @@ bool URigVMController::SetPinIsWatched(URigVMPin* InPin, bool bIsWatched, bool b
 	}
 
 	FRigVMSetPinWatchAction Action;
-	if (bUndo)
+	if (bSetupUndoRedo)
 	{
 		Action = FRigVMSetPinWatchAction(InPin, bIsWatched);
 		Action.Title = bIsWatched ? TEXT("Watch Pin") : TEXT("Unwatch Pin");
@@ -2365,7 +2365,7 @@ bool URigVMController::SetPinIsWatched(URigVMPin* InPin, bool bIsWatched, bool b
 		Graph->MarkPackageDirty();
 	}
 
-	if (bUndo)
+	if (bSetupUndoRedo)
 	{
 		ActionStack->EndAction(Action);
 	}
@@ -2391,7 +2391,7 @@ FString URigVMController::GetPinDefaultValue(const FString& InPinPath)
 	return Pin->GetDefaultValue();
 }
 
-bool URigVMController::SetPinDefaultValue(const FString& InPinPath, const FString& InDefaultValue, bool bResizeArrays, bool bUndo, bool bMergeUndoAction)
+bool URigVMController::SetPinDefaultValue(const FString& InPinPath, const FString& InDefaultValue, bool bResizeArrays, bool bSetupUndoRedo, bool bMergeUndoAction)
 {
 	if (!IsValidGraph())
 	{
@@ -2409,18 +2409,18 @@ bool URigVMController::SetPinDefaultValue(const FString& InPinPath, const FStrin
 	{
 		if (Pin->GetName() == URigVMVariableNode::VariableName)
 		{
-			return SetVariableName(VariableNode, *InDefaultValue, bUndo);
+			return SetVariableName(VariableNode, *InDefaultValue, bSetupUndoRedo);
 		}
 	}
 	if (URigVMParameterNode* ParameterNode = Cast<URigVMParameterNode>(Pin->GetNode()))
 	{
 		if (Pin->GetName() == URigVMParameterNode::ParameterName)
 		{
-			return SetParameterName(ParameterNode, *InDefaultValue, bUndo);
+			return SetParameterName(ParameterNode, *InDefaultValue, bSetupUndoRedo);
 		}
 	}
 
-	SetPinDefaultValue(Pin, InDefaultValue, bResizeArrays, bUndo, bMergeUndoAction);
+	SetPinDefaultValue(Pin, InDefaultValue, bResizeArrays, bSetupUndoRedo, bMergeUndoAction);
 	URigVMPin* PinForLink = Pin->GetPinForLink();
 	if (PinForLink != Pin)
 	{
@@ -2430,7 +2430,7 @@ bool URigVMController::SetPinDefaultValue(const FString& InPinPath, const FStrin
 	return true;
 }
 
-void URigVMController::SetPinDefaultValue(URigVMPin* InPin, const FString& InDefaultValue, bool bResizeArrays, bool bUndo, bool bMergeUndoAction)
+void URigVMController::SetPinDefaultValue(URigVMPin* InPin, const FString& InDefaultValue, bool bResizeArrays, bool bSetupUndoRedo, bool bMergeUndoAction)
 {
 	check(InPin);
 	ensure(!InDefaultValue.IsEmpty());
@@ -2441,7 +2441,7 @@ void URigVMController::SetPinDefaultValue(URigVMPin* InPin, const FString& InDef
 	}
 
 	FRigVMSetPinDefaultValueAction Action;
-	if (bUndo)
+	if (bSetupUndoRedo)
 	{
 		Action = FRigVMSetPinDefaultValueAction(InPin, InDefaultValue);
 		Action.Title = FString::Printf(TEXT("Set Pin Default Value"));
@@ -2459,11 +2459,11 @@ void URigVMController::SetPinDefaultValue(URigVMPin* InPin, const FString& InDef
 			{
 				while (Elements.Num() > InPin->SubPins.Num())
 				{
-					InsertArrayPin(InPin, INDEX_NONE, FString(), bUndo);
+					InsertArrayPin(InPin, INDEX_NONE, FString(), bSetupUndoRedo);
 				}
 				while (Elements.Num() < InPin->SubPins.Num())
 				{
-					RemoveArrayPin(InPin->SubPins.Last()->GetPinPath(), bUndo);
+					RemoveArrayPin(InPin->SubPins.Last()->GetPinPath(), bSetupUndoRedo);
 				}
 			}
 			else
@@ -2519,13 +2519,13 @@ void URigVMController::SetPinDefaultValue(URigVMPin* InPin, const FString& InDef
 		}
 	}
 
-	if (bUndo)
+	if (bSetupUndoRedo)
 	{
 		ActionStack->EndAction(Action, bMergeUndoAction);
 	}
 }
 
-bool URigVMController::ResetPinDefaultValue(const FString& InPinPath, bool bUndo)
+bool URigVMController::ResetPinDefaultValue(const FString& InPinPath, bool bSetupUndoRedo)
 {
 	if (!IsValidGraph())
 	{
@@ -2545,10 +2545,10 @@ bool URigVMController::ResetPinDefaultValue(const FString& InPinPath, bool bUndo
 		return false;
 	}
 
-	return ResetPinDefaultValue(Pin, bUndo);
+	return ResetPinDefaultValue(Pin, bSetupUndoRedo);
 }
 
-bool URigVMController::ResetPinDefaultValue(URigVMPin* InPin, bool bUndo)
+bool URigVMController::ResetPinDefaultValue(URigVMPin* InPin, bool bSetupUndoRedo)
 {
 	check(InPin);
 
@@ -2609,7 +2609,7 @@ bool URigVMController::ResetPinDefaultValue(URigVMPin* InPin, bool bUndo)
 
 			if (!DefaultValue.IsEmpty())
 			{
-				SetPinDefaultValue(InPin, DefaultValue, true, bUndo, false);
+				SetPinDefaultValue(InPin, DefaultValue, true, bSetupUndoRedo, false);
 				return true;
 			}
 		}
@@ -2678,12 +2678,12 @@ TArray<FString> URigVMController::SplitDefaultValue(const FString& InDefaultValu
 	return Parts;
 }
 
-FString URigVMController::AddArrayPin(const FString& InArrayPinPath, const FString& InDefaultValue, bool bUndo)
+FString URigVMController::AddArrayPin(const FString& InArrayPinPath, const FString& InDefaultValue, bool bSetupUndoRedo)
 {
-	return InsertArrayPin(InArrayPinPath, INDEX_NONE, InDefaultValue, bUndo);
+	return InsertArrayPin(InArrayPinPath, INDEX_NONE, InDefaultValue, bSetupUndoRedo);
 }
 
-FString URigVMController::DuplicateArrayPin(const FString& InArrayElementPinPath, bool bUndo)
+FString URigVMController::DuplicateArrayPin(const FString& InArrayElementPinPath, bool bSetupUndoRedo)
 {
 	if (!IsValidGraph())
 	{
@@ -2708,10 +2708,10 @@ FString URigVMController::DuplicateArrayPin(const FString& InArrayElementPinPath
 	ensure(ArrayPin->IsArray());
 
 	FString DefaultValue = ElementPin->GetDefaultValue();
-	return InsertArrayPin(ArrayPin->GetPinPath(), ElementPin->GetPinIndex() + 1, DefaultValue, bUndo);
+	return InsertArrayPin(ArrayPin->GetPinPath(), ElementPin->GetPinIndex() + 1, DefaultValue, bSetupUndoRedo);
 }
 
-FString URigVMController::InsertArrayPin(const FString& InArrayPinPath, int32 InIndex, const FString& InDefaultValue, bool bUndo)
+FString URigVMController::InsertArrayPin(const FString& InArrayPinPath, int32 InIndex, const FString& InDefaultValue, bool bSetupUndoRedo)
 {
 	if (!IsValidGraph())
 	{
@@ -2725,7 +2725,7 @@ FString URigVMController::InsertArrayPin(const FString& InArrayPinPath, int32 In
 		return FString();
 	}
 
-	URigVMPin* ElementPin = InsertArrayPin(ArrayPin, InIndex, InDefaultValue, bUndo);
+	URigVMPin* ElementPin = InsertArrayPin(ArrayPin, InIndex, InDefaultValue, bSetupUndoRedo);
 	if (ElementPin)
 	{
 		return ElementPin->GetPinPath();
@@ -2734,7 +2734,7 @@ FString URigVMController::InsertArrayPin(const FString& InArrayPinPath, int32 In
 	return FString();
 }
 
-URigVMPin* URigVMController::InsertArrayPin(URigVMPin* ArrayPin, int32 InIndex, const FString& InDefaultValue, bool bUndo)
+URigVMPin* URigVMController::InsertArrayPin(URigVMPin* ArrayPin, int32 InIndex, const FString& InDefaultValue, bool bSetupUndoRedo)
 {
 	if (!ArrayPin->IsArray())
 	{
@@ -2754,7 +2754,7 @@ URigVMPin* URigVMController::InsertArrayPin(URigVMPin* ArrayPin, int32 InIndex,
 	}
 
 	FRigVMInsertArrayPinAction Action;
-	if (bUndo)
+	if (bSetupUndoRedo)
 	{
 		Action = FRigVMInsertArrayPinAction(ArrayPin, InIndex, InDefaultValue);
 		Action.Title = FString::Printf(TEXT("Insert Array Pin"));
@@ -2804,7 +2804,7 @@ URigVMPin* URigVMController::InsertArrayPin(URigVMPin* ArrayPin, int32 InIndex,
 		Graph->MarkPackageDirty();
 	}
 
-	if (bUndo)
+	if (bSetupUndoRedo)
 	{
 		ActionStack->EndAction(Action);
 	}
@@ -2812,7 +2812,7 @@ URigVMPin* URigVMController::InsertArrayPin(URigVMPin* ArrayPin, int32 InIndex,
 	return Pin;
 }
 
-bool URigVMController::RemoveArrayPin(const FString& InArrayElementPinPath, bool bUndo)
+bool URigVMController::RemoveArrayPin(const FString& InArrayElementPinPath, bool bSetupUndoRedo)
 {
 	if(!IsValidGraph())
 	{
@@ -2837,7 +2837,7 @@ bool URigVMController::RemoveArrayPin(const FString& InArrayElementPinPath, bool
 	ensure(ArrayPin->IsArray());
 
 	FRigVMRemoveArrayPinAction Action;
-	if (bUndo)
+	if (bSetupUndoRedo)
 	{
 		Action = FRigVMRemoveArrayPinAction(ArrayElementPin);
 		Action.Title = FString::Printf(TEXT("Remove Array Pin"));
@@ -2845,7 +2845,7 @@ bool URigVMController::RemoveArrayPin(const FString& InArrayElementPinPath, bool
 	}
 
 	int32 IndexToRemove = ArrayElementPin->GetPinIndex();
-	if (!RemovePin(ArrayElementPin, bUndo))
+	if (!RemovePin(ArrayElementPin, bSetupUndoRedo))
 	{
 		return false;
 	}
@@ -2862,7 +2862,7 @@ bool URigVMController::RemoveArrayPin(const FString& InArrayElementPinPath, bool
 	}
 	Notify(ERigVMGraphNotifType::PinArraySizeChanged, ArrayPin);
 
-	if (bUndo)
+	if (bSetupUndoRedo)
 	{
 		ActionStack->EndAction(Action);
 	}
@@ -2870,14 +2870,14 @@ bool URigVMController::RemoveArrayPin(const FString& InArrayElementPinPath, bool
 	return true;
 }
 
-bool URigVMController::RemovePin(URigVMPin* InPinToRemove, bool bUndo)
+bool URigVMController::RemovePin(URigVMPin* InPinToRemove, bool bSetupUndoRedo)
 {
-	if (bUndo)
+	if (bSetupUndoRedo)
 	{
-		BreakAllLinks(InPinToRemove, true, bUndo);
-		BreakAllLinks(InPinToRemove, false, bUndo);
-		BreakAllLinksRecursive(InPinToRemove, true, false, bUndo);
-		BreakAllLinksRecursive(InPinToRemove, false, false, bUndo);
+		BreakAllLinks(InPinToRemove, true, bSetupUndoRedo);
+		BreakAllLinks(InPinToRemove, false, bSetupUndoRedo);
+		BreakAllLinksRecursive(InPinToRemove, true, false, bSetupUndoRedo);
+		BreakAllLinksRecursive(InPinToRemove, false, false, bSetupUndoRedo);
 	}
 
 	if (URigVMPin* ParentPin = InPinToRemove->GetParentPin())
@@ -2894,7 +2894,7 @@ bool URigVMController::RemovePin(URigVMPin* InPinToRemove, bool bUndo)
 	TArray<URigVMPin*> SubPins = InPinToRemove->GetSubPins();
 	for (URigVMPin* SubPin : SubPins)
 	{
-		if (!RemovePin(SubPin, bUndo))
+		if (!RemovePin(SubPin, bSetupUndoRedo))
 		{
 			return false;
 		}
@@ -2903,12 +2903,12 @@ bool URigVMController::RemovePin(URigVMPin* InPinToRemove, bool bUndo)
 	return true;
 }
 
-bool URigVMController::ClearArrayPin(const FString& InArrayPinPath, bool bUndo)
+bool URigVMController::ClearArrayPin(const FString& InArrayPinPath, bool bSetupUndoRedo)
 {
-	return SetArrayPinSize(InArrayPinPath, 0, FString(), bUndo);
+	return SetArrayPinSize(InArrayPinPath, 0, FString(), bSetupUndoRedo);
 }
 
-bool URigVMController::SetArrayPinSize(const FString& InArrayPinPath, int32 InSize, const FString& InDefaultValue, bool bUndo)
+bool URigVMController::SetArrayPinSize(const FString& InArrayPinPath, int32 InSize, const FString& InDefaultValue, bool bSetupUndoRedo)
 {
 	if (!IsValidGraph())
 	{
@@ -2929,7 +2929,7 @@ bool URigVMController::SetArrayPinSize(const FString& InArrayPinPath, int32 InSi
 	}
 
 	FRigVMBaseAction Action;
-	if (bUndo)
+	if (bSetupUndoRedo)
 	{
 		Action.Title = FString::Printf(TEXT("Set Array Pin Size (%d)"), InSize);
 		ActionStack->BeginAction(Action);
@@ -2951,9 +2951,9 @@ bool URigVMController::SetArrayPinSize(const FString& InArrayPinPath, int32 InSi
 
 	while (Pin->GetSubPins().Num() > InSize)
 	{
-		if (!RemoveArrayPin(Pin->GetSubPins()[Pin->GetSubPins().Num()-1]->GetPinPath(), bUndo))
+		if (!RemoveArrayPin(Pin->GetSubPins()[Pin->GetSubPins().Num()-1]->GetPinPath(), bSetupUndoRedo))
 		{
-			if (bUndo)
+			if (bSetupUndoRedo)
 			{
 				ActionStack->CancelAction(Action);
 			}
@@ -2964,9 +2964,9 @@ bool URigVMController::SetArrayPinSize(const FString& InArrayPinPath, int32 InSi
 
 	while (Pin->GetSubPins().Num() < InSize)
 	{
-		if (AddArrayPin(Pin->GetPinPath(), DefaultValue, bUndo).IsEmpty())
+		if (AddArrayPin(Pin->GetPinPath(), DefaultValue, bSetupUndoRedo).IsEmpty())
 		{
-			if (bUndo)
+			if (bSetupUndoRedo)
 			{
 				ActionStack->CancelAction(Action);
 			}
@@ -2975,7 +2975,7 @@ bool URigVMController::SetArrayPinSize(const FString& InArrayPinPath, int32 InSi
 		AddedPins++;
 	}
 
-	if (bUndo)
+	if (bSetupUndoRedo)
 	{
 		if (RemovedPins > 0 || AddedPins > 0)
 		{
@@ -2990,7 +2990,7 @@ bool URigVMController::SetArrayPinSize(const FString& InArrayPinPath, int32 InSi
 	return RemovedPins > 0 || AddedPins > 0;
 }
 
-bool URigVMController::BindPinToVariable(const FString& InPinPath, const FString& InNewBoundVariablePath, bool bUndo)
+bool URigVMController::BindPinToVariable(const FString& InPinPath, const FString& InNewBoundVariablePath, bool bSetupUndoRedo)
 {
 	if (!IsValidGraph())
 	{
@@ -3004,15 +3004,15 @@ bool URigVMController::BindPinToVariable(const FString& InPinPath, const FString
 		return false;
 	}
 
-	return BindPinToVariable(Pin, InNewBoundVariablePath, bUndo);
+	return BindPinToVariable(Pin, InNewBoundVariablePath, bSetupUndoRedo);
 }
 
-bool URigVMController::UnbindPinFromVariable(const FString& InPinPath, bool bUndo)
+bool URigVMController::UnbindPinFromVariable(const FString& InPinPath, bool bSetupUndoRedo)
 {
-	return BindPinToVariable(InPinPath, FString(), bUndo);
+	return BindPinToVariable(InPinPath, FString(), bSetupUndoRedo);
 }
 
-bool URigVMController::BindPinToVariable(URigVMPin* InPin, const FString& InNewBoundVariablePath, bool bUndo)
+bool URigVMController::BindPinToVariable(URigVMPin* InPin, const FString& InNewBoundVariablePath, bool bSetupUndoRedo)
 {
 	if (!IsValidPinForGraph(InPin))
 	{
@@ -3050,7 +3050,7 @@ bool URigVMController::BindPinToVariable(URigVMPin* InPin, const FString& InNewB
 	}
 
 	FRigVMBaseAction Action;
-	if (bUndo)
+	if (bSetupUndoRedo)
 	{
 		if (InNewBoundVariablePath.IsEmpty())
 		{
@@ -3063,15 +3063,15 @@ bool URigVMController::BindPinToVariable(URigVMPin* InPin, const FString& InNewB
 		ActionStack->BeginAction(Action);
 	}
 
-	if (!InPin->IsBoundToVariable() && bUndo)
+	if (!InPin->IsBoundToVariable() && bSetupUndoRedo)
 	{
 		// break all links on pin towards parent + children
-		BreakAllLinks(InPin, true, bUndo);
-		BreakAllLinksRecursive(InPin, true, true, bUndo);
-		BreakAllLinksRecursive(InPin, true, false, bUndo);
+		BreakAllLinks(InPin, true, bSetupUndoRedo);
+		BreakAllLinksRecursive(InPin, true, true, bSetupUndoRedo);
+		BreakAllLinksRecursive(InPin, true, false, bSetupUndoRedo);
 	}
 
-	if (bUndo)
+	if (bSetupUndoRedo)
 	{
 		ActionStack->AddAction(FRigVMSetPinBoundVariableAction(InPin, InNewBoundVariablePath));
 	}
@@ -3079,7 +3079,7 @@ bool URigVMController::BindPinToVariable(URigVMPin* InPin, const FString& InNewB
 	InPin->BoundVariablePath = InNewBoundVariablePath;
 	Notify(ERigVMGraphNotifType::PinBoundVariableChanged, InPin);
 
-	if (bUndo)
+	if (bSetupUndoRedo)
 	{
 		ActionStack->EndAction(Action);
 	}
@@ -3087,7 +3087,7 @@ bool URigVMController::BindPinToVariable(URigVMPin* InPin, const FString& InNewB
 	return true;
 }
 
-bool URigVMController::MakeBindingsFromVariableNode(const FName& InNodeName, bool bUndo)
+bool URigVMController::MakeBindingsFromVariableNode(const FName& InNodeName, bool bSetupUndoRedo)
 {
 	if (!IsValidGraph())
 	{
@@ -3096,13 +3096,13 @@ bool URigVMController::MakeBindingsFromVariableNode(const FName& InNodeName, boo
 
 	if (URigVMVariableNode* VariableNode = Cast<URigVMVariableNode>(Graph->FindNodeByName(InNodeName)))
 	{
-		return MakeBindingsFromVariableNode(VariableNode, bUndo);
+		return MakeBindingsFromVariableNode(VariableNode, bSetupUndoRedo);
 	}
 
 	return false;
 }
 
-bool URigVMController::MakeBindingsFromVariableNode(URigVMVariableNode* InNode, bool bUndo)
+bool URigVMController::MakeBindingsFromVariableNode(URigVMVariableNode* InNode, bool bSetupUndoRedo)
 {
 	check(InNode);
 
@@ -3121,7 +3121,7 @@ bool URigVMController::MakeBindingsFromVariableNode(URigVMVariableNode* InNode,
 
 	if (Links.Num() > 0)
 	{
-		if (bUndo)
+		if (bSetupUndoRedo)
 		{
 			OpenUndoBracket(TEXT("Turn Variable Node into Bindings"));
 		}
@@ -3137,15 +3137,15 @@ bool URigVMController::MakeBindingsFromVariableNode(URigVMVariableNode* InNode,
 				VariablePathToBind = FString::Printf(TEXT("%s.%s"), *VariablePathToBind, *SegmentPath);
 			}
 
-			if (!BindPinToVariable(TargetPin, VariablePathToBind, bUndo))
+			if (!BindPinToVariable(TargetPin, VariablePathToBind, bSetupUndoRedo))
 			{
 				CancelUndoBracket();
 			}
 		}
 
-		RemoveNode(InNode, bUndo, true);
+		RemoveNode(InNode, bSetupUndoRedo, true);
 
-		if (bUndo)
+		if (bSetupUndoRedo)
 		{
 			CloseUndoBracket();
 		}
@@ -3156,12 +3156,12 @@ bool URigVMController::MakeBindingsFromVariableNode(URigVMVariableNode* InNode,
 
 }
 
-bool URigVMController::MakeVariableNodeFromBinding(const FString& InPinPath, const FVector2D& InNodePosition, bool bUndo)
+bool URigVMController::MakeVariableNodeFromBinding(const FString& InPinPath, const FVector2D& InNodePosition, bool bSetupUndoRedo)
 {
-	return PromotePinToVariable(InPinPath, true, InNodePosition, bUndo);
+	return PromotePinToVariable(InPinPath, true, InNodePosition, bSetupUndoRedo);
 }
 
-bool URigVMController::PromotePinToVariable(const FString& InPinPath, bool bCreateVariableNode, const FVector2D& InNodePosition, bool bUndo)
+bool URigVMController::PromotePinToVariable(const FString& InPinPath, bool bCreateVariableNode, const FVector2D& InNodePosition, bool bSetupUndoRedo)
 {
 	if (!IsValidGraph())
 	{
@@ -3175,10 +3175,10 @@ bool URigVMController::PromotePinToVariable(const FString& InPinPath, bool bCrea
 		return false;
 	}
 
-	return PromotePinToVariable(Pin, bCreateVariableNode, InNodePosition, bUndo);
+	return PromotePinToVariable(Pin, bCreateVariableNode, InNodePosition, bSetupUndoRedo);
 }
 
-bool URigVMController::PromotePinToVariable(URigVMPin* InPin, bool bCreateVariableNode, const FVector2D& InNodePosition, bool bUndo)
+bool URigVMController::PromotePinToVariable(URigVMPin* InPin, bool bCreateVariableNode, const FVector2D& InNodePosition, bool bSetupUndoRedo)
 {
 	check(InPin);
 
@@ -3224,23 +3224,23 @@ bool URigVMController::PromotePinToVariable(URigVMPin* InPin, bool bCreateVariab
 			FString(),
 			InNodePosition,
 			FString(),
-			bUndo))
+			bSetupUndoRedo))
 		{
 			if (URigVMPin* ValuePin = VariableNode->FindPin(URigVMVariableNode::ValueName))
 			{
-				return AddLink(ValuePin, InPin, bUndo);
+				return AddLink(ValuePin, InPin, bSetupUndoRedo);
 			}
 		}
 	}
 	else
 	{
-		return BindPinToVariable(InPin, VariableForPin.Name.ToString(), bUndo);
+		return BindPinToVariable(InPin, VariableForPin.Name.ToString(), bSetupUndoRedo);
 	}
 
 	return false;
 }
 
-bool URigVMController::AddLink(const FString& InOutputPinPath, const FString& InInputPinPath, bool bUndo)
+bool URigVMController::AddLink(const FString& InOutputPinPath, const FString& InInputPinPath, bool bSetupUndoRedo)
 {
 	if(!IsValidGraph())
 	{
@@ -3275,10 +3275,10 @@ bool URigVMController::AddLink(const FString& InOutputPinPath, const FString& In
 	}
 	InputPin = InputPin->GetPinForLink();
 
-	return AddLink(OutputPin, InputPin, bUndo);
+	return AddLink(OutputPin, InputPin, bSetupUndoRedo);
 }
 
-bool URigVMController::AddLink(URigVMPin* OutputPin, URigVMPin* InputPin, bool bUndo)
+bool URigVMController::AddLink(URigVMPin* OutputPin, URigVMPin* InputPin, bool bSetupUndoRedo)
 {
 	if(OutputPin == nullptr)
 	{
@@ -3316,7 +3316,7 @@ bool URigVMController::AddLink(URigVMPin* OutputPin, URigVMPin* InputPin, bool b
 	ensure(!InputPin->IsLinkedTo(OutputPin));
 
 	FRigVMAddLinkAction Action;
-	if (bUndo)
+	if (bSetupUndoRedo)
 	{
 		Action = FRigVMAddLinkAction(OutputPin, InputPin);
 		Action.Title = FString::Printf(TEXT("Add Link"));
@@ -3325,17 +3325,17 @@ bool URigVMController::AddLink(URigVMPin* OutputPin, URigVMPin* InputPin, bool b
 
 	if (OutputPin->IsExecuteContext())
 	{
-		BreakAllLinks(OutputPin, false, bUndo);
+		BreakAllLinks(OutputPin, false, bSetupUndoRedo);
 	}
 
-	BreakAllLinks(InputPin, true, bUndo);
-	if (bUndo)
+	BreakAllLinks(InputPin, true, bSetupUndoRedo);
+	if (bSetupUndoRedo)
 	{
-		BreakAllLinksRecursive(InputPin, true, true, bUndo);
-		BreakAllLinksRecursive(InputPin, true, false, bUndo);
+		BreakAllLinksRecursive(InputPin, true, true, bSetupUndoRedo);
+		BreakAllLinksRecursive(InputPin, true, false, bSetupUndoRedo);
 	}
 
-	if (bUndo)
+	if (bSetupUndoRedo)
 	{
 		ExpandPinRecursively(OutputPin->GetParentPin(), true);
 		ExpandPinRecursively(InputPin->GetParentPin(), true);
@@ -3356,14 +3356,14 @@ bool URigVMController::AddLink(URigVMPin* OutputPin, URigVMPin* InputPin, bool b
 	}
 	Notify(ERigVMGraphNotifType::LinkAdded, Link);
 
-	UpdateRerouteNodeAfterChangingLinks(OutputPin, bUndo);
-	UpdateRerouteNodeAfterChangingLinks(InputPin, bUndo);
+	UpdateRerouteNodeAfterChangingLinks(OutputPin, bSetupUndoRedo);
+	UpdateRerouteNodeAfterChangingLinks(InputPin, bSetupUndoRedo);
 
 	TArray<URigVMNode*> NodesVisited;
-	PotentiallyResolvePrototypeNode(Cast<URigVMPrototypeNode>(InputPin->GetNode()), bUndo, NodesVisited);
-	PotentiallyResolvePrototypeNode(Cast<URigVMPrototypeNode>(OutputPin->GetNode()), bUndo, NodesVisited);
+	PotentiallyResolvePrototypeNode(Cast<URigVMPrototypeNode>(InputPin->GetNode()), bSetupUndoRedo, NodesVisited);
+	PotentiallyResolvePrototypeNode(Cast<URigVMPrototypeNode>(OutputPin->GetNode()), bSetupUndoRedo, NodesVisited);
 
-	if (bUndo)
+	if (bSetupUndoRedo)
 	{
 		ActionStack->EndAction(Action);
 	}
@@ -3371,7 +3371,7 @@ bool URigVMController::AddLink(URigVMPin* OutputPin, URigVMPin* InputPin, bool b
 	return true;
 }
 
-bool URigVMController::BreakLink(const FString& InOutputPinPath, const FString& InInputPinPath, bool bUndo)
+bool URigVMController::BreakLink(const FString& InOutputPinPath, const FString& InInputPinPath, bool bSetupUndoRedo)
 {
 	if(!IsValidGraph())
 	{
@@ -3394,10 +3394,10 @@ bool URigVMController::BreakLink(const FString& InOutputPinPath, const FString&
 	}
 	InputPin = InputPin->GetPinForLink();
 
-	return BreakLink(OutputPin, InputPin, bUndo);
+	return BreakLink(OutputPin, InputPin, bSetupUndoRedo);
 }
 
-bool URigVMController::BreakLink(URigVMPin* OutputPin, URigVMPin* InputPin, bool bUndo)
+bool URigVMController::BreakLink(URigVMPin* OutputPin, URigVMPin* InputPin, bool bSetupUndoRedo)
 {
 	if(!IsValidPinForGraph(OutputPin) || !IsValidPinForGraph(InputPin))
 	{
@@ -3415,7 +3415,7 @@ bool URigVMController::BreakLink(URigVMPin* OutputPin, URigVMPin* InputPin, bool
 		if (Link->SourcePin == OutputPin && Link->TargetPin == InputPin)
 		{
 			FRigVMBreakLinkAction Action;
-			if (bUndo)
+			if (bSetupUndoRedo)
 			{
 				Action = FRigVMBreakLinkAction(OutputPin, InputPin);
 				Action.Title = FString::Printf(TEXT("Break Link"));
@@ -3434,10 +3434,10 @@ bool URigVMController::BreakLink(URigVMPin* OutputPin, URigVMPin* InputPin, bool
 
 			DestroyObject(Link);
 
-			UpdateRerouteNodeAfterChangingLinks(OutputPin, bUndo);
-			UpdateRerouteNodeAfterChangingLinks(InputPin, bUndo);
+			UpdateRerouteNodeAfterChangingLinks(OutputPin, bSetupUndoRedo);
+			UpdateRerouteNodeAfterChangingLinks(InputPin, bSetupUndoRedo);
 
-			if (bUndo)
+			if (bSetupUndoRedo)
 			{
 				ActionStack->EndAction(Action);
 			}
@@ -3449,7 +3449,7 @@ bool URigVMController::BreakLink(URigVMPin* OutputPin, URigVMPin* InputPin, bool
 	return false;
 }
 
-bool URigVMController::BreakAllLinks(const FString& InPinPath, bool bAsInput, bool bUndo)
+bool URigVMController::BreakAllLinks(const FString& InPinPath, bool bAsInput, bool bSetupUndoRedo)
 {
 	if(!IsValidGraph())
 	{
@@ -3469,22 +3469,22 @@ bool URigVMController::BreakAllLinks(const FString& InPinPath, bool bAsInput, bo
 		return false;
 	}
 
-	return BreakAllLinks(Pin, bAsInput, bUndo);
+	return BreakAllLinks(Pin, bAsInput, bSetupUndoRedo);
 }
 
-bool URigVMController::BreakAllLinks(URigVMPin* Pin, bool bAsInput, bool bUndo)
+bool URigVMController::BreakAllLinks(URigVMPin* Pin, bool bAsInput, bool bSetupUndoRedo)
 {
 	FRigVMBaseAction Action;
-	if (bUndo)
+	if (bSetupUndoRedo)
 	{
 		Action.Title = FString::Printf(TEXT("Break All Links"));
 		ActionStack->BeginAction(Action);
 	}
 
 	int32 LinksBroken = 0;
-	if (Pin->IsBoundToVariable() && bAsInput && bUndo)
+	if (Pin->IsBoundToVariable() && bAsInput && bSetupUndoRedo)
 	{
-		BindPinToVariable(Pin, FString(), bUndo);
+		BindPinToVariable(Pin, FString(), bSetupUndoRedo);
 		LinksBroken++;
 	}
 
@@ -3494,15 +3494,15 @@ bool URigVMController::BreakAllLinks(URigVMPin* Pin, bool bAsInput, bool bUndo)
 		URigVMLink* Link = Links[LinkIndex];
 		if (bAsInput && Link->GetTargetPin() == Pin)
 		{
-			LinksBroken += BreakLink(Link->GetSourcePin(), Pin, bUndo) ? 1 : 0;
+			LinksBroken += BreakLink(Link->GetSourcePin(), Pin, bSetupUndoRedo) ? 1 : 0;
 		}
 		else if (!bAsInput && Link->GetSourcePin() == Pin)
 		{
-			LinksBroken += BreakLink(Pin, Link->GetTargetPin(), bUndo) ? 1 : 0;
+			LinksBroken += BreakLink(Pin, Link->GetTargetPin(), bSetupUndoRedo) ? 1 : 0;
 		}
 	}
 
-	if (bUndo)
+	if (bSetupUndoRedo)
 	{
 		if (LinksBroken > 0)
 		{
@@ -3517,35 +3517,35 @@ bool URigVMController::BreakAllLinks(URigVMPin* Pin, bool bAsInput, bool bUndo)
 	return LinksBroken > 0;
 }
 
-void URigVMController::BreakAllLinksRecursive(URigVMPin* Pin, bool bAsInput, bool bTowardsParent, bool bUndo)
+void URigVMController::BreakAllLinksRecursive(URigVMPin* Pin, bool bAsInput, bool bTowardsParent, bool bSetupUndoRedo)
 {
 	if (bTowardsParent)
 	{
 		URigVMPin* ParentPin = Pin->GetParentPin();
 		if (ParentPin)
 		{
-			BreakAllLinks(ParentPin, bAsInput, bUndo);
-			BreakAllLinksRecursive(ParentPin, bAsInput, bTowardsParent, bUndo);
+			BreakAllLinks(ParentPin, bAsInput, bSetupUndoRedo);
+			BreakAllLinksRecursive(ParentPin, bAsInput, bTowardsParent, bSetupUndoRedo);
 		}
 	}
 	else
 	{
 		for (URigVMPin* SubPin : Pin->SubPins)
 		{
-			BreakAllLinks(SubPin, bAsInput, bUndo);
-			BreakAllLinksRecursive(SubPin, bAsInput, bTowardsParent, bUndo);
+			BreakAllLinks(SubPin, bAsInput, bSetupUndoRedo);
+			BreakAllLinksRecursive(SubPin, bAsInput, bTowardsParent, bSetupUndoRedo);
 		}
 	}
 }
 
-void URigVMController::ExpandPinRecursively(URigVMPin* InPin, bool bUndo)
+void URigVMController::ExpandPinRecursively(URigVMPin* InPin, bool bSetupUndoRedo)
 {
 	if (InPin == nullptr)
 	{
 		return;
 	}
 
-	if (bUndo)
+	if (bSetupUndoRedo)
 	{
 		OpenUndoBracket(TEXT("Expand Pin Recursively"));
 	}
@@ -3553,14 +3553,14 @@ void URigVMController::ExpandPinRecursively(URigVMPin* InPin, bool bUndo)
 	bool bExpandedSomething = false;
 	while (InPin)
 	{
-		if (SetPinExpansion(InPin, true, bUndo))
+		if (SetPinExpansion(InPin, true, bSetupUndoRedo))
 		{
 			bExpandedSomething = true;
 		}
 		InPin = InPin->GetParentPin();
 	}
 
-	if (bUndo)
+	if (bSetupUndoRedo)
 	{
 		if (bExpandedSomething)
 		{
@@ -3573,7 +3573,7 @@ void URigVMController::ExpandPinRecursively(URigVMPin* InPin, bool bUndo)
 	}
 }
 
-bool URigVMController::SetVariableName(URigVMVariableNode* InVariableNode, const FName& InVariableName, bool bUndo)
+bool URigVMController::SetVariableName(URigVMVariableNode* InVariableNode, const FName& InVariableName, bool bSetupUndoRedo)
 {
 	if (!IsValidNodeForGraph(InVariableNode))
 	{
@@ -3623,14 +3623,14 @@ bool URigVMController::SetVariableName(URigVMVariableNode* InVariableNode, const
 		Notify(ERigVMGraphNotifType::VariableRemoved, InVariableNode);
 	}
 
-	SetPinDefaultValue(InVariableNode->FindPin(URigVMVariableNode::VariableName), VariableName.ToString(), false, bUndo, false);
+	SetPinDefaultValue(InVariableNode->FindPin(URigVMVariableNode::VariableName), VariableName.ToString(), false, bSetupUndoRedo, false);
 
 	Notify(ERigVMGraphNotifType::VariableAdded, InVariableNode);
 
 	return true;
 }
 
-bool URigVMController::SetParameterName(URigVMParameterNode* InParameterNode, const FName& InParameterName, bool bUndo)
+bool URigVMController::SetParameterName(URigVMParameterNode* InParameterNode, const FName& InParameterName, bool bSetupUndoRedo)
 {
 	if (!IsValidNodeForGraph(InParameterNode))
 	{
@@ -3680,14 +3680,14 @@ bool URigVMController::SetParameterName(URigVMParameterNode* InParameterNode, co
 		Notify(ERigVMGraphNotifType::ParameterRemoved, InParameterNode);
 	}
 
-	SetPinDefaultValue(InParameterNode->FindPin(URigVMParameterNode::ParameterName), ParameterName.ToString(), false, bUndo, false);
+	SetPinDefaultValue(InParameterNode->FindPin(URigVMParameterNode::ParameterName), ParameterName.ToString(), false, bSetupUndoRedo, false);
 
 	Notify(ERigVMGraphNotifType::ParameterAdded, InParameterNode);
 
 	return true;
 }
 
-URigVMRerouteNode* URigVMController::AddFreeRerouteNode(bool bShowAsFullNode, const FString& InCPPType, const FName& InCPPTypeObjectPath, bool bIsConstant, const FName& InCustomWidgetName, const FString& InDefaultValue, const FVector2D& InPosition, const FString& InNodeName, bool bUndo)
+URigVMRerouteNode* URigVMController::AddFreeRerouteNode(bool bShowAsFullNode, const FString& InCPPType, const FName& InCPPTypeObjectPath, bool bIsConstant, const FName& InCustomWidgetName, const FString& InDefaultValue, const FVector2D& InPosition, const FString& InNodeName, bool bSetupUndoRedo)
 {
 	if (!IsValidGraph())
 	{
@@ -3695,7 +3695,7 @@ URigVMRerouteNode* URigVMController::AddFreeRerouteNode(bool bShowAsFullNode, co
 	}
 
 	FRigVMBaseAction Action;
-	if (bUndo)
+	if (bSetupUndoRedo)
 	{
 		Action.Title = FString::Printf(TEXT("Add Reroute"));
 		ActionStack->BeginAction(Action);
@@ -3733,14 +3733,14 @@ URigVMRerouteNode* URigVMController::AddFreeRerouteNode(bool bShowAsFullNode, co
 		}
 	});
 
-	if (bUndo)
+	if (bSetupUndoRedo)
 	{
 		ActionStack->AddAction(FRigVMAddRerouteNodeAction(Node));
 	}
 
 	Notify(ERigVMGraphNotifType::NodeAdded, Node);
 
-	if (bUndo)
+	if (bSetupUndoRedo)
 	{
 		ActionStack->EndAction(Action);
 	}
@@ -3748,7 +3748,7 @@ URigVMRerouteNode* URigVMController::AddFreeRerouteNode(bool bShowAsFullNode, co
 	return Node;
 }
 
-URigVMBranchNode* URigVMController::AddBranchNode(const FVector2D& InPosition, const FString& InNodeName, bool bUndo)
+URigVMBranchNode* URigVMController::AddBranchNode(const FVector2D& InPosition, const FString& InNodeName, bool bSetupUndoRedo)
 {
 	if(!IsValidGraph())
 	{
@@ -3790,7 +3790,7 @@ URigVMBranchNode* URigVMController::AddBranchNode(const FVector2D& InPosition, c
 
 	Notify(ERigVMGraphNotifType::NodeAdded, Node);
 
-	if (bUndo)
+	if (bSetupUndoRedo)
 	{
 		ActionStack->AddAction(FRigVMAddBranchNodeAction(Node));
 	}
@@ -3798,7 +3798,7 @@ URigVMBranchNode* URigVMController::AddBranchNode(const FVector2D& InPosition, c
 	return Node;
 }
 
-URigVMIfNode* URigVMController::AddIfNode(const FString& InCPPType, const FName& InCPPTypeObjectPath, const FVector2D& InPosition, const FString& InNodeName, bool bUndo)
+URigVMIfNode* URigVMController::AddIfNode(const FString& InCPPType, const FName& InCPPTypeObjectPath, const FVector2D& InPosition, const FString& InNodeName, bool bSetupUndoRedo)
 {
 	if(!IsValidGraph())
 	{
@@ -3883,7 +3883,7 @@ URigVMIfNode* URigVMController::AddIfNode(const FString& InCPPType, const FName&
 
 	Notify(ERigVMGraphNotifType::NodeAdded, Node);
 
-	if (bUndo)
+	if (bSetupUndoRedo)
 	{
 		ActionStack->AddAction(FRigVMAddIfNodeAction(Node));
 	}
@@ -3891,7 +3891,7 @@ URigVMIfNode* URigVMController::AddIfNode(const FString& InCPPType, const FName&
 	return Node;
 }
 
-URigVMSelectNode* URigVMController::AddSelectNode(const FString& InCPPType, const FName& InCPPTypeObjectPath, const FVector2D& InPosition, const FString& InNodeName, bool bUndo)
+URigVMSelectNode* URigVMController::AddSelectNode(const FString& InCPPType, const FName& InCPPTypeObjectPath, const FVector2D& InPosition, const FString& InNodeName, bool bSetupUndoRedo)
 {
 	if (!IsValidGraph())
 	{
@@ -3955,7 +3955,7 @@ URigVMSelectNode* URigVMController::AddSelectNode(const FString& InCPPType, cons
 
 	SetArrayPinSize(ValuePin->GetPinPath(), 2, DefaultValue, false);
 
-	if (bUndo)
+	if (bSetupUndoRedo)
 	{
 		ActionStack->AddAction(FRigVMAddSelectNodeAction(Node));
 	}
@@ -3963,7 +3963,7 @@ URigVMSelectNode* URigVMController::AddSelectNode(const FString& InCPPType, cons
 	return Node;
 }
 
-URigVMPrototypeNode* URigVMController::AddPrototypeNode(const FName& InNotation, const FVector2D& InPosition, const FString& InNodeName, bool bUndo)
+URigVMPrototypeNode* URigVMController::AddPrototypeNode(const FName& InNotation, const FVector2D& InPosition, const FString& InNodeName, bool bSetupUndoRedo)
 {
 	if (!IsValidGraph())
 	{
@@ -4008,7 +4008,7 @@ URigVMPrototypeNode* URigVMController::AddPrototypeNode(const FName& InNotation,
 
 	Notify(ERigVMGraphNotifType::NodeAdded, Node);
 
-	if (bUndo)
+	if (bSetupUndoRedo)
 	{
 		ActionStack->AddAction(FRigVMAddPrototypeNodeAction(Node));
 	}
@@ -4017,7 +4017,7 @@ URigVMPrototypeNode* URigVMController::AddPrototypeNode(const FName& InNotation,
 }
 
 
-URigVMEnumNode* URigVMController::AddEnumNode(const FName& InCPPTypeObjectPath, const FVector2D& InPosition, const FString& InNodeName, bool bUndo)
+URigVMEnumNode* URigVMController::AddEnumNode(const FName& InCPPTypeObjectPath, const FVector2D& InPosition, const FString& InNodeName, bool bSetupUndoRedo)
 {
 	if (!IsValidGraph())
 	{
@@ -4060,7 +4060,7 @@ URigVMEnumNode* URigVMController::AddEnumNode(const FName& InCPPTypeObjectPath,
 
 	Notify(ERigVMGraphNotifType::NodeAdded, Node);
 
-	if (bUndo)
+	if (bSetupUndoRedo)
 	{
 		ActionStack->AddAction(FRigVMAddEnumNodeAction(Node));
 	}
@@ -5226,13 +5226,13 @@ void URigVMController::PostProcessDefaultValue(URigVMPin* Pin, FString& OutDefau
 	}
 }
 
-void URigVMController::PotentiallyResolvePrototypeNode(URigVMPrototypeNode* InNode, bool bUndo)
+void URigVMController::PotentiallyResolvePrototypeNode(URigVMPrototypeNode* InNode, bool bSetupUndoRedo)
 {
 	TArray<URigVMNode*> NodesVisited;
-	PotentiallyResolvePrototypeNode(InNode, bUndo, NodesVisited);
+	PotentiallyResolvePrototypeNode(InNode, bSetupUndoRedo, NodesVisited);
 }
 
-void URigVMController::PotentiallyResolvePrototypeNode(URigVMPrototypeNode* InNode, bool bUndo, TArray<URigVMNode*>& NodesVisited)
+void URigVMController::PotentiallyResolvePrototypeNode(URigVMPrototypeNode* InNode, bool bSetupUndoRedo, TArray<URigVMNode*>& NodesVisited)
 {
 	if (InNode == nullptr)
 	{
@@ -5257,7 +5257,7 @@ void URigVMController::PotentiallyResolvePrototypeNode(URigVMPrototypeNode* InNo
 			{
 				if (!LinkedPin->CPPType.IsEmpty())
 				{
-					ChangePinType(Pin, LinkedPin->CPPType, LinkedPin->CPPTypeObjectPath, bUndo);
+					ChangePinType(Pin, LinkedPin->CPPType, LinkedPin->CPPTypeObjectPath, bSetupUndoRedo);
 					break;
 				}
 			}
@@ -5295,20 +5295,20 @@ void URigVMController::PotentiallyResolvePrototypeNode(URigVMPrototypeNode* InNo
 		FString NodeName = InNode->GetName();
 		FVector2D NodePosition = InNode->GetPosition();
 
-		RemoveNode(InNode, bUndo);
+		RemoveNode(InNode, bSetupUndoRedo);
 
-		if (URigVMNode* NewNode = AddStructNode(Function.Struct, Function.GetMethodName(), NodePosition, NodeName, bUndo))
+		if (URigVMNode* NewNode = AddStructNode(Function.Struct, Function.GetMethodName(), NodePosition, NodeName, bSetupUndoRedo))
 		{
 			// set default values again
 			for (TPair<FString, FString> Pair : DefaultValues)
 			{
-				SetPinDefaultValue(Pair.Key, Pair.Value, true, bUndo, false);
+				SetPinDefaultValue(Pair.Key, Pair.Value, true, bSetupUndoRedo, false);
 			}
 
 			// reestablish links
 			for (TPair<FString, FString> Pair : LinkPaths)
 			{
-				AddLink(Pair.Key, Pair.Value, bUndo);
+				AddLink(Pair.Key, Pair.Value, bSetupUndoRedo);
 			}
 		}
 
@@ -5325,7 +5325,7 @@ void URigVMController::PotentiallyResolvePrototypeNode(URigVMPrototypeNode* InNo
 				{
 					if (!Type->CPPType.IsEmpty())
 					{
-						ChangePinType(Pin, Type->CPPType, Type->GetCPPTypeObjectPath(), bUndo);
+						ChangePinType(Pin, Type->CPPType, Type->GetCPPTypeObjectPath(), bSetupUndoRedo);
 					}
 				}
 			}
@@ -5337,11 +5337,11 @@ void URigVMController::PotentiallyResolvePrototypeNode(URigVMPrototypeNode* InNo
 	LinkedNodes.Append(InNode->GetLinkedTargetNodes());
 	for (URigVMNode* LinkedNode : LinkedNodes)
 	{
-		PotentiallyResolvePrototypeNode(Cast<URigVMPrototypeNode>(LinkedNode), bUndo, NodesVisited);
+		PotentiallyResolvePrototypeNode(Cast<URigVMPrototypeNode>(LinkedNode), bSetupUndoRedo, NodesVisited);
 	}
 }
 
-bool URigVMController::ChangePinType(const FString& InPinPath, const FString& InCPPType, const FName& InCPPTypeObjectPath, bool bUndo)
+bool URigVMController::ChangePinType(const FString& InPinPath, const FString& InCPPType, const FName& InCPPTypeObjectPath, bool bSetupUndoRedo)
 {
 	if (!IsValidGraph())
 	{
@@ -5350,13 +5350,13 @@ bool URigVMController::ChangePinType(const FString& InPinPath, const FString& In
 
 	if (URigVMPin* Pin = Graph->FindPin(InPinPath))
 	{
-		return ChangePinType(Pin, InCPPType, InCPPTypeObjectPath, bUndo);
+		return ChangePinType(Pin, InCPPType, InCPPTypeObjectPath, bSetupUndoRedo);
 	}
 
 	return false;
 }
 
-bool URigVMController::ChangePinType(URigVMPin* InPin, const FString& InCPPType, const FName& InCPPTypeObjectPath, bool bUndo)
+bool URigVMController::ChangePinType(URigVMPin* InPin, const FString& InCPPType, const FName& InCPPTypeObjectPath, bool bSetupUndoRedo)
 {
 	if (InPin->CPPType == InCPPType)
 	{
@@ -5371,7 +5371,7 @@ bool URigVMController::ChangePinType(URigVMPin* InPin, const FString& InCPPType,
 
 	Notify(ERigVMGraphNotifType::PinTypeChanged, InPin);
 
-	if (bUndo)
+	if (bSetupUndoRedo)
 	{
 		ActionStack->AddAction(Action);
 	}
@@ -5381,7 +5381,7 @@ bool URigVMController::ChangePinType(URigVMPin* InPin, const FString& InCPPType,
 
 #if WITH_EDITOR
 
-void URigVMController::RewireLinks(URigVMPin* InOldPin, URigVMPin* InNewPin, bool bAsInput, bool bUndo, TArray<URigVMLink*> InLinks)
+void URigVMController::RewireLinks(URigVMPin* InOldPin, URigVMPin* InNewPin, bool bAsInput, bool bSetupUndoRedo, TArray<URigVMLink*> InLinks)
 {
 	ensure(InOldPin->GetRootPin() == InOldPin);
 	ensure(InNewPin->GetRootPin() == InNewPin);
diff --git a/Engine/Source/Developer/RigVMDeveloper/Public/RigVMModel/RigVMController.h b/Engine/Source/Developer/RigVMDeveloper/Public/RigVMModel/RigVMController.h
--- a/Engine/Source/Developer/RigVMDeveloper/Public/RigVMModel/RigVMController.h
+++ b/Engine/Source/Developer/RigVMDeveloper/Public/RigVMModel/RigVMController.h
@@ -84,133 +84,133 @@ public:
 	// StructNode represent a RIGVM_METHOD declaration on a USTRUCT.
 	// This causes a NodeAdded modified event.
 	UFUNCTION(BlueprintCallable, Category = RigVMController)
-	URigVMStructNode* AddStructNode(UScriptStruct* InScriptStruct, const FName& InMethodName, const FVector2D& InPosition = FVector2D::ZeroVector, const FString& InNodeName = TEXT(""), bool bUndo = true);
+	URigVMStructNode* AddStructNode(UScriptStruct* InScriptStruct, const FName& InMethodName, const FVector2D& InPosition = FVector2D::ZeroVector, const FString& InNodeName = TEXT(""), bool bSetupUndoRedo = true);
 
 	// Adds a Function / Struct Node to the edited Graph given its struct object path name.
 	// StructNode represent a RIGVM_METHOD declaration on a USTRUCT.
 	// This causes a NodeAdded modified event.
 	UFUNCTION(BlueprintCallable, Category = RigVMController)
-	URigVMStructNode* AddStructNodeFromStructPath(const FString& InScriptStructPath, const FName& InMethodName, const FVector2D& InPosition = FVector2D::ZeroVector, const FString& InNodeName = TEXT(""), bool bUndo = true);
+	URigVMStructNode* AddStructNodeFromStructPath(const FString& InScriptStructPath, const FName& InMethodName, const FVector2D& InPosition = FVector2D::ZeroVector, const FString& InNodeName = TEXT(""), bool bSetupUndoRedo = true);
 
 	// Adds a Variable Node to the edited Graph.
 	// Variables represent local work state for the function and
 	// can be read from and written to. 
 	// This causes a NodeAdded modified event.
 	UFUNCTION(BlueprintCallable, Category = RigVMController)
-	URigVMVariableNode* AddVariableNode(const FName& InVariableName, const FString& InCPPType, UObject* InCPPTypeObject, bool bIsGetter, const FString& InDefaultValue, const FVector2D& InPosition = FVector2D::ZeroVector, const FString& InNodeName = TEXT(""), bool bUndo = true);
+	URigVMVariableNode* AddVariableNode(const FName& InVariableName, const FString& InCPPType, UObject* InCPPTypeObject, bool bIsGetter, const FString& InDefaultValue, const FVector2D& InPosition = FVector2D::ZeroVector, const FString& InNodeName = TEXT(""), bool bSetupUndoRedo = true);
 
 	// Adds a Variable Node to the edited Graph given a struct object path name.
 	// Variables represent local work state for the function and
 	// can be read from (bIsGetter == true) or written to (bIsGetter == false). 
 	// This causes a NodeAdded modified event.
 	UFUNCTION(BlueprintCallable, Category = RigVMController)
-	URigVMVariableNode* AddVariableNodeFromObjectPath(const FName& InVariableName, const FString& InCPPType, const FString& InCPPTypeObjectPath, bool bIsGetter, const FString& InDefaultValue, const FVector2D& InPosition = FVector2D::ZeroVector, const FString& InNodeName = TEXT(""), bool bUndo = true);
+	URigVMVariableNode* AddVariableNodeFromObjectPath(const FName& InVariableName, const FString& InCPPType, const FString& InCPPTypeObjectPath, bool bIsGetter, const FString& InDefaultValue, const FVector2D& InPosition = FVector2D::ZeroVector, const FString& InNodeName = TEXT(""), bool bSetupUndoRedo = true);
 
 	// Refreshes the variable node with the new data
 	UFUNCTION(BlueprintCallable, Category = RigVMController)
-	void RefreshVariableNode(const FName& InNodeName, const FName& InVariableName, const FString& InCPPType, UObject* InCPPTypeObject, bool bUndo);
+	void RefreshVariableNode(const FName& InNodeName, const FName& InVariableName, const FString& InCPPType, UObject* InCPPTypeObject, bool bSetupUndoRedo);
 
 	// Removes all nodes related to a given variable
 	UFUNCTION(BlueprintCallable, Category = RigVMController)
-	void OnExternalVariableRemoved(const FName& InVarName, bool bUndo);
+	void OnExternalVariableRemoved(const FName& InVarName, bool bSetupUndoRedo);
 
 	// Renames the variable name in all relevant nodes
 	UFUNCTION(BlueprintCallable, Category = RigVMController)
-	void OnExternalVariableRenamed(const FName& InOldVarName, const FName& InNewVarName, bool bUndo);
+	void OnExternalVariableRenamed(const FName& InOldVarName, const FName& InNewVarName, bool bSetupUndoRedo);
 
 	// Changes the data type of all nodes matching a given variable name
 	UFUNCTION(BlueprintCallable, Category = RigVMController)
-	void OnExternalVariableTypeChanged(const FName& InVarName, const FString& InCPPType, UObject* InCPPTypeObject, bool bUndo);
+	void OnExternalVariableTypeChanged(const FName& InVarName, const FString& InCPPType, UObject* InCPPTypeObject, bool bSetupUndoRedo);
 
 	// Refreshes the variable node with the new data
 	UFUNCTION(BlueprintCallable, Category = RigVMController)
-	URigVMVariableNode* ReplaceParameterNodeWithVariable(const FName& InNodeName, const FName& InVariableName, const FString& InCPPType, UObject* InCPPTypeObject, bool bUndo);
+	URigVMVariableNode* ReplaceParameterNodeWithVariable(const FName& InNodeName, const FName& InVariableName, const FString& InCPPType, UObject* InCPPTypeObject, bool bSetupUndoRedo);
 
 	// Adds a Parameter Node to the edited Graph.
 	// Parameters represent input or output arguments to the Graph / Function.
 	// Input Parameters are constant values / literals.
 	// This causes a NodeAdded modified event.
 	UFUNCTION(BlueprintCallable, Category = RigVMController)
-	URigVMParameterNode* AddParameterNode(const FName& InParameterName, const FString& InCPPType, UObject* InCPPTypeObject, bool bIsInput, const FString& InDefaultValue, const FVector2D& InPosition = FVector2D::ZeroVector, const FString& InNodeName = TEXT(""), bool bUndo = true);
+	URigVMParameterNode* AddParameterNode(const FName& InParameterName, const FString& InCPPType, UObject* InCPPTypeObject, bool bIsInput, const FString& InDefaultValue, const FVector2D& InPosition = FVector2D::ZeroVector, const FString& InNodeName = TEXT(""), bool bSetupUndoRedo = true);
 
 	// Adds a Parameter Node to the edited Graph given a struct object path name.
 	// Parameters represent input or output arguments to the Graph / Function.
 	// Input Parameters are constant values / literals.
 	// This causes a NodeAdded modified event.
 	UFUNCTION(BlueprintCallable, Category = RigVMController)
-	URigVMParameterNode* AddParameterNodeFromObjectPath(const FName& InParameterName, const FString& InCPPType, const FString& InCPPTypeObjectPath, bool bIsInput, const FString& InDefaultValue, const FVector2D& InPosition = FVector2D::ZeroVector, const FString& InNodeName = TEXT(""), bool bUndo = true);
+	URigVMParameterNode* AddParameterNodeFromObjectPath(const FName& InParameterName, const FString& InCPPType, const FString& InCPPTypeObjectPath, bool bIsInput, const FString& InDefaultValue, const FVector2D& InPosition = FVector2D::ZeroVector, const FString& InNodeName = TEXT(""), bool bSetupUndoRedo = true);
 
 	// Adds a Comment Node to the edited Graph.
 	// Comments can be used to annotate the Graph.
 	// This causes a NodeAdded modified event.
 	UFUNCTION(BlueprintCallable, Category = RigVMController)
-	URigVMCommentNode* AddCommentNode(const FString& InCommentText, const FVector2D& InPosition = FVector2D::ZeroVector, const FVector2D& InSize = FVector2D(400.f, 300.f), const FLinearColor& InColor = FLinearColor::Black, const FString& InNodeName = TEXT(""), bool bUndo = true);
+	URigVMCommentNode* AddCommentNode(const FString& InCommentText, const FVector2D& InPosition = FVector2D::ZeroVector, const FVector2D& InSize = FVector2D(400.f, 300.f), const FLinearColor& InColor = FLinearColor::Black, const FString& InNodeName = TEXT(""), bool bSetupUndoRedo = true);
 
 	// Adds a Reroute Node on an existing Link to the edited Graph.
 	// Reroute Nodes can be used to visually improve the data flow,
 	// they don't require any additional memory though and are purely
 	// cosmetic. This causes a NodeAdded modified event.
 	UFUNCTION(BlueprintCallable, Category = RigVMController)
-	URigVMRerouteNode* AddRerouteNodeOnLink(URigVMLink* InLink, bool bShowAsFullNode, const FVector2D& InPosition = FVector2D::ZeroVector, const FString& InNodeName = TEXT(""), bool bUndo = true);
+	URigVMRerouteNode* AddRerouteNodeOnLink(URigVMLink* InLink, bool bShowAsFullNode, const FVector2D& InPosition = FVector2D::ZeroVector, const FString& InNodeName = TEXT(""), bool bSetupUndoRedo = true);
 
 	// Adds a Reroute Node on an existing Link to the edited Graph given the Link's string representation.
 	// Reroute Nodes can be used to visually improve the data flow,
 	// they don't require any additional memory though and are purely
 	// cosmetic. This causes a NodeAdded modified event.
 	UFUNCTION(BlueprintCallable, Category = RigVMController)
-	URigVMRerouteNode* AddRerouteNodeOnLinkPath(const FString& InLinkPinPathRepresentation, bool bShowAsFullNode, const FVector2D& InPosition = FVector2D::ZeroVector, const FString& InNodeName = TEXT(""), bool bUndo = true);
+	URigVMRerouteNode* AddRerouteNodeOnLinkPath(const FString& InLinkPinPathRepresentation, bool bShowAsFullNode, const FVector2D& InPosition = FVector2D::ZeroVector, const FString& InNodeName = TEXT(""), bool bSetupUndoRedo = true);
 
 	// Adds a Reroute Node on an existing Pin to the editor Graph.
 	// Reroute Nodes can be used to visually improve the data flow,
 	// they don't require any additional memory though and are purely
 	// cosmetic. This causes a NodeAdded modified event.
 	UFUNCTION(BlueprintCallable, Category = RigVMController)
-	URigVMRerouteNode* AddRerouteNodeOnPin(const FString& InPinPath, bool bAsInput, bool bShowAsFullNode, const FVector2D& InPosition = FVector2D::ZeroVector, const FString& InNodeName = TEXT(""), bool bUndo = true);
+	URigVMRerouteNode* AddRerouteNodeOnPin(const FString& InPinPath, bool bAsInput, bool bShowAsFullNode, const FVector2D& InPosition = FVector2D::ZeroVector, const FString& InNodeName = TEXT(""), bool bSetupUndoRedo = true);
 
 	// Adds a free Reroute Node
 	UFUNCTION(BlueprintCallable, Category = RigVMController)
-	URigVMRerouteNode* AddFreeRerouteNode(bool bShowAsFullNode, const FString& InCPPType, const FName& InCPPTypeObjectPath, bool bIsConstant, const FName& InCustomWidgetName, const FString& InDefaultValue, const FVector2D& InPosition = FVector2D::ZeroVector, const FString& InNodeName = TEXT(""), bool bUndo = true);
+	URigVMRerouteNode* AddFreeRerouteNode(bool bShowAsFullNode, const FString& InCPPType, const FName& InCPPTypeObjectPath, bool bIsConstant, const FName& InCustomWidgetName, const FString& InDefaultValue, const FVector2D& InPosition = FVector2D::ZeroVector, const FString& InNodeName = TEXT(""), bool bSetupUndoRedo = true);
 
 	// Adds a branch node to the graph.
 	// Branch nodes can be used to split the execution of into multiple branches,
 	// allowing to drive behavior by logic.
 	UFUNCTION(BlueprintCallable, Category = RigVMController)
-	URigVMBranchNode* AddBranchNode(const FVector2D& InPosition = FVector2D::ZeroVector, const FString& InNodeName = TEXT(""), bool bUndo = true);
+	URigVMBranchNode* AddBranchNode(const FVector2D& InPosition = FVector2D::ZeroVector, const FString& InNodeName = TEXT(""), bool bSetupUndoRedo = true);
 
 	// Adds an if node to the graph.
 	// If nodes can be used to pick between two values based on a condition.
 	UFUNCTION(BlueprintCallable, Category = RigVMController)
-	URigVMIfNode* AddIfNode(const FString& InCPPType, const FName& InCPPTypeObjectPath, const FVector2D& InPosition = FVector2D::ZeroVector, const FString& InNodeName = TEXT(""), bool bUndo = true);
+	URigVMIfNode* AddIfNode(const FString& InCPPType, const FName& InCPPTypeObjectPath, const FVector2D& InPosition = FVector2D::ZeroVector, const FString& InNodeName = TEXT(""), bool bSetupUndoRedo = true);
 
 	// Adds a select node to the graph.
 	// Select nodes can be used to pick between multiple values based on an index.
 	UFUNCTION(BlueprintCallable, Category = RigVMController)
-	URigVMSelectNode* AddSelectNode(const FString& InCPPType, const FName& InCPPTypeObjectPath, const FVector2D& InPosition = FVector2D::ZeroVector, const FString& InNodeName = TEXT(""), bool bUndo = true);
+	URigVMSelectNode* AddSelectNode(const FString& InCPPType, const FName& InCPPTypeObjectPath, const FVector2D& InPosition = FVector2D::ZeroVector, const FString& InNodeName = TEXT(""), bool bSetupUndoRedo = true);
 
 	// Adds a prototype node to the graph.
 	UFUNCTION(BlueprintCallable, Category = RigVMController)
-	URigVMPrototypeNode* AddPrototypeNode(const FName& InNotation, const FVector2D& InPosition = FVector2D::ZeroVector, const FString& InNodeName = TEXT(""), bool bUndo = true);
+	URigVMPrototypeNode* AddPrototypeNode(const FName& InNotation, const FVector2D& InPosition = FVector2D::ZeroVector, const FString& InNodeName = TEXT(""), bool bSetupUndoRedo = true);
 
 	// Adds a Function / Struct Node to the edited Graph as an injected node
 	// StructNode represent a RIGVM_METHOD declaration on a USTRUCT.
 	// This causes a NodeAdded modified event.
 	UFUNCTION(BlueprintCallable, Category = RigVMController)
-	URigVMInjectionInfo* AddInjectedNode(const FString& InPinPath, bool bAsInput, UScriptStruct* InScriptStruct, const FName& InMethodName, const FName& InInputPinName, const FName& InOutputPinName, const FString& InNodeName = TEXT(""), bool bUndo = true);
+	URigVMInjectionInfo* AddInjectedNode(const FString& InPinPath, bool bAsInput, UScriptStruct* InScriptStruct, const FName& InMethodName, const FName& InInputPinName, const FName& InOutputPinName, const FString& InNodeName = TEXT(""), bool bSetupUndoRedo = true);
 
 	// Adds a Function / Struct Node to the edited Graph as an injected node
 	// StructNode represent a RIGVM_METHOD declaration on a USTRUCT.
 	// This causes a NodeAdded modified event.
 	UFUNCTION(BlueprintCallable, Category = RigVMController)
-	URigVMInjectionInfo* AddInjectedNodeFromStructPath(const FString& InPinPath, bool bAsInput, const FString& InScriptStructPath, const FName& InMethodName, const FName& InInputPinName, const FName& InOutputPinName, const FString& InNodeName = TEXT(""), bool bUndo = true);
+	URigVMInjectionInfo* AddInjectedNodeFromStructPath(const FString& InPinPath, bool bAsInput, const FString& InScriptStructPath, const FName& InMethodName, const FName& InInputPinName, const FName& InOutputPinName, const FString& InNodeName = TEXT(""), bool bSetupUndoRedo = true);
 
 	// Ejects the last injected node on a pin
 	UFUNCTION(BlueprintCallable, Category = RigVMController)
-	URigVMNode* EjectNodeFromPin(const FString& InPinPath, bool bUndo = true);
+	URigVMNode* EjectNodeFromPin(const FString& InPinPath, bool bSetupUndoRedo = true);
 
 	// Adds an enum node to the graph
 	// Enum nodes can be used to represent constant enum values within the graph
 	UFUNCTION(BlueprintCallable, Category = RigVMController)
-	URigVMEnumNode* AddEnumNode(const FName& InCPPTypeObjectPath, const FVector2D& InPosition = FVector2D::ZeroVector, const FString& InNodeName = TEXT(""), bool bUndo = true);
+	URigVMEnumNode* AddEnumNode(const FName& InCPPTypeObjectPath, const FVector2D& InPosition = FVector2D::ZeroVector, const FString& InNodeName = TEXT(""), bool bSetupUndoRedo = true);
 
 	// Un-does the last action on the stack.
 	// Note: This should really only be used for unit tests,
@@ -257,7 +257,7 @@ public:
 
 	// Exports the given nodes as text
 	UFUNCTION(BlueprintCallable, Category = RigVMController)
-	TArray<FName> ImportNodesFromText(const FString& InText, bool bUndo = true);
+	TArray<FName> ImportNodesFromText(const FString& InText, bool bSetupUndoRedo = true);
 
 	// Returns a unique name
 	static FName GetUniqueName(const FName& InName, TFunction<bool(const FName&)> IsNameAvailableFunction);
@@ -267,102 +267,102 @@ public:
 	// Removes a node from the graph
 	// This causes a NodeRemoved modified event.
 	UFUNCTION(BlueprintCallable, Category = RigVMController)
-	bool RemoveNode(URigVMNode* InNode, bool bUndo = true, bool bRecursive = false);
+	bool RemoveNode(URigVMNode* InNode, bool bSetupUndoRedo = true, bool bRecursive = false);
 
 	// Removes a node from the graph given the node's name.
 	// This causes a NodeRemoved modified event.
 	UFUNCTION(BlueprintCallable, Category = RigVMController)
-	bool RemoveNodeByName(const FName& InNodeName, bool bUndo = true, bool bRecursive = false);
+	bool RemoveNodeByName(const FName& InNodeName, bool bSetupUndoRedo = true, bool bRecursive = false);
 
 	// Selects a single node in the graph.
 	// This causes a NodeSelected / NodeDeselected modified event.
 	UFUNCTION(BlueprintCallable, Category = RigVMController)
-	bool SelectNode(URigVMNode* InNode, bool bSelect = true, bool bUndo = true);
+	bool SelectNode(URigVMNode* InNode, bool bSelect = true, bool bSetupUndoRedo = true);
 
 	// Selects a single node in the graph by name.
 	// This causes a NodeSelected / NodeDeselected modified event.
 	UFUNCTION(BlueprintCallable, Category = RigVMController)
-	bool SelectNodeByName(const FName& InNodeName, bool bSelect = true, bool bUndo = true);
+	bool SelectNodeByName(const FName& InNodeName, bool bSelect = true, bool bSetupUndoRedo = true);
 
 	// Deselects all currently selected nodes in the graph.
 	// This might cause several NodeDeselected modified event.
 	UFUNCTION(BlueprintCallable, Category = RigVMController)
-	bool ClearNodeSelection(bool bUndo = true);
+	bool ClearNodeSelection(bool bSetupUndoRedo = true);
 
 	// Selects the nodes given the selection
 	// This might cause several NodeDeselected modified event.
 	UFUNCTION(BlueprintCallable, Category = RigVMController)
-	bool SetNodeSelection(const TArray<FName>& InNodeNames, bool bUndo = true);
+	bool SetNodeSelection(const TArray<FName>& InNodeNames, bool bSetupUndoRedo = true);
 
 	// Sets the position of a node in the graph.
 	// This causes a NodePositionChanged modified event.
 	UFUNCTION(BlueprintCallable, Category = RigVMController)
-	bool SetNodePosition(URigVMNode* InNode, const FVector2D& InPosition, bool bUndo = true, bool bMergeUndoAction = false);
+	bool SetNodePosition(URigVMNode* InNode, const FVector2D& InPosition, bool bSetupUndoRedo = true, bool bMergeUndoAction = false);
 
 	// Sets the position of a node in the graph by name.
 	// This causes a NodePositionChanged modified event.
 	UFUNCTION(BlueprintCallable, Category = RigVMController)
-	bool SetNodePositionByName(const FName& InNodeName, const FVector2D& InPosition, bool bUndo = true, bool bMergeUndoAction = false);
+	bool SetNodePositionByName(const FName& InNodeName, const FVector2D& InPosition, bool bSetupUndoRedo = true, bool bMergeUndoAction = false);
 
 	// Sets the size of a node in the graph.
 	// This causes a NodeSizeChanged modified event.
 	UFUNCTION(BlueprintCallable, Category = RigVMController)
-	bool SetNodeSize(URigVMNode* InNode, const FVector2D& InSize, bool bUndo = true, bool bMergeUndoAction = false);
+	bool SetNodeSize(URigVMNode* InNode, const FVector2D& InSize, bool bSetupUndoRedo = true, bool bMergeUndoAction = false);
 
 	// Sets the size of a node in the graph by name.
 	// This causes a NodeSizeChanged modified event.
 	UFUNCTION(BlueprintCallable, Category = RigVMController)
-	bool SetNodeSizeByName(const FName& InNodeName, const FVector2D& InSize, bool bUndo = true, bool bMergeUndoAction = false);
+	bool SetNodeSizeByName(const FName& InNodeName, const FVector2D& InSize, bool bSetupUndoRedo = true, bool bMergeUndoAction = false);
 
 	// Sets the color of a node in the graph.
 	// This causes a NodeColorChanged modified event.
 	UFUNCTION(BlueprintCallable, Category = RigVMController)
-	bool SetNodeColor(URigVMNode* InNode, const FLinearColor& InColor, bool bUndo = true, bool bMergeUndoAction = false);
+	bool SetNodeColor(URigVMNode* InNode, const FLinearColor& InColor, bool bSetupUndoRedo = true, bool bMergeUndoAction = false);
 
 	// Sets the color of a node in the graph by name.
 	// This causes a NodeColorChanged modified event.
 	UFUNCTION(BlueprintCallable, Category = RigVMController)
-	bool SetNodeColorByName(const FName& InNodeName, const FLinearColor& InColor, bool bUndo = true, bool bMergeUndoAction = false);
+	bool SetNodeColorByName(const FName& InNodeName, const FLinearColor& InColor, bool bSetupUndoRedo = true, bool bMergeUndoAction = false);
 
 	// Sets the comment text of a comment node in the graph.
 	// This causes a CommentTextChanged modified event.
 	UFUNCTION(BlueprintCallable, Category = RigVMController)
-	bool SetCommentText(URigVMNode* InNode, const FString& InCommentText, bool bUndo = true);
+	bool SetCommentText(URigVMNode* InNode, const FString& InCommentText, bool bSetupUndoRedo = true);
 
 	// Sets the comment text of a comment node in the graph by name.
 	// This causes a CommentTextChanged modified event.
 	UFUNCTION(BlueprintCallable, Category = RigVMController)
-	bool SetCommentTextByName(const FName& InNodeName, const FString& InCommentText, bool bUndo = true);
+	bool SetCommentTextByName(const FName& InNodeName, const FString& InCommentText, bool bSetupUndoRedo = true);
 
 	// Sets the compactness of a reroute node in the graph.
 	// This causes a RerouteCompactnessChanged modified event.
 	UFUNCTION(BlueprintCallable, Category = RigVMController)
-	bool SetRerouteCompactness(URigVMNode* InNode, bool bShowAsFullNode, bool bUndo = true);
+	bool SetRerouteCompactness(URigVMNode* InNode, bool bShowAsFullNode, bool bSetupUndoRedo = true);
 
 	// Sets the compactness of a reroute node in the graph by name.
 	// This causes a RerouteCompactnessChanged modified event.
 	UFUNCTION(BlueprintCallable, Category = RigVMController)
-	bool SetRerouteCompactnessByName(const FName& InNodeName, bool bShowAsFullNode, bool bUndo = true);
+	bool SetRerouteCompactnessByName(const FName& InNodeName, bool bShowAsFullNode, bool bSetupUndoRedo = true);
 
 	// Renames a variable in the graph.
 	// This causes a VariableRenamed modified event.
 	UFUNCTION(BlueprintCallable, Category = RigVMController)
-	bool RenameVariable(const FName& InOldName, const FName& InNewName, bool bUndo = true);
+	bool RenameVariable(const FName& InOldName, const FName& InNewName, bool bSetupUndoRedo = true);
 
 	// Renames a parameter in the graph.
 	// This causes a ParameterRenamed modified event.
 	UFUNCTION(BlueprintCallable, Category = RigVMController)
-	bool RenameParameter(const FName& InOldName, const FName& InNewName, bool bUndo = true);
+	bool RenameParameter(const FName& InOldName, const FName& InNewName, bool bSetupUndoRedo = true);
 
 	// Sets the pin to be expanded or not
 	// This causes a PinExpansionChanged modified event.
 	UFUNCTION(BlueprintCallable, Category = RigVMController)
-	bool SetPinExpansion(const FString& InPinPath, bool bIsExpanded, bool bUndo = true);
+	bool SetPinExpansion(const FString& InPinPath, bool bIsExpanded, bool bSetupUndoRedo = true);
 
 	// Sets the pin to be watched (or not)
 	// This causes a PinWatchedChanged modified event.
 	UFUNCTION(BlueprintCallable, Category = RigVMController)
-	bool SetPinIsWatched(const FString& InPinPath, bool bIsWatched, bool bUndo = true);
+	bool SetPinIsWatched(const FString& InPinPath, bool bIsWatched, bool bSetupUndoRedo = true);
 
 	// Returns the default value of a pin given its pinpath.
 	UFUNCTION(BlueprintCallable, Category = RigVMController)
@@ -371,79 +371,79 @@ public:
 	// Sets the default value of a pin given its pinpath.
 	// This causes a PinDefaultValueChanged modified event.
 	UFUNCTION(BlueprintCallable, Category = RigVMController)
-	bool SetPinDefaultValue(const FString& InPinPath, const FString& InDefaultValue, bool bResizeArrays = true, bool bUndo = true, bool bMergeUndoAction = false);
+	bool SetPinDefaultValue(const FString& InPinPath, const FString& InDefaultValue, bool bResizeArrays = true, bool bSetupUndoRedo = true, bool bMergeUndoAction = false);
 
 	// Resets the default value of a pin given its pinpath.
 	// This causes a PinDefaultValueChanged modified event.
 	UFUNCTION(BlueprintCallable, Category = RigVMController)
-	bool ResetPinDefaultValue(const FString& InPinPath, bool bUndo = true);
+	bool ResetPinDefaultValue(const FString& InPinPath, bool bSetupUndoRedo = true);
 
 	// Adds an array element pin to the end of an array pin.
 	// This causes a PinArraySizeChanged modified event.
 	UFUNCTION(BlueprintCallable, Category = RigVMController)
-	FString AddArrayPin(const FString& InArrayPinPath, const FString& InDefaultValue = TEXT(""), bool bUndo = true);
+	FString AddArrayPin(const FString& InArrayPinPath, const FString& InDefaultValue = TEXT(""), bool bSetupUndoRedo = true);
 
 	// Duplicates an array element pin.
 	// This causes a PinArraySizeChanged modified event.
 	UFUNCTION(BlueprintCallable, Category = RigVMController)
-	FString DuplicateArrayPin(const FString& InArrayElementPinPath, bool bUndo = true);
+	FString DuplicateArrayPin(const FString& InArrayElementPinPath, bool bSetupUndoRedo = true);
 
 	// Inserts an array element pin into an array pin.
 	// This causes a PinArraySizeChanged modified event.
 	UFUNCTION(BlueprintCallable, Category = RigVMController)
-	FString InsertArrayPin(const FString& InArrayPinPath, int32 InIndex = -1, const FString& InDefaultValue = TEXT(""), bool bUndo = true);
+	FString InsertArrayPin(const FString& InArrayPinPath, int32 InIndex = -1, const FString& InDefaultValue = TEXT(""), bool bSetupUndoRedo = true);
 
 	// Removes an array element pin from an array pin.
 	// This causes a PinArraySizeChanged modified event.
 	UFUNCTION(BlueprintCallable, Category = RigVMController)
-	bool RemoveArrayPin(const FString& InArrayElementPinPath, bool bUndo = true);
+	bool RemoveArrayPin(const FString& InArrayElementPinPath, bool bSetupUndoRedo = true);
 
 	// Removes all (but one) array element pin from an array pin.
 	// This causes a PinArraySizeChanged modified event.
 	UFUNCTION(BlueprintCallable, Category = RigVMController)
-	bool ClearArrayPin(const FString& InArrayPinPath, bool bUndo = true);
+	bool ClearArrayPin(const FString& InArrayPinPath, bool bSetupUndoRedo = true);
 
 	// Sets the size of the array pin
 	// This causes a PinArraySizeChanged modified event.
 	UFUNCTION(BlueprintCallable, Category = RigVMController)
-	bool SetArrayPinSize(const FString& InArrayPinPath, int32 InSize, const FString& InDefaultValue = TEXT(""), bool bUndo = true);
+	bool SetArrayPinSize(const FString& InArrayPinPath, int32 InSize, const FString& InDefaultValue = TEXT(""), bool bSetupUndoRedo = true);
 
 	// Binds a pin to a variable (or removes the binding given NAME_None)
 	// This causes a PinBoundVariableChanged modified event.
 	UFUNCTION(BlueprintCallable, Category = RigVMController)
-	bool BindPinToVariable(const FString& InPinPath, const FString& InNewBoundVariablePath, bool bUndo = true);
+	bool BindPinToVariable(const FString& InPinPath, const FString& InNewBoundVariablePath, bool bSetupUndoRedo = true);
 
 	// Removes the binging of a pin to a variable
 	// This causes a PinBoundVariableChanged modified event.
 	UFUNCTION(BlueprintCallable, Category = RigVMController)
-	bool UnbindPinFromVariable(const FString& InPinPath, bool bUndo = true);
+	bool UnbindPinFromVariable(const FString& InPinPath, bool bSetupUndoRedo = true);
 
 	// Turns a variable node into one or more bindings
 	UFUNCTION(BlueprintCallable, Category = RigVMController)
-	bool MakeBindingsFromVariableNode(const FName& InNodeName, bool bUndo = true);
+	bool MakeBindingsFromVariableNode(const FName& InNodeName, bool bSetupUndoRedo = true);
 
 	// Turns a binding to a variable node
 	UFUNCTION(BlueprintCallable, Category = RigVMController)
-	bool MakeVariableNodeFromBinding(const FString& InPinPath, const FVector2D& InNodePosition = FVector2D::ZeroVector, bool bUndo = true);
+	bool MakeVariableNodeFromBinding(const FString& InPinPath, const FVector2D& InNodePosition = FVector2D::ZeroVector, bool bSetupUndoRedo = true);
 
 	// Promotes a pin to a variable
 	UFUNCTION(BlueprintCallable, Category = RigVMController)
-	bool PromotePinToVariable(const FString& InPinPath, bool bCreateVariableNode, const FVector2D& InNodePosition = FVector2D::ZeroVector, bool bUndo = true);
+	bool PromotePinToVariable(const FString& InPinPath, bool bCreateVariableNode, const FVector2D& InNodePosition = FVector2D::ZeroVector, bool bSetupUndoRedo = true);
 
 	// Adds a link to the graph.
 	// This causes a LinkAdded modified event.
 	UFUNCTION(BlueprintCallable, Category = RigVMController)
-	bool AddLink(const FString& InOutputPinPath, const FString& InInputPinPath, bool bUndo = true);
+	bool AddLink(const FString& InOutputPinPath, const FString& InInputPinPath, bool bSetupUndoRedo = true);
 
 	// Removes a link from the graph.
 	// This causes a LinkRemoved modified event.
 	UFUNCTION(BlueprintCallable, Category = RigVMController)
-	bool BreakLink(const FString& InOutputPinPath, const FString& InInputPinPath, bool bUndo = true);
+	bool BreakLink(const FString& InOutputPinPath, const FString& InInputPinPath, bool bSetupUndoRedo = true);
 
 	// Removes all links on a given pin from the graph.
 	// This might cause multiple LinkRemoved modified event.
 	UFUNCTION(BlueprintCallable, Category = RigVMController)
-	bool BreakAllLinks(const FString& InPinPath, bool bAsInput = true, bool bUndo = true);
+	bool BreakAllLinks(const FString& InPinPath, bool bAsInput = true, bool bSetupUndoRedo = true);
 
 	// Sets the execute context struct type to use
 	void SetExecuteContextStruct(UStruct* InExecuteContextStruct);
@@ -495,25 +495,25 @@ private:
 	void ConfigurePinFromPin(URigVMPin* InOutPin, URigVMPin* InPin);
 	virtual bool ShouldStructBeUnfolded(const UStruct* InStruct);
 	virtual bool ShouldPinBeUnfolded(URigVMPin* InPin);
-	void SetPinDefaultValue(URigVMPin* InPin, const FString& InDefaultValue, bool bResizeArrays, bool bUndo, bool bMergeUndoAction);
-	bool ResetPinDefaultValue(URigVMPin* InPin, bool bUndo);
+	void SetPinDefaultValue(URigVMPin* InPin, const FString& InDefaultValue, bool bResizeArrays, bool bSetupUndoRedo, bool bMergeUndoAction);
+	bool ResetPinDefaultValue(URigVMPin* InPin, bool bSetupUndoRedo);
 	static TArray<FString> SplitDefaultValue(const FString& InDefaultValue);
-	URigVMPin* InsertArrayPin(URigVMPin* ArrayPin, int32 InIndex, const FString& InDefaultValue, bool bUndo);
-	bool RemovePin(URigVMPin* InPinToRemove, bool bUndo);
+	URigVMPin* InsertArrayPin(URigVMPin* ArrayPin, int32 InIndex, const FString& InDefaultValue, bool bSetupUndoRedo);
+	bool RemovePin(URigVMPin* InPinToRemove, bool bSetupUndoRedo);
 	FProperty* FindPropertyForPin(const FString& InPinPath);
-	bool BindPinToVariable(URigVMPin* InPin, const FString& InNewBoundVariablePath, bool bUndo);
-	bool MakeBindingsFromVariableNode(URigVMVariableNode* InNode, bool bUndo);
-	bool PromotePinToVariable(URigVMPin* InPin, bool bCreateVariableNode, const FVector2D& InNodePosition, bool bUndo);
-	bool AddLink(URigVMPin* OutputPin, URigVMPin* InputPin, bool bUndo);
-	bool BreakLink(URigVMPin* OutputPin, URigVMPin* InputPin, bool bUndo);
-	bool BreakAllLinks(URigVMPin* Pin, bool bAsInput, bool bUndo);
-	void BreakAllLinksRecursive(URigVMPin* Pin, bool bAsInput, bool bTowardsParent, bool bUndo);
-	void UpdateRerouteNodeAfterChangingLinks(URigVMPin* PinChanged, bool bUndo = true);
-	bool SetPinExpansion(URigVMPin* InPin, bool bIsExpanded, bool bUndo = true);
-	void ExpandPinRecursively(URigVMPin* InPin, bool bUndo);
-	bool SetPinIsWatched(URigVMPin* InPin, bool bIsWatched, bool bUndo);
-	bool SetVariableName(URigVMVariableNode* InVariableNode, const FName& InVariableName, bool bUndo);
-	bool SetParameterName(URigVMParameterNode* InParameterNode, const FName& InParameterName, bool bUndo);
+	bool BindPinToVariable(URigVMPin* InPin, const FString& InNewBoundVariablePath, bool bSetupUndoRedo);
+	bool MakeBindingsFromVariableNode(URigVMVariableNode* InNode, bool bSetupUndoRedo);
+	bool PromotePinToVariable(URigVMPin* InPin, bool bCreateVariableNode, const FVector2D& InNodePosition, bool bSetupUndoRedo);
+	bool AddLink(URigVMPin* OutputPin, URigVMPin* InputPin, bool bSetupUndoRedo);
+	bool BreakLink(URigVMPin* OutputPin, URigVMPin* InputPin, bool bSetupUndoRedo);
+	bool BreakAllLinks(URigVMPin* Pin, bool bAsInput, bool bSetupUndoRedo);
+	void BreakAllLinksRecursive(URigVMPin* Pin, bool bAsInput, bool bTowardsParent, bool bSetupUndoRedo);
+	void UpdateRerouteNodeAfterChangingLinks(URigVMPin* PinChanged, bool bSetupUndoRedo = true);
+	bool SetPinExpansion(URigVMPin* InPin, bool bIsExpanded, bool bSetupUndoRedo = true);
+	void ExpandPinRecursively(URigVMPin* InPin, bool bSetupUndoRedo);
+	bool SetPinIsWatched(URigVMPin* InPin, bool bIsWatched, bool bSetupUndoRedo);
+	bool SetVariableName(URigVMVariableNode* InVariableNode, const FName& InVariableName, bool bSetupUndoRedo);
+	bool SetParameterName(URigVMParameterNode* InParameterNode, const FName& InParameterName, bool bSetupUndoRedo);
 	static void ForEveryPinRecursively(URigVMPin* InPin, TFunction<void(URigVMPin*)> OnEachPinFunction);
 	static void ForEveryPinRecursively(URigVMNode* InNode, TFunction<void(URigVMPin*)> OnEachPinFunction);
 
@@ -536,13 +536,13 @@ private:
 	static void CreateDefaultValueForStructIfRequired(UScriptStruct* InStruct, FString& OutDefaultValue);
 	static void PostProcessDefaultValue(URigVMPin* Pin, FString& OutDefaultValue);
 
-	void PotentiallyResolvePrototypeNode(URigVMPrototypeNode* InNode, bool bUndo);
-	void PotentiallyResolvePrototypeNode(URigVMPrototypeNode* InNode, bool bUndo, TArray<URigVMNode*>& NodesVisited);
-	bool ChangePinType(const FString& InPinPath, const FString& InCPPType, const FName& InCPPTypeObjectPath, bool bUndo);
-	bool ChangePinType(URigVMPin* InPin, const FString& InCPPType, const FName& InCPPTypeObjectPath, bool bUndo);
+	void PotentiallyResolvePrototypeNode(URigVMPrototypeNode* InNode, bool bSetupUndoRedo);
+	void PotentiallyResolvePrototypeNode(URigVMPrototypeNode* InNode, bool bSetupUndoRedo, TArray<URigVMNode*>& NodesVisited);
+	bool ChangePinType(const FString& InPinPath, const FString& InCPPType, const FName& InCPPTypeObjectPath, bool bSetupUndoRedo);
+	bool ChangePinType(URigVMPin* InPin, const FString& InCPPType, const FName& InCPPTypeObjectPath, bool bSetupUndoRedo);
 
 #if WITH_EDITOR
-	void RewireLinks(URigVMPin* OldPin, URigVMPin* NewPin, bool bAsInput, bool bUndo, TArray<URigVMLink*> InLinks = TArray<URigVMLink*>());
+	void RewireLinks(URigVMPin* OldPin, URigVMPin* NewPin, bool bAsInput, bool bSetupUndoRedo, TArray<URigVMLink*> InLinks = TArray<URigVMLink*>());
 #endif
 
 	void DestroyObject(UObject* InObjectToDestroy);
