Control Rig: Fix making variable node out of variable binding with segmentpaths

diff --git a/Engine/Source/Developer/RigVMDeveloper/Private/RigVMModel/RigVMController.cpp b/Engine/Source/Developer/RigVMDeveloper/Private/RigVMModel/RigVMController.cpp
--- a/Engine/Source/Developer/RigVMDeveloper/Private/RigVMModel/RigVMController.cpp
+++ b/Engine/Source/Developer/RigVMDeveloper/Private/RigVMModel/RigVMController.cpp
@@ -3188,10 +3188,20 @@ bool URigVMController::PromotePinToVariable(URigVMPin* InPin, bool bCreateVariab
 	}
 
 	FRigVMExternalVariable VariableForPin;
+	FString SegmentPath;
 	if (InPin->IsBoundToVariable())
 	{
 		VariableForPin = GetExternalVariableByName(*InPin->GetBoundVariableName());
 		check(VariableForPin.IsValid(true /* allow nullptr */));
+		SegmentPath = InPin->GetBoundVariablePath();
+		if (SegmentPath.StartsWith(VariableForPin.Name.ToString() + TEXT(".")))
+		{
+			SegmentPath = SegmentPath.RightChop(VariableForPin.Name.ToString().Len());
+		}
+		else
+		{
+			SegmentPath.Empty();
+		}
 	}
 	else
 	{
@@ -3228,7 +3238,7 @@ bool URigVMController::PromotePinToVariable(URigVMPin* InPin, bool bCreateVariab
 		{
 			if (URigVMPin* ValuePin = VariableNode->FindPin(URigVMVariableNode::ValueName))
 			{
-				return AddLink(ValuePin, InPin, bSetupUndoRedo);
+				return AddLink(ValuePin->GetPinPath() + SegmentPath, InPin->GetPinPath(), bSetupUndoRedo);
 			}
 		}
 	}
