Control Rig: Clean up reroute nodes when making bindings from variable node

diff --git a/Engine/Source/Developer/RigVMDeveloper/Private/RigVMModel/RigVMController.cpp b/Engine/Source/Developer/RigVMDeveloper/Private/RigVMModel/RigVMController.cpp
--- a/Engine/Source/Developer/RigVMDeveloper/Private/RigVMModel/RigVMController.cpp
+++ b/Engine/Source/Developer/RigVMDeveloper/Private/RigVMModel/RigVMController.cpp
@@ -3106,10 +3106,34 @@ bool URigVMController::MakeBindingsFromVariableNode(URigVMVariableNode* InNode,
 {
 	check(InNode);
 
-	TArray<URigVMLink*> Links;
+	TArray<TPair<URigVMPin*, URigVMPin*>> Pairs;
+	TArray<URigVMNode*> NodesToRemove;
+	NodesToRemove.Add(InNode);
+
 	if (URigVMPin* ValuePin = InNode->FindPin(URigVMVariableNode::ValueName))
 	{
-		Links = ValuePin->GetTargetLinks(true);
+		TArray<URigVMLink*> Links = ValuePin->GetTargetLinks(true);
+		for (URigVMLink* Link : Links)
+		{
+			URigVMPin* SourcePin = Link->GetSourcePin();
+
+			TArray<URigVMPin*> TargetPins;
+			TargetPins.Add(Link->GetTargetPin());
+
+			for (int32 TargetPinIndex = 0; TargetPinIndex < TargetPins.Num(); TargetPinIndex++)
+			{
+				URigVMPin* TargetPin = TargetPins[TargetPinIndex];
+				if (Cast<URigVMRerouteNode>(TargetPin->GetNode()))
+				{
+					NodesToRemove.AddUnique(TargetPin->GetNode());
+					TargetPins.Append(TargetPin->GetLinkedTargetPins(false /* recursive */));
+				}
+				else
+				{
+					Pairs.Add(TPair<URigVMPin*, URigVMPin*>(SourcePin, TargetPin));
+				}
+			}
+		}
 	}
 
 	FName VariableName = InNode->GetVariableName();
@@ -3119,17 +3143,17 @@ bool URigVMController::MakeBindingsFromVariableNode(URigVMVariableNode* InNode,
 		return false;
 	}
 
-	if (Links.Num() > 0)
+	if (Pairs.Num() > 0)
 	{
 		if (bSetupUndoRedo)
 		{
 			OpenUndoBracket(TEXT("Turn Variable Node into Bindings"));
 		}
 
-		for (URigVMLink* Link : Links)
+		for (const TPair<URigVMPin*, URigVMPin*>& Pair : Pairs)
 		{
-			URigVMPin* SourcePin = Link->GetSourcePin();
-			URigVMPin* TargetPin = Link->GetTargetPin();
+			URigVMPin* SourcePin = Pair.Key;
+			URigVMPin* TargetPin = Pair.Value;
 			FString SegmentPath = SourcePin->GetSegmentPath();
 			FString VariablePathToBind = VariableName.ToString();
 			if (!SegmentPath.IsEmpty())
@@ -3143,7 +3167,10 @@ bool URigVMController::MakeBindingsFromVariableNode(URigVMVariableNode* InNode,
 			}
 		}
 
-		RemoveNode(InNode, bSetupUndoRedo, true);
+		for (URigVMNode* NodeToRemove : NodesToRemove)
+		{
+			RemoveNode(NodeToRemove, bSetupUndoRedo, true);
+		}
 
 		if (bSetupUndoRedo)
 		{
