RigVM: Allow editor only introspection on Register CPP Types and Objects.

diff --git a/Engine/Source/Developer/RigVMDeveloper/Private/RigVMCompiler/RigVMCompiler.cpp b/Engine/Source/Developer/RigVMDeveloper/Private/RigVMCompiler/RigVMCompiler.cpp
--- a/Engine/Source/Developer/RigVMDeveloper/Private/RigVMCompiler/RigVMCompiler.cpp
+++ b/Engine/Source/Developer/RigVMDeveloper/Private/RigVMCompiler/RigVMCompiler.cpp
@@ -1513,6 +1513,10 @@ FRigVMOperand URigVMCompiler::FindOrAddRegister(const FRigVMVarExprAST* InVarExp
 				}
 
 				int32 Register = Memory.AddRegisterArray(!Pin->IsDynamicArray() && !bIsDebugValue, RegisterName, ScriptStruct->GetStructureSize(), DefaultValues.Num(), Pin->IsArray(), Data.GetData(), NumSlices, ERigVMRegisterType::Struct, ScriptStruct);
+#if WITH_EDITORONLY_DATA
+				Memory.Registers[Register].CPPType = *ScriptStruct->GetStructCPPName();
+				Memory.Registers[Register].CPPTypeObject = ScriptStruct;
+#endif
 				ScriptStruct->DestroyStruct(Data.GetData(), DefaultValues.Num());
 
 				Operand = Memory.GetOperand(Register);
@@ -1533,6 +1537,10 @@ FRigVMOperand URigVMCompiler::FindOrAddRegister(const FRigVMVarExprAST* InVarExp
 				}
 
 				int32 Register = Memory.AddRegisterArray<uint8>(!Pin->IsDynamicArray() && !bIsDebugValue, RegisterName, Values.Num(), Pin->IsArray(), Values.GetData(), NumSlices, ERigVMRegisterType::Plain, nullptr);
+#if WITH_EDITORONLY_DATA
+				Memory.Registers[Register].CPPType = *Enum->CppType;
+				Memory.Registers[Register].CPPTypeObject = Enum;
+#endif
 				Operand = Memory.GetOperand(Register);
 			}
 			else if (BaseCPPType == TEXT("bool"))
@@ -1543,6 +1551,10 @@ FRigVMOperand URigVMCompiler::FindOrAddRegister(const FRigVMVarExprAST* InVarExp
 					Values.Add((DefaultValue == TEXT("True")) || (DefaultValue == TEXT("true")) || (DefaultValue == TEXT("1")));
 				}
 				int32 Register = Memory.AddRegisterArray<bool>(!Pin->IsDynamicArray() && !bIsDebugValue, RegisterName, Values.Num(), Pin->IsArray(), (uint8*)Values.GetData(), NumSlices, ERigVMRegisterType::Plain, nullptr);
+#if WITH_EDITORONLY_DATA
+				Memory.Registers[Register].CPPType = *BaseCPPType;
+				Memory.Registers[Register].CPPTypeObject = nullptr;
+#endif
 				Operand = Memory.GetOperand(Register);
 			}
 			else if (BaseCPPType == TEXT("int32"))
@@ -1560,6 +1572,10 @@ FRigVMOperand URigVMCompiler::FindOrAddRegister(const FRigVMVarExprAST* InVarExp
 					}
 				}
 				int32 Register = Memory.AddRegisterArray<int32>(!Pin->IsDynamicArray() && !bIsDebugValue, RegisterName, Values.Num(), Pin->IsArray(), (uint8*)Values.GetData(), NumSlices, ERigVMRegisterType::Plain, nullptr);
+#if WITH_EDITORONLY_DATA
+				Memory.Registers[Register].CPPType = *BaseCPPType;
+				Memory.Registers[Register].CPPTypeObject = nullptr;
+#endif
 				Operand = Memory.GetOperand(Register);
 			}
 			else if (BaseCPPType == TEXT("float"))
@@ -1577,6 +1593,10 @@ FRigVMOperand URigVMCompiler::FindOrAddRegister(const FRigVMVarExprAST* InVarExp
 					}
 				}
 				int32 Register = Memory.AddRegisterArray<float>(!Pin->IsDynamicArray() && !bIsDebugValue, RegisterName, Values.Num(), Pin->IsArray(), (uint8*)Values.GetData(), NumSlices, ERigVMRegisterType::Plain, nullptr);
+#if WITH_EDITORONLY_DATA
+				Memory.Registers[Register].CPPType = *BaseCPPType;
+				Memory.Registers[Register].CPPTypeObject = nullptr;
+#endif
 				Operand = Memory.GetOperand(Register);
 			}
 			else if (BaseCPPType == TEXT("FName"))
@@ -1587,11 +1607,19 @@ FRigVMOperand URigVMCompiler::FindOrAddRegister(const FRigVMVarExprAST* InVarExp
 					Values.Add(*DefaultValue);
 				}
 				int32 Register = Memory.AddRegisterArray<FName>(!Pin->IsDynamicArray() && !bIsDebugValue, RegisterName, Values.Num(), Pin->IsArray(), (uint8*)Values.GetData(), NumSlices, ERigVMRegisterType::Name, nullptr);
+#if WITH_EDITORONLY_DATA
+				Memory.Registers[Register].CPPType = *BaseCPPType;
+				Memory.Registers[Register].CPPTypeObject = nullptr;
+#endif
 				Operand = Memory.GetOperand(Register);
 			}
 			else if (BaseCPPType == TEXT("FString"))
 			{
 				int32 Register = Memory.AddRegisterArray<FString>(!Pin->IsDynamicArray() && !bIsDebugValue, RegisterName, DefaultValues.Num(), Pin->IsArray(), (uint8*)DefaultValues.GetData(), NumSlices, ERigVMRegisterType::String, nullptr);
+#if WITH_EDITORONLY_DATA
+				Memory.Registers[Register].CPPType = *BaseCPPType;
+				Memory.Registers[Register].CPPTypeObject = nullptr;
+#endif
 				Operand = Memory.GetOperand(Register);
 			}
 			else
diff --git a/Engine/Source/Runtime/RigVM/Public/RigVMCore/RigVMMemory.h b/Engine/Source/Runtime/RigVM/Public/RigVMCore/RigVMMemory.h
--- a/Engine/Source/Runtime/RigVM/Public/RigVMCore/RigVMMemory.h
+++ b/Engine/Source/Runtime/RigVM/Public/RigVMCore/RigVMMemory.h
@@ -157,6 +157,10 @@ struct RIGVM_API FRigVMRegister
 		, ScriptStructIndex(INDEX_NONE)
 		, bIsArray(false)
 		, bIsDynamic(false)
+#if WITH_EDITORONLY_DATA
+		, CPPType(NAME_None)
+		, CPPTypeObject(nullptr)
+#endif
 	{
 	}
 
@@ -206,6 +210,17 @@ struct RIGVM_API FRigVMRegister
 	UPROPERTY()
 	bool bIsDynamic;
 
+#if WITH_EDITORONLY_DATA
+
+	// Defines the CPP type used for the register
+	UPROPERTY(transient)
+	FName CPPType;
+
+	UPROPERTY(transient)
+	UObject* CPPTypeObject;
+
+#endif
+
 	void Serialize(FArchive& Ar);
 	void Save(FArchive& Ar);
 	void Load(FArchive& Ar);
