Control Rig: Perform reroute tracing as part of the inlining pass, clean up unit tests

diff --git a/Engine/Source/Developer/RigVMDeveloper/Private/RigVMCompiler/RigVMAST.cpp b/Engine/Source/Developer/RigVMDeveloper/Private/RigVMCompiler/RigVMAST.cpp
--- a/Engine/Source/Developer/RigVMDeveloper/Private/RigVMCompiler/RigVMAST.cpp
+++ b/Engine/Source/Developer/RigVMDeveloper/Private/RigVMCompiler/RigVMAST.cpp
@@ -817,6 +817,8 @@ const FRigVMVarExprAST* FRigVMCallExternExprAST::FindVarWithPinName(const FName&
 	return nullptr;
 }
 
+const TArray<URigVMPin*> FRigVMParserAST::EmptyLinks;
+
 FRigVMParserAST::FRigVMParserAST(URigVMGraph* InGraph, URigVMController* InController, const FRigVMParserASTSettings& InSettings, const TArray<FRigVMExternalVariable>& InExternalVariables, const TArray<FRigVMUserDataArray>& InRigVMUserData)
 {
 	DECLARE_SCOPE_HIERARCHICAL_COUNTER_FUNC()
@@ -825,7 +827,9 @@ FRigVMParserAST::FRigVMParserAST(URigVMGraph* InGraph, URigVMController* InContr
 	LastCycleCheckExpr = nullptr;
 	LinksToSkip = InSettings.LinksToSkip;
 
-	const TArray<URigVMNode*> Nodes = InGraph->GetNodes();
+	// construct the inlined nodes and links information
+	Inline(InGraph);
+
 	for (URigVMNode* Node : Nodes)
 	{
 		if(Node->IsEvent())
@@ -865,11 +869,7 @@ FRigVMParserAST::FRigVMParserAST(URigVMGraph* InGraph, URigVMController* InContr
 
 	FoldEntries();
 	InjectExitsToEntries();
-
-	if (InSettings.bFoldReroutes || InSettings.bFoldAssignments)
-	{
-		FoldNoOps();
-	}
+	FoldNoOps();
 
 	// keep folding constant branches and values while we can
 	bool bContinueToFoldConstantBranches = InSettings.bFoldConstantBranches;
@@ -907,7 +907,9 @@ FRigVMParserAST::FRigVMParserAST(URigVMGraph* InGraph, const TArray<URigVMNode*>
 	Block->Name = TEXT("NodesToCompute");
 	RootExpressions.Add(Block);
 
-	for (URigVMNode* Node : InNodesToCompute)
+	Inline(InGraph, InNodesToCompute);
+
+	for (URigVMNode* Node : Nodes)
 	{
 		if (Node->IsEvent())
 		{
@@ -978,14 +980,13 @@ FRigVMExprAST* FRigVMParserAST::TraverseMutableNode(URigVMNode* InNode, FRigVMEx
 						}
 					}
 
-					const TArray<URigVMLink*>& Links = SourcePin->GetLinks();
-					for(URigVMLink* Link : Links)
+					const TArray<URigVMPin*>& TargetPins = GetTargetPins(SourcePin);
+					for(URigVMPin* TargetPin : TargetPins)
 					{
-						if (LinksToSkip.Contains(Link))
+						if(ShouldLinkBeSkipped(FRigVMPinPair(SourcePin, TargetPin)))
 						{
 							continue;
 						}
-						URigVMPin* TargetPin = Link->GetTargetPin();
 						TraverseMutableNode(TargetPin->GetNode(), ParentExpr);
 					}
 				}
@@ -1099,12 +1100,15 @@ FRigVMExprAST* FRigVMParserAST::TraversePin(URigVMPin* InPin, FRigVMExprAST* InP
 {
 	ensure(!SubjectToExpression.Contains(InPin));
 
-	TArray<URigVMLink*> SourceLinks = InPin->GetSourceLinks(true);
+	TArray<FRigVMPinPair> Links = GetSourceLinks(InPin, true);
 	
 	if (LinksToSkip.Num() > 0)
 	{
-		const TArray<URigVMLink*>& LocalLinksToSkip = LinksToSkip;
-		SourceLinks.RemoveAll([LocalLinksToSkip](URigVMLink* LinkToCheck) { return LocalLinksToSkip.Contains(LinkToCheck); });
+		Links.RemoveAll([this](const FRigVMPinPair& LinkToCheck)
+			{
+				return this->ShouldLinkBeSkipped(LinkToCheck);
+			}
+		);
 	}
 
 	struct Local
@@ -1145,7 +1149,7 @@ FRigVMExprAST* FRigVMParserAST::TraversePin(URigVMPin* InPin, FRigVMExprAST* InP
 
 	if ((InPin->GetDirection() == ERigVMPinDirection::Input ||
 		InPin->GetDirection() == ERigVMPinDirection::Visible) &&
-		SourceLinks.Num() == 0 &&
+		Links.Num() == 0 &&
 		SubPinsBoundToVariables.Num() == 0)
 	{
 		if (Cast<URigVMParameterNode>(InPin->GetNode()) ||
@@ -1196,9 +1200,9 @@ FRigVMExprAST* FRigVMParserAST::TraversePin(URigVMPin* InPin, FRigVMExprAST* InP
 	{
 		bool bHasSourceLinkToRoot = false;
 		URigVMPin* RootPin = InPin->GetRootPin();
-		for (URigVMLink* SourceLink : SourceLinks)
+		for (const FRigVMPinPair& SourceLink : Links)
 		{
-			if (SourceLink->GetTargetPin() == RootPin)
+			if (SourceLink.Value == RootPin)
 			{
 				bHasSourceLinkToRoot = true;
 				break;
@@ -1206,8 +1210,8 @@ FRigVMExprAST* FRigVMParserAST::TraversePin(URigVMPin* InPin, FRigVMExprAST* InP
 		}
 
 		if (!bHasSourceLinkToRoot && 
-			InPin->GetSourceLinks(false).Num() == 0 &&
-			(InPin->GetDirection() == ERigVMPinDirection::IO || SourceLinks.Num() > 0 || SubPinsBoundToVariables.Num() > 0))
+			GetSourcePins(InPin).Num() == 0 &&
+			(InPin->GetDirection() == ERigVMPinDirection::IO || Links.Num() > 0 || SubPinsBoundToVariables.Num() > 0))
 		{
 			FRigVMLiteralExprAST* LiteralExpr = MakeExpr<FRigVMLiteralExprAST>(InPin);
 			FRigVMCopyExprAST* LiteralCopyExpr = MakeExpr<FRigVMCopyExprAST>(InPin, InPin);
@@ -1232,7 +1236,7 @@ FRigVMExprAST* FRigVMParserAST::TraversePin(URigVMPin* InPin, FRigVMExprAST* InP
 
 	if ((InPin->GetDirection() == ERigVMPinDirection::IO || InPin->GetDirection() == ERigVMPinDirection::Input) &&
 		(InParentExpr->IsA(FRigVMExprAST::If) || InParentExpr->IsA(FRigVMExprAST::Select)) &&
-		SourceLinks.Num() > 0)
+		Links.Num() > 0)
 	{
 		FRigVMBlockExprAST* BlockExpr = MakeExpr<FRigVMBlockExprAST>(FRigVMExprAST::EType::Block);
 		BlockExpr->AddParent(PinExpr);
@@ -1240,7 +1244,7 @@ FRigVMExprAST* FRigVMParserAST::TraversePin(URigVMPin* InPin, FRigVMExprAST* InP
 		ParentExprForLinks = BlockExpr;
 	}
 
-	for (URigVMLink* SourceLink : SourceLinks)
+	for (const FRigVMPinPair& SourceLink : Links)
 	{
 		TraverseLink(SourceLink, ParentExprForLinks);
 	}
@@ -1258,12 +1262,10 @@ FRigVMExprAST* FRigVMParserAST::TraversePin(URigVMPin* InPin, FRigVMExprAST* InP
 	return PinExpr;
 }
 
-FRigVMExprAST* FRigVMParserAST::TraverseLink(URigVMLink* InLink, FRigVMExprAST* InParentExpr)
+FRigVMExprAST* FRigVMParserAST::TraverseLink(const FRigVMPinPair& InLink, FRigVMExprAST* InParentExpr)
 {
-	ensure(!SubjectToExpression.Contains(InLink));
-
-	URigVMPin* SourcePin = InLink->GetSourcePin();
-	URigVMPin* TargetPin = InLink->GetTargetPin();
+	URigVMPin* SourcePin = InLink.Key;
+	URigVMPin* TargetPin = InLink.Value;
 	URigVMPin* SourceRootPin = SourcePin->GetRootPin();
 	URigVMPin* TargetRootPin = TargetPin->GetRootPin();
 
@@ -1287,9 +1289,8 @@ FRigVMExprAST* FRigVMParserAST::TraverseLink(URigVMLink* InLink, FRigVMExprAST*
 		AssignExpr = MakeExpr<FRigVMAssignExprAST>(FRigVMExprAST::EType::Assign, SourcePin, TargetPin);
 	}
 
-	AssignExpr->Name = *InLink->GetPinPathRepresentation();
+	AssignExpr->Name = *GetLinkAsString(InLink);
 	AssignExpr->AddParent(InParentExpr);
-	SubjectToExpression.Add(InLink, AssignExpr);
 
 	FRigVMExprAST* NodeExpr = TraverseNode(SourcePin->GetNode(), AssignExpr);
 	if (NodeExpr)
@@ -1528,7 +1529,7 @@ void FRigVMParserAST::FoldNoOps()
 						continue;
 					}
 				}
-				if(URigVMVariableNode* VariableNode = Cast<URigVMVariableNode>(Node))
+				if (URigVMVariableNode* VariableNode = Cast<URigVMVariableNode>(Node))
 				{
 					if (!VariableNode->IsGetter())
 					{
@@ -1662,7 +1663,6 @@ bool FRigVMParserAST::FoldConstantValuesToLiterals(URigVMGraph* InGraph, URigVMC
 	TArray<URigVMPin*> PinsToCompute;
 	TArray<URigVMNode*> NodesToCompute;
 
-	const TArray<URigVMNode*> Nodes = InGraph->GetNodes();
 	for (URigVMNode* Node : Nodes)
 	{
 		if (Cast<URigVMParameterNode>(Node) != nullptr ||
@@ -1716,7 +1716,7 @@ bool FRigVMParserAST::FoldConstantValuesToLiterals(URigVMGraph* InGraph, URigVMC
 				}
 			}
 
-			TArray<URigVMPin*> SourcePins = Pin->GetLinkedSourcePins(true);
+			TArray<FRigVMPinPair> SourcePins = GetSourceLinks(Pin, true);
 			if (SourcePins.Num() == 0)
 			{
 				continue;
@@ -1728,9 +1728,9 @@ bool FRigVMParserAST::FoldConstantValuesToLiterals(URigVMGraph* InGraph, URigVMC
 			}
 
 			bool bFoundValidSourcePin = false;
-			for (URigVMPin* SourcePin : SourcePins)
+			for (const FRigVMPinPair& SourcePin : SourcePins)
 			{
-				URigVMNode* SourceNode = SourcePin->GetNode();
+				URigVMNode* SourceNode = SourcePin.Key->GetNode();
 				check(SourceNode);
 
 				if (Cast<URigVMParameterNode>(SourceNode) != nullptr ||
@@ -1741,7 +1741,7 @@ bool FRigVMParserAST::FoldConstantValuesToLiterals(URigVMGraph* InGraph, URigVMC
 					continue;
 				}
 
-				PinsToCompute.AddUnique(SourcePin);
+				PinsToCompute.AddUnique(SourcePin.Key);
 				NodesToCompute.AddUnique(SourceNode);
 				bFoundValidSourcePin = true;
 			}
@@ -1862,10 +1862,10 @@ bool FRigVMParserAST::FoldConstantValuesToLiterals(URigVMGraph* InGraph, URigVMC
 			}
 		}
 
-		TArray<URigVMPin*> TargetPins = PinToCompute->GetLinkedTargetPins();
+		const TArray<URigVMPin*>& TargetPins = GetTargetPins(PinToCompute);
 		for (URigVMPin* TargetPin : TargetPins)
 		{
-			PinDefaultValueOverrides.FindOrAdd(TargetPin, DefaultValue);
+			PinDefaultValueOverrides.FindOrAdd(TargetPin) = DefaultValue;
 		}
 	}
 
@@ -1930,7 +1930,6 @@ bool FRigVMParserAST::FoldUnreachableBranches(URigVMGraph* InGraph)
 
 	TArray<FRigVMExprAST*> ExpressionsToRemove;
 
-	const TArray<URigVMNode*> Nodes = InGraph->GetNodes();
 	for (URigVMNode* Node : Nodes)
 	{
 		if (Cast<URigVMParameterNode>(Node) != nullptr ||
@@ -2279,6 +2278,10 @@ FString FRigVMParserAST::DumpText() const
 	FString Result;
 	for (FRigVMExprAST* RootExpr : RootExpressions)
 	{
+		if(RootExpr == GetObsoleteBlock(false /* create */))
+		{
+			continue;
+		}
 		Result += TEXT("\n") + RootExpr->DumpText();
 	}
 	return Result;
@@ -2295,15 +2298,19 @@ FString FRigVMParserAST::DumpDot() const
 
 	for (FRigVMExprAST* RootExpr : RootExpressions)
 	{
+		if (RootExpr == GetObsoleteBlock(false /* create */))
+		{
+			continue;
+		}
 		Result += RootExpr->DumpDot(OutExpressionDefined, TEXT("  "));
 	}
 	Result += TEXT("\n}");
 	return Result;
 }
 
-FRigVMBlockExprAST* FRigVMParserAST::GetObsoleteBlock()
+FRigVMBlockExprAST* FRigVMParserAST::GetObsoleteBlock(bool bCreateIfMissing)
 {
-	if (ObsoleteBlock == nullptr)
+	if (ObsoleteBlock == nullptr && bCreateIfMissing)
 	{
 		ObsoleteBlock = MakeExpr<FRigVMBlockExprAST>(FRigVMExprAST::EType::Block);
 		ObsoleteBlock->bIsObsolete = true;
@@ -2312,9 +2319,9 @@ FRigVMBlockExprAST* FRigVMParserAST::GetObsoleteBlock()
 	return ObsoleteBlock;
 }
 
-const FRigVMBlockExprAST* FRigVMParserAST::GetObsoleteBlock() const
+const FRigVMBlockExprAST* FRigVMParserAST::GetObsoleteBlock(bool bCreateIfMissing) const
 {
-	if (ObsoleteBlock == nullptr)
+	if (ObsoleteBlock == nullptr && bCreateIfMissing)
 	{
 		FRigVMParserAST* MutableThis = (FRigVMParserAST*)this;
 		MutableThis->ObsoleteBlock = MutableThis->MakeExpr<FRigVMBlockExprAST>(FRigVMExprAST::EType::Block);
@@ -2401,3 +2408,205 @@ void FRigVMParserAST::TraverseChildren(const FRigVMExprAST* InExpr, TFunctionRef
 	}
 }
 
+const TArray<URigVMPin*>& FRigVMParserAST::GetSourcePins(URigVMPin* InPin) const
+{
+	if (const TArray<URigVMPin*>* Links = SourceLinks.Find(InPin))
+	{
+		return *Links;
+	}
+	return EmptyLinks;
+}
+
+const TArray<URigVMPin*>& FRigVMParserAST::GetTargetPins(URigVMPin* InPin) const
+{
+	if (const TArray<URigVMPin*>* Links = TargetLinks.Find(InPin))
+	{
+		return *Links;
+	}
+	return EmptyLinks;
+}
+
+TArray<FRigVMPinPair> FRigVMParserAST::GetSourceLinks(URigVMPin* InPin, bool bRecursive) const
+{
+	const TArray<URigVMPin*>& SourcePins = GetSourcePins(InPin);
+
+	TArray<FRigVMPinPair> Pairs;
+	for (URigVMPin* SourcePin : SourcePins)
+	{
+		Pairs.Add(FRigVMPinPair(SourcePin, InPin));
+	}
+
+	if (bRecursive)
+	{
+		for (URigVMPin* SubPin : InPin->GetSubPins())
+		{
+			Pairs.Append(GetSourceLinks(SubPin, true));
+		}
+	}
+
+	return Pairs;
+}
+
+TArray<FRigVMPinPair> FRigVMParserAST::GetTargetLinks(URigVMPin* InPin, bool bRecursive) const
+{
+	const TArray<URigVMPin*>& TargetPins = GetTargetPins(InPin);
+
+	TArray<FRigVMPinPair> Pairs;
+	for (URigVMPin* TargetPin : TargetPins)
+	{
+		Pairs.Add(FRigVMPinPair(InPin, TargetPin));
+	}
+
+	if (bRecursive)
+	{
+		for (URigVMPin* SubPin : InPin->GetSubPins())
+		{
+			Pairs.Append(GetTargetLinks(SubPin, true));
+		}
+	}
+
+	return Pairs;
+}
+
+void FRigVMParserAST::Inline(URigVMGraph* InGraph)
+{
+	Inline(InGraph, InGraph->GetNodes());
+}
+
+void FRigVMParserAST::Inline(URigVMGraph* InGraph, const TArray<URigVMNode*>& InNodes)
+{
+	struct LocalPinTraversalInfo
+	{
+		URigVMPin::FDefaultValueOverride* PinDefaultValueOverrides;
+		TMap<URigVMPin*, URigVMPin*> SourcePins;
+		TMap<URigVMPin*, TArray<URigVMPin*>>* TargetLinks;
+		TMap<URigVMPin*, TArray<URigVMPin*>>* SourceLinks;
+
+		static URigVMPin* FindSourcePin(URigVMPin* InPin, LocalPinTraversalInfo& OutTraversalInfo)
+		{
+			if (URigVMPin* const* SourcePin = OutTraversalInfo.SourcePins.Find(InPin))
+			{
+				return *SourcePin;
+			}
+
+			TArray<URigVMLink*> SourceLinks = InPin->GetSourceLinks(false /* recursive */);
+			if (SourceLinks.Num() > 0)
+			{
+				URigVMPin* SourcePin = SourceLinks[0]->GetSourcePin();
+
+				// Only continue the recursion if the current source pin
+				// is on a reroute node.
+				if(SourcePin->GetNode()->IsA<URigVMRerouteNode>())
+				{
+					URigVMPin* ParentPin = SourcePin;
+					FString SegmentPath;
+
+					while (ParentPin != nullptr)
+					{
+						if (URigVMPin* SourceSourcePin = FindSourcePin(ParentPin, OutTraversalInfo))
+						{
+							if (!SegmentPath.IsEmpty())
+							{
+								SourceSourcePin = SourceSourcePin->FindSubPin(SegmentPath);
+							}
+
+							OutTraversalInfo.SourcePins.FindOrAdd(InPin) = SourceSourcePin;
+							return SourceSourcePin;
+						}
+
+						if (SegmentPath.IsEmpty())
+						{
+							SegmentPath = ParentPin->GetName();
+						}
+						else
+						{
+							SegmentPath = URigVMPin::JoinPinPath(ParentPin->GetName(), SegmentPath);
+						}
+
+						ParentPin = ParentPin->GetParentPin();
+					}
+				}
+
+				OutTraversalInfo.SourcePins.FindOrAdd(InPin) = SourcePin;
+				return SourcePin;
+			}
+
+			OutTraversalInfo.SourcePins.FindOrAdd(InPin) = nullptr;
+			return nullptr;
+		}
+
+		static void VisitPin(URigVMPin* InPin, LocalPinTraversalInfo& OutTraversalInfo)
+		{
+			if (URigVMPin* SourcePin = FindSourcePin(InPin, OutTraversalInfo))
+			{
+				// The source pin is the final determined source pin, since
+				// FindSourcePin is recursive.
+				// If the source pin is on a reroute node, this means that
+				// we only care about the default value - since it is a 
+				// "hanging" reroute without any live input.
+
+				if (SourcePin->GetNode()->IsA<URigVMRerouteNode>())
+				{
+					OutTraversalInfo.PinDefaultValueOverrides->FindOrAdd(InPin) = SourcePin->GetDefaultValue();
+				}
+				else
+				{
+					OutTraversalInfo.SourceLinks->FindOrAdd(InPin).Add(SourcePin);
+					OutTraversalInfo.TargetLinks->FindOrAdd(SourcePin).Add(InPin);
+				}
+			}
+
+			for (URigVMPin* SubPin : InPin->GetSubPins())
+			{
+				VisitPin(SubPin, OutTraversalInfo);
+			}
+		}
+	};
+
+	Nodes.Reset();
+	SourceLinks.Reset();
+	TargetLinks.Reset();
+
+	// a) find all of the relevant nodes
+	//    once we support encapsulation we'll have to traverse into library nodes
+	Nodes = InNodes;
+
+	// c) flatten links from an entry node / to a return node
+	//    once we have encapsulation implement this
+	//    also traverse links along reroutes and flatten them
+	LocalPinTraversalInfo TraversalInfo;
+	TraversalInfo.PinDefaultValueOverrides = &PinDefaultValueOverrides;
+	TraversalInfo.TargetLinks = &TargetLinks;
+	TraversalInfo.SourceLinks = &SourceLinks;
+
+	for (URigVMNode* Node : Nodes)
+	{
+		if (Node->IsA<URigVMRerouteNode>())
+		{
+			continue;
+		}
+
+		for (URigVMPin* Pin : Node->GetPins())
+		{
+			LocalPinTraversalInfo::VisitPin(Pin, TraversalInfo);
+		}
+	}
+}
+
+bool FRigVMParserAST::ShouldLinkBeSkipped(const FRigVMPinPair& InLink) const
+{
+	for (URigVMLink* LinkToSkip : LinksToSkip)
+	{
+		if (LinkToSkip->GetSourcePin() == InLink.Key &&
+			LinkToSkip->GetTargetPin() == InLink.Value)
+		{
+			return true;
+		}
+	}
+	return false;
+}
+
+FString FRigVMParserAST::GetLinkAsString(const FRigVMPinPair& InLink)
+{
+	return FString::Printf(TEXT("%s -> %s"), *InLink.Key->GetPinPath(), *InLink.Value->GetPinPath());
+}
\ No newline at end of file
diff --git a/Engine/Source/Developer/RigVMDeveloper/Private/RigVMCompiler/RigVMCompiler.cpp b/Engine/Source/Developer/RigVMDeveloper/Private/RigVMCompiler/RigVMCompiler.cpp
--- a/Engine/Source/Developer/RigVMDeveloper/Private/RigVMCompiler/RigVMCompiler.cpp
+++ b/Engine/Source/Developer/RigVMDeveloper/Private/RigVMCompiler/RigVMCompiler.cpp
@@ -160,6 +160,10 @@ bool URigVMCompiler::Compile(URigVMGraph* InGraph, URigVMController* InControlle
 	// define all parameters independent from sorted nodes
 	for (const FRigVMExprAST* Expr : AST->Expressions)
 	{
+		if (Expr->IsA(FRigVMExprAST::EType::Literal))
+		{
+			continue;
+		}
 		if (Expr->IsA(FRigVMExprAST::EType::Var))
 		{
 			const FRigVMVarExprAST* VarExpr = Expr->To<FRigVMVarExprAST>();
@@ -811,6 +815,7 @@ void URigVMCompiler::TraverseAssign(const FRigVMAssignExprAST* InExpr, FRigVMCom
 				FRigVMOperand WatchOperand = FindOrAddRegister(SourceExpr, WorkData, true /*debug watch*/);
 				if (WatchOperand.IsValid())
 				{
+					Source.RegisterOffset = INDEX_NONE;
 					WorkData.VM->GetByteCode().AddCopyOp(Source, WatchOperand);
 				}
 			}
@@ -823,6 +828,7 @@ void URigVMCompiler::TraverseAssign(const FRigVMAssignExprAST* InExpr, FRigVMCom
 				FRigVMOperand WatchOperand = FindOrAddRegister(TargetExpr, WorkData, true /*debug watch*/);
 				if (WatchOperand.IsValid())
 				{
+					Target.RegisterOffset = INDEX_NONE;
 					WorkData.VM->GetByteCode().AddCopyOp(Target, WatchOperand);
 				}
 			}
@@ -1261,7 +1267,7 @@ FRigVMOperand URigVMCompiler::FindOrAddRegister(const FRigVMVarExprAST* InVarExp
 	FString RegisterKey = Hash;
 
 	bool bIsExecutePin = Pin->IsExecuteContext();
-	bool bIsLiteral = InVarExpr->GetType() == FRigVMExprAST::EType::Literal;
+	bool bIsLiteral = InVarExpr->GetType() == FRigVMExprAST::EType::Literal && !bIsDebugValue;
 	bool bIsParameter = InVarExpr->IsGraphParameter();
 
 	FRigVMOperand const* ExistingOperand = WorkData.PinPathToOperand->Find(Hash);
@@ -1635,7 +1641,7 @@ FRigVMOperand URigVMCompiler::FindOrAddRegister(const FRigVMVarExprAST* InVarExp
 			}
 		}
 
-		if (bIsParameter)
+		if (bIsParameter && !bIsLiteral)
 		{
 			URigVMParameterNode* ParameterNode = Cast<URigVMParameterNode>(Pin->GetNode());
 			check(ParameterNode);
diff --git a/Engine/Source/Developer/RigVMDeveloper/Private/RigVMModel/RigVMController.cpp b/Engine/Source/Developer/RigVMDeveloper/Private/RigVMModel/RigVMController.cpp
--- a/Engine/Source/Developer/RigVMDeveloper/Private/RigVMModel/RigVMController.cpp
+++ b/Engine/Source/Developer/RigVMDeveloper/Private/RigVMModel/RigVMController.cpp
@@ -972,10 +972,6 @@ URigVMRerouteNode* URigVMController::AddRerouteNodeOnPin(const FString& InPinPat
 
 	ForEveryPinRecursively(ValuePin, [](URigVMPin* Pin) {
 		Pin->bIsExpanded = true;
-		if (Pin->GetParentPin() != nullptr)
-		{
-			Pin->Direction = ERigVMPinDirection::Visible;
-		}
 	});
 
 	Graph->Nodes.Add(Node);
diff --git a/Engine/Source/Developer/RigVMDeveloper/Public/RigVMCompiler/RigVMAST.h b/Engine/Source/Developer/RigVMDeveloper/Public/RigVMCompiler/RigVMAST.h
--- a/Engine/Source/Developer/RigVMDeveloper/Public/RigVMCompiler/RigVMAST.h
+++ b/Engine/Source/Developer/RigVMDeveloper/Public/RigVMCompiler/RigVMAST.h
@@ -25,11 +25,13 @@ class FRigVMIfExprAST;
 class FRigVMSelectExprAST;
 
 class URigVMPin;
-class URigVMNode;
 class URigVMLink;
+class URigVMNode;
 class URigVMGraph;
 class URigVMController;
 
+typedef TPair<URigVMPin*, URigVMPin*> FRigVMPinPair;
+
 /*
  * Base class for an expression within an abstract syntax tree.
  * The base implements parent / child relationships as well as a
@@ -1087,10 +1089,6 @@ struct RIGVMDEVELOPER_API FRigVMParserASTSettings
 {
 	GENERATED_BODY()
 
-	// remove no op nodes - used for reroutes and other expressions.
-	UPROPERTY(EditAnywhere, Category = "AST")
-	bool bFoldReroutes;
-
 	// fold assignments / copies
 	UPROPERTY(EditAnywhere, Category = "AST")
 	bool bFoldAssignments;
@@ -1111,7 +1109,6 @@ struct RIGVMDEVELOPER_API FRigVMParserASTSettings
 	static FRigVMParserASTSettings Fast()
 	{
 		FRigVMParserASTSettings Settings;
-		Settings.bFoldReroutes = false;
 		Settings.bFoldAssignments = false;
 		Settings.bFoldLiterals = false;
 		Settings.bFoldConstantBranches = false;
@@ -1123,7 +1120,6 @@ struct RIGVMDEVELOPER_API FRigVMParserASTSettings
 	static FRigVMParserASTSettings Optimized()
 	{
 		FRigVMParserASTSettings Settings;
-		Settings.bFoldReroutes = true;
 		Settings.bFoldAssignments = true;
 		Settings.bFoldLiterals = true;
 		Settings.bFoldConstantBranches = true;
@@ -1198,8 +1194,8 @@ public:
 	FString DumpDot() const;
 
 	// returns an obsolete block for unmanaged expression
-	FRigVMBlockExprAST* GetObsoleteBlock();
-	const FRigVMBlockExprAST* GetObsoleteBlock() const;
+	FRigVMBlockExprAST* GetObsoleteBlock(bool bCreateIfMissing = true);
+	const FRigVMBlockExprAST* GetObsoleteBlock(bool bCreateIfMissing = true) const;
 
 	// returns the AST's override table for pin defaults
 	const URigVMPin::FDefaultValueOverride& GetPinDefaultOverrides() const { return PinDefaultValueOverrides; }
@@ -1287,7 +1283,7 @@ private:
 	// helper function to refresh the expression indices (used after deleting an expression)
 	void RefreshExprIndices();
 
-	// helper function to fold / remove the no op expressions
+	// helper function to fold / remove the no op expressions (parameters mostly)
 	void FoldNoOps();
 
 	// helper function to fold / merge redundant literals with the same value
@@ -1302,6 +1298,16 @@ private:
 	// helper function to fold unreachable branches 
 	bool FoldUnreachableBranches(URigVMGraph* InGraph);
 
+	// helper function to inline all contributing nodes of the graph
+	void Inline(URigVMGraph* InGraph);
+	void Inline(URigVMGraph* InGraph, const TArray<URigVMNode*>& InNodes);
+
+	// helper functions to retrieve links for a given pin
+	const TArray<URigVMPin*>& GetSourcePins(URigVMPin* InPin) const;
+	const TArray<URigVMPin*>& GetTargetPins(URigVMPin* InPin) const;
+	TArray<FRigVMPinPair> GetSourceLinks(URigVMPin* InPin, bool bRecursive = false) const;
+	TArray<FRigVMPinPair> GetTargetLinks(URigVMPin* InPin, bool bRecursive = false) const;
+
 	// traverse a single mutable node (constructs entry, call extern and other expressions)
 	FRigVMExprAST* TraverseMutableNode(URigVMNode* InNode, FRigVMExprAST* InParentExpr);
 
@@ -1318,7 +1324,11 @@ private:
 	FRigVMExprAST* TraversePin(URigVMPin* InPin, FRigVMExprAST* InParentExpr);
 
 	// traverse a single link (constructs assign + copy expressions)
-	FRigVMExprAST* TraverseLink(URigVMLink* InLink, FRigVMExprAST* InParentExpr);
+	FRigVMExprAST* TraverseLink(const FRigVMPinPair& InLink, FRigVMExprAST* InParentExpr);
+
+	bool ShouldLinkBeSkipped(const FRigVMPinPair& InLink) const;
+
+	static FString GetLinkAsString(const FRigVMPinPair& InLink);
 
 	TMap<UObject*, FRigVMExprAST*> SubjectToExpression;
 	TMap<UObject*, int32> NodeExpressionIndex;
@@ -1326,6 +1336,10 @@ private:
 	TArray<FRigVMExprAST*> RootExpressions;
 	FRigVMBlockExprAST* ObsoleteBlock;
 
+	TArray<URigVMNode*> Nodes;
+	TMap<URigVMPin*, TArray<URigVMPin*>> TargetLinks;
+	TMap<URigVMPin*, TArray<URigVMPin*>> SourceLinks;
+	static const TArray<URigVMPin*> EmptyLinks;
 	URigVMPin::FDefaultValueOverride PinDefaultValueOverrides;
 
 	enum ETraverseRelationShip
