Control Rig & Engine Test: Added Unit Test for handling empty arrays in a struct. Added also the fix that prevents this test from crashing. The crash happened because we were skipping empty arrays completely when gettng the default value.

diff --git a/Engine/Source/Developer/RigVMDeveloper/Private/RigVMModel/RigVMPin.cpp b/Engine/Source/Developer/RigVMDeveloper/Private/RigVMModel/RigVMPin.cpp
--- a/Engine/Source/Developer/RigVMDeveloper/Private/RigVMModel/RigVMPin.cpp
+++ b/Engine/Source/Developer/RigVMDeveloper/Private/RigVMModel/RigVMPin.cpp
@@ -299,7 +299,7 @@ FString URigVMPin::GetDefaultValue(const FDefaultValueOverride& InDefaultValueOv
 				{
 					ElementDefaultValue = TEXT("\"") + ElementDefaultValue + TEXT("\"");
 				}
-				else if (ElementDefaultValue.IsEmpty() || ElementDefaultValue == TEXT("()"))
+				else if (ElementDefaultValue.IsEmpty())
 				{
 					continue;
 				}
