Control Rig: RigVM Inline pass to fold reroutes with parent / child segment paths

diff --git a/Engine/Source/Developer/RigVMDeveloper/Private/RigVMCompiler/RigVMAST.cpp b/Engine/Source/Developer/RigVMDeveloper/Private/RigVMCompiler/RigVMAST.cpp
--- a/Engine/Source/Developer/RigVMDeveloper/Private/RigVMCompiler/RigVMAST.cpp
+++ b/Engine/Source/Developer/RigVMDeveloper/Private/RigVMCompiler/RigVMAST.cpp
@@ -15,6 +15,9 @@
 #include "RigVMModel/RigVMController.h"
 #include "RigVMCore/RigVMExecuteContext.h"
 #include "Stats/StatsHierarchical.h"
+#include "RigVMDeveloperModule.h"
+
+// #define FRIGVMASTPARSER_DEBUG_TRAVERSAL
 
 FRigVMExprAST::FRigVMExprAST(EType InType, UObject* InSubject)
 	: Name(NAME_None)
@@ -2481,58 +2484,138 @@ void FRigVMParserAST::Inline(URigVMGraph* InGraph, const TArray<URigVMNode*>& In
 		TMap<URigVMPin*, URigVMPin*> SourcePins;
 		TMap<URigVMPin*, TArray<URigVMPin*>>* TargetLinks;
 		TMap<URigVMPin*, TArray<URigVMPin*>>* SourceLinks;
+#ifdef FRIGVMASTPARSER_DEBUG_TRAVERSAL
+		FString Indentation;
+#endif
 
 		static URigVMPin* FindSourcePin(URigVMPin* InPin, LocalPinTraversalInfo& OutTraversalInfo)
 		{
+#ifdef FRIGVMASTPARSER_DEBUG_TRAVERSAL
+			FString& Indent = OutTraversalInfo.Indentation;
+#endif
+			if (InPin->GetDirection() != ERigVMPinDirection::Input &&
+				InPin->GetDirection() != ERigVMPinDirection::IO)
+			{
+				return nullptr;
+			}
+
 			if (URigVMPin* const* SourcePin = OutTraversalInfo.SourcePins.Find(InPin))
 			{
+#ifdef FRIGVMASTPARSER_DEBUG_TRAVERSAL
+				UE_LOG(LogRigVMDeveloper, Display, TEXT("%s= %s (%s)"), *Indent, *InPin->GetPinPath(), *InPin->GetCPPType());
+#endif
 				return *SourcePin;
 			}
 
-			TArray<URigVMLink*> SourceLinks = InPin->GetSourceLinks(false /* recursive */);
-			if (SourceLinks.Num() > 0)
-			{
-				URigVMPin* SourcePin = SourceLinks[0]->GetSourcePin();
+#ifdef FRIGVMASTPARSER_DEBUG_TRAVERSAL
+			UE_LOG(LogRigVMDeveloper, Display, TEXT("%s| %s (%s)"), *Indent, *InPin->GetPinPath(), *InPin->GetCPPType());
+#endif
+			TArray<FString> SegmentPath;
+			URigVMPin* SourcePin = nullptr;
 
-				// Only continue the recursion if the current source pin
-				// is on a reroute node.
-				if(SourcePin->GetNode()->IsA<URigVMRerouteNode>())
+			URigVMPin* ChildPin = InPin;
+			while (ChildPin != nullptr)
+			{
+				TArray<URigVMLink*> SourceLinks = ChildPin->GetSourceLinks(false /* recursive */);
+				if (SourceLinks.Num() > 0)
 				{
-					URigVMPin* ParentPin = SourcePin;
-					FString SegmentPath;
+					SourcePin = SourceLinks[0]->GetSourcePin();
 
-					while (ParentPin != nullptr)
+					// only continue the recursion on reroutes
+					if (SourcePin->GetNode()->IsA<URigVMRerouteNode>())
 					{
-						if (URigVMPin* SourceSourcePin = FindSourcePin(ParentPin, OutTraversalInfo))
+#ifdef FRIGVMASTPARSER_DEBUG_TRAVERSAL
+						Indent += TEXT("  ");
+#endif
+						if (URigVMPin* SourceSourcePin = FindSourcePin(SourcePin, OutTraversalInfo))
 						{
-							if (!SegmentPath.IsEmpty())
-							{
-								SourceSourcePin = SourceSourcePin->FindSubPin(SegmentPath);
-							}
-
-							OutTraversalInfo.SourcePins.FindOrAdd(InPin) = SourceSourcePin;
-							return SourceSourcePin;
+							SourcePin = SourceSourcePin;
 						}
+#ifdef FRIGVMASTPARSER_DEBUG_TRAVERSAL
+						Indent = Indent.LeftChop(2);
+#endif
+					}
 
-						if (SegmentPath.IsEmpty())
+					break;
+				}
+
+				URigVMPin* ParentPin = ChildPin->GetParentPin();
+				if (ParentPin)
+				{
+					// if we found a parent pin which has a source that is not a reroute
+					if (URigVMPin** ParentSourcePinPtr = OutTraversalInfo.SourcePins.Find(ParentPin))
+					{
+						URigVMPin* ParentSourcePin = *ParentSourcePinPtr;
+						if (ParentSourcePin)
 						{
-							SegmentPath = ParentPin->GetName();
+							if (!ParentSourcePin->GetNode()->IsA<URigVMRerouteNode>())
+							{
+								SourcePin = nullptr;
+								break;
+							}
 						}
-						else
+					}
+
+					SegmentPath.Push(ChildPin->GetName());
+#ifdef FRIGVMASTPARSER_DEBUG_TRAVERSAL
+					UE_LOG(LogRigVMDeveloper, Display, TEXT("%sPath %s"), *Indent, *URigVMPin::JoinPinPath(SegmentPath));
+					UE_LOG(LogRigVMDeveloper, Display, TEXT("%s> %s (%s)"), *Indent, *ParentPin->GetPinPath(), *ParentPin->GetCPPType());
+#endif
+				}
+				ChildPin = ParentPin;
+			}
+			
+			if (SourcePin)
+			{
+				while (!SegmentPath.IsEmpty())
+				{
+					FString Segment = SegmentPath.Pop();
+
+#ifdef FRIGVMASTPARSER_DEBUG_TRAVERSAL
+					if (!SegmentPath.IsEmpty())
+					{
+						UE_LOG(LogRigVMDeveloper, Display, TEXT("%sPath %s"), *Indent, *URigVMPin::JoinPinPath(SegmentPath));
+					}
+					UE_LOG(LogRigVMDeveloper, Display, TEXT("%sChecking segment %s"), *Indent, *Segment);
+#endif
+					if (URigVMPin* SourceSubPin = SourcePin->FindSubPin(Segment))
+					{
+#ifdef FRIGVMASTPARSER_DEBUG_TRAVERSAL
+						UE_LOG(LogRigVMDeveloper, Display, TEXT("%s< %s (%s)"), *Indent, *SourceSubPin->GetPinPath(), *SourceSubPin->GetCPPType());
+#endif
+						SourcePin = SourceSubPin;
+
+						// only continue the recursion on reroutes
+						if (SourcePin->GetNode()->IsA<URigVMRerouteNode>())
 						{
-							SegmentPath = URigVMPin::JoinPinPath(ParentPin->GetName(), SegmentPath);
+#ifdef FRIGVMASTPARSER_DEBUG_TRAVERSAL
+							Indent += TEXT("  ");
+#endif
+							if (URigVMPin* SourceSourceSubPin = FindSourcePin(SourcePin, OutTraversalInfo))
+							{
+								SourcePin = SourceSourceSubPin;
+							}
+#ifdef FRIGVMASTPARSER_DEBUG_TRAVERSAL
+							Indent = Indent.LeftChop(2);
+#endif
 						}
-
-						ParentPin = ParentPin->GetParentPin();
+					}
+					else
+					{
+						SourcePin = nullptr;
+						break;
 					}
 				}
-
-				OutTraversalInfo.SourcePins.FindOrAdd(InPin) = SourcePin;
-				return SourcePin;
 			}
 
-			OutTraversalInfo.SourcePins.FindOrAdd(InPin) = nullptr;
-			return nullptr;
+			OutTraversalInfo.SourcePins.FindOrAdd(InPin) = SourcePin;
+#ifdef FRIGVMASTPARSER_DEBUG_TRAVERSAL
+			if (SourcePin)
+			{
+				UE_LOG(LogRigVMDeveloper, Display, TEXT("%s%s --> %s"), *Indent, *InPin->GetPinPath(), *SourcePin->GetPinPath());
+			}
+#endif
+			return SourcePin;
 		}
 
 		static void VisitPin(URigVMPin* InPin, LocalPinTraversalInfo& OutTraversalInfo)
diff --git a/Engine/Source/Developer/RigVMDeveloper/Private/RigVMModel/RigVMController.cpp b/Engine/Source/Developer/RigVMDeveloper/Private/RigVMModel/RigVMController.cpp
--- a/Engine/Source/Developer/RigVMDeveloper/Private/RigVMModel/RigVMController.cpp
+++ b/Engine/Source/Developer/RigVMDeveloper/Private/RigVMModel/RigVMController.cpp
@@ -3759,13 +3759,6 @@ URigVMRerouteNode* URigVMController::AddFreeRerouteNode(bool bShowAsFullNode, co
 		SetPinDefaultValue(ValuePin, InDefaultValue, true, false, false);
 	}
 
-	ForEveryPinRecursively(ValuePin, [](URigVMPin* Pin) {
-		if (Pin->GetParentPin() != nullptr)
-		{
-			Pin->Direction = ERigVMPinDirection::Visible;
-		}
-	});
-
 	if (bSetupUndoRedo)
 	{
 		ActionStack->AddAction(FRigVMAddRerouteNodeAction(Node));
