Control Rig VM: Encapsulation, First steps

Implementation of relevant nodes for collapse feature, controller + actions as well as unit test

To be audited by [REDACTED_PERSON]

diff --git a/Engine/Source/Developer/RigVMDeveloper/Private/RigVMModel/Nodes/RigVMCollapseNode.cpp b/Engine/Source/Developer/RigVMDeveloper/Private/RigVMModel/Nodes/RigVMCollapseNode.cpp
new file mode 100644
--- /dev/null
+++ b/Engine/Source/Developer/RigVMDeveloper/Private/RigVMModel/Nodes/RigVMCollapseNode.cpp
@@ -0,0 +1,15 @@
+// Copyright Epic Games, Inc. All Rights Reserved.
+
+#include "RigVMModel/Nodes/RigVMCollapseNode.h"
+
+FString URigVMCollapseNode::GetNodeTitle() const
+{
+	// todo
+	return URigVMNode::GetNodeTitle();
+}
+
+FText URigVMCollapseNode::GetToolTipText() const
+{
+	// todo
+	return URigVMNode::GetToolTipText();
+}
diff --git a/Engine/Source/Developer/RigVMDeveloper/Private/RigVMModel/Nodes/RigVMFunctionEntryNode.cpp b/Engine/Source/Developer/RigVMDeveloper/Private/RigVMModel/Nodes/RigVMFunctionEntryNode.cpp
new file mode 100644
--- /dev/null
+++ b/Engine/Source/Developer/RigVMDeveloper/Private/RigVMModel/Nodes/RigVMFunctionEntryNode.cpp
@@ -0,0 +1,9 @@
+// Copyright Epic Games, Inc. All Rights Reserved.
+
+#include "RigVMModel/Nodes/RigVMFunctionEntryNode.h"
+
+bool URigVMFunctionEntryNode::IsDefinedAsVarying() const
+{ 
+	// todo
+	return true; 
+}
diff --git a/Engine/Source/Developer/RigVMDeveloper/Private/RigVMModel/Nodes/RigVMFunctionReturnNode.cpp b/Engine/Source/Developer/RigVMDeveloper/Private/RigVMModel/Nodes/RigVMFunctionReturnNode.cpp
new file mode 100644
--- /dev/null
+++ b/Engine/Source/Developer/RigVMDeveloper/Private/RigVMModel/Nodes/RigVMFunctionReturnNode.cpp
@@ -0,0 +1,9 @@
+// Copyright Epic Games, Inc. All Rights Reserved.
+
+#include "RigVMModel/Nodes/RigVMFunctionReturnNode.h"
+
+bool URigVMFunctionReturnNode::IsDefinedAsVarying() const
+{ 
+	// todo
+	return true; 
+}
diff --git a/Engine/Source/Developer/RigVMDeveloper/Private/RigVMModel/Nodes/RigVMLibraryNode.cpp b/Engine/Source/Developer/RigVMDeveloper/Private/RigVMModel/Nodes/RigVMLibraryNode.cpp
new file mode 100644
--- /dev/null
+++ b/Engine/Source/Developer/RigVMDeveloper/Private/RigVMModel/Nodes/RigVMLibraryNode.cpp
@@ -0,0 +1,61 @@
+// Copyright Epic Games, Inc. All Rights Reserved.
+
+#include "RigVMModel/Nodes/RigVMLibraryNode.h"
+#include "RigVMModel/Nodes/RigVMFunctionEntryNode.h"
+#include "RigVMModel/Nodes/RigVMFunctionReturnNode.h"
+#include "RigVMModel/RigVMGraph.h"
+
+bool URigVMLibraryNode::IsDefinedAsConstant() const
+{
+	return !bDefinedAsVarying;
+}
+
+bool URigVMLibraryNode::IsDefinedAsVarying() const
+{
+	return bDefinedAsVarying;
+}
+
+TArray<URigVMNode*> URigVMLibraryNode::GetContainedNodes() const
+{
+	if(URigVMGraph* Graph = GetContainedGraph())
+	{
+		return Graph->GetNodes();
+	}
+	return TArray<URigVMNode*>();
+}
+
+TArray<URigVMLink*> URigVMLibraryNode::GetContainedLinks() const
+{
+	if (URigVMGraph* Graph = GetContainedGraph())
+	{
+		return Graph->GetLinks();
+	}
+	return TArray<URigVMLink*>();
+
+}
+
+URigVMFunctionEntryNode* URigVMLibraryNode::GetEntryNode() const
+{
+	TArray<URigVMNode*> ContainedNodes = GetContainedNodes();
+	for (URigVMNode* ContainedNode : ContainedNodes)
+	{
+		if (URigVMFunctionEntryNode* EntryNode = Cast<URigVMFunctionEntryNode>(ContainedNode))
+		{
+			return EntryNode;
+		}
+	}
+	return nullptr;
+}
+
+URigVMFunctionReturnNode* URigVMLibraryNode::GetReturnNode() const
+{
+	TArray<URigVMNode*> ContainedNodes = GetContainedNodes();
+	for (URigVMNode* ContainedNode : ContainedNodes)
+	{
+		if (URigVMFunctionReturnNode* ReturnNode = Cast<URigVMFunctionReturnNode>(ContainedNode))
+		{
+			return ReturnNode;
+		}
+	}
+	return nullptr;
+}
diff --git a/Engine/Source/Developer/RigVMDeveloper/Private/RigVMModel/RigVMController.cpp b/Engine/Source/Developer/RigVMDeveloper/Private/RigVMModel/RigVMController.cpp
--- a/Engine/Source/Developer/RigVMDeveloper/Private/RigVMModel/RigVMController.cpp
+++ b/Engine/Source/Developer/RigVMDeveloper/Private/RigVMModel/RigVMController.cpp
@@ -2,6 +2,8 @@
 
 #include "RigVMModel/RigVMController.h"
 #include "RigVMModel/RigVMControllerActions.h"
+#include "RigVMModel/Nodes/RigVMFunctionEntryNode.h"
+#include "RigVMModel/Nodes/RigVMFunctionReturnNode.h"
 #include "RigVMCore/RigVMRegistry.h"
 #include "RigVMCore/RigVMExecuteContext.h"
 #include "RigVMCompiler/RigVMCompiler.h"
@@ -1237,7 +1239,7 @@ URigVMNode* URigVMController::EjectNodeFromPin(const FString& InPinPath, bool bS
 	FName OutputPinName = Injection->OutputPin->GetFName();
 
 	TMap<FName, FString> DefaultValues;
-	for (URigVMPin* PinOnNode : Injection->StructNode->Pins)
+	for (URigVMPin* PinOnNode : Injection->StructNode->GetPins())
 	{
 		if (PinOnNode->GetDirection() == ERigVMPinDirection::Input ||
 			PinOnNode->GetDirection() == ERigVMPinDirection::Visible ||
@@ -1551,7 +1553,7 @@ TArray<FName> URigVMController::ImportNodesFromText(const FString& InText, bool
 
 		if (bSetupUndoRedo)
 		{
-			ActionStack->AddAction(FRigVMRemoveNodeAction(CreatedNode));
+			ActionStack->AddAction(FRigVMRemoveNodeAction(CreatedNode, this));
 		}
 
 		if (URigVMStructNode* StructNode = Cast<URigVMStructNode>(CreatedNode))
@@ -1646,8 +1648,674 @@ FName URigVMController::GetUniqueName(const FName& InName, TFunction<bool(const
 	return *Name;
 }
 
+URigVMCollapseNode* URigVMController::CollapseNodes(const TArray<FName>& InNodeNames, const FString& InCollapseNodeName, bool bSetupUndoRedo)
+{
+	if (!IsValidGraph())
+	{
+		return nullptr;
+	}
+
+	TArray<URigVMNode*> Nodes;
+	for (const FName& NodeName : InNodeNames)
+	{
+		URigVMNode* Node = Graph->FindNodeByName(NodeName);
+		if (Node == nullptr)
+		{
+			ReportErrorf(TEXT("Cannot find node '%s'."), *NodeName.ToString());
+			return nullptr;
+		}
+		Nodes.AddUnique(Node);
+	}
+
+	return CollapseNodes(Nodes, InCollapseNodeName, bSetupUndoRedo);
+}
+
+TArray<URigVMNode*> URigVMController::ExpandLibraryNode(const FName& InNodeName, bool bSetupUndoRedo)
+{
+	if (!IsValidGraph())
+	{
+		return TArray<URigVMNode*>();
+	}
+
+	URigVMNode* Node = Graph->FindNodeByName(InNodeName);
+	if (Node == nullptr)
+	{
+		ReportErrorf(TEXT("Cannot find collapse node '%s'."), *InNodeName.ToString());
+		return TArray<URigVMNode*>();
+	}
+
+	URigVMLibraryNode* LibNode = Cast<URigVMCollapseNode>(Node);
+	if (LibNode == nullptr)
+	{
+		ReportErrorf(TEXT("Node '%s' is not a library node (not collapse nor function)."), *InNodeName.ToString());
+		return TArray<URigVMNode*>();
+	}
+
+	return ExpandLibraryNode(LibNode, bSetupUndoRedo);
+}
+
 #endif
 
+URigVMCollapseNode* URigVMController::CollapseNodes(const TArray<URigVMNode*>& InNodes, const FString& InCollapseNodeName, bool bSetupUndoRedo)
+{
+	TArray<URigVMNode*> Nodes;
+	for (URigVMNode* Node : InNodes)
+	{
+		if (!IsValidNodeForGraph(Node))
+		{
+			return nullptr;
+		}
+
+		// filter out certain nodes
+		if (Node->IsEvent())
+		{
+			continue;
+		}
+
+		Nodes.Add(Node);
+	}
+
+	if (Nodes.Num() == 0)
+	{
+		return nullptr;
+	}
+
+	FBox2D Bounds;
+	TArray<FName> NodeNames;
+	for (URigVMNode* Node : Nodes)
+	{
+		NodeNames.Add(Node->GetFName());
+		Bounds += Node->GetPosition();
+	}
+
+  	FVector2D Diagonal = Bounds.Max - Bounds.Min;
+	FVector2D Center = (Bounds.Min + Bounds.Max) * 0.5f;
+	
+	FRigVMCollapseNodesAction CollapseAction;
+	CollapseAction.Title = TEXT("Collapse Nodes");
+
+	if (bSetupUndoRedo)
+	{
+		ActionStack->BeginAction(CollapseAction);
+	}
+
+	FString CollapseNodeName = GetValidNodeName(InCollapseNodeName.IsEmpty() ? FString(TEXT("CollapseNode")) : InCollapseNodeName);
+	URigVMCollapseNode* CollapseNode = NewObject<URigVMCollapseNode>(Graph, *CollapseNodeName);
+	CollapseNode->ContainedGraph = NewObject<URigVMGraph>(CollapseNode, TEXT("ContainedGraph"));
+	CollapseNode->Position = Center;
+	Graph->Nodes.Add(CollapseNode);
+
+	bool bContainsOutputs = false;
+
+	TArray<URigVMPin*> PinsToCollapse;
+	TMap<URigVMPin*, URigVMPin*> CollapsedPins;
+	TArray<URigVMLink*> LinksToRewire;
+	TArray<URigVMLink*> AllLinks = Graph->GetLinks();
+
+	// find all pins to collapse. we need this to find out if
+	// we might have a parent pin of a given linked pin already 
+	// collapsed.
+	for (URigVMLink* Link : AllLinks)
+	{
+		bool bSourceToBeCollapsed = Nodes.Contains(Link->GetSourcePin()->GetNode());
+		bool bTargetToBeCollapsed = Nodes.Contains(Link->GetTargetPin()->GetNode());
+		if (bSourceToBeCollapsed == bTargetToBeCollapsed)
+		{
+			continue;
+		}
+
+		URigVMPin* PinToCollapse = Link->GetSourcePin();
+		PinsToCollapse.AddUnique(PinToCollapse);
+		LinksToRewire.Add(Link);
+	}
+
+	// now looper over the links to be rewired
+	for (URigVMLink* Link : LinksToRewire)
+	{
+		bool bSourceToBeCollapsed = Nodes.Contains(Link->GetSourcePin()->GetNode());
+		bool bTargetToBeCollapsed = Nodes.Contains(Link->GetTargetPin()->GetNode());
+
+		URigVMPin* PinToCollapse = bSourceToBeCollapsed ? Link->GetSourcePin() : Link->GetTargetPin();
+		if (CollapsedPins.Contains(PinToCollapse))
+		{
+			continue;
+		}
+
+		// for links that connect to the right side of the collapse
+		// node, we need to skip sub pins of already exposed pins
+		if (bSourceToBeCollapsed)
+		{
+			bool bParentPinCollapsed = false;
+			URigVMPin* ParentPin = PinToCollapse->GetParentPin();
+			while (ParentPin != nullptr)
+			{
+				if (PinsToCollapse.Contains(ParentPin))
+				{
+					bParentPinCollapsed = true;
+					break;
+				}
+				ParentPin = ParentPin->GetParentPin();
+			}
+
+			if (bParentPinCollapsed)
+			{
+				continue;
+			}
+		}
+
+		FName PinName = GetUniqueName(PinToCollapse->GetFName(), [CollapseNode](const FName& InName) {
+			return CollapseNode->FindPin(InName.ToString()) == nullptr;
+		});
+
+		URigVMPin* CollapsedPin = NewObject<URigVMPin>(CollapseNode, PinName);
+		ConfigurePinFromPin(CollapsedPin, PinToCollapse);
+
+		if (CollapsedPin->IsStruct())
+		{
+			AddPinsForStruct(CollapsedPin->GetScriptStruct(), CollapseNode, CollapsedPin, CollapsedPin->GetDirection(), FString(), false);
+		}
+
+		bContainsOutputs = bContainsOutputs || bSourceToBeCollapsed;
+
+		CollapseNode->Pins.Add(CollapsedPin);
+
+		FPinState PinState = GetPinState(PinToCollapse);
+		ApplyPinState(CollapsedPin, PinState);
+
+		CollapsedPins.Add(PinToCollapse, CollapsedPin);
+	}
+
+	Notify(ERigVMGraphNotifType::NodeAdded, CollapseNode);
+
+	URigVMFunctionEntryNode* EntryNode = NewObject<URigVMFunctionEntryNode>(CollapseNode->ContainedGraph, TEXT("Entry"));
+	CollapseNode->ContainedGraph->Nodes.Add(EntryNode);
+	EntryNode->Position = -Diagonal * 0.5f - FVector2D(160.f, 0.f);
+	RefreshFunctionPins(EntryNode);
+
+	URigVMFunctionReturnNode* ReturnNode = nullptr;
+	if (bContainsOutputs)
+	{
+		ReturnNode = NewObject<URigVMFunctionReturnNode>(CollapseNode->ContainedGraph, TEXT("Return"));
+		CollapseNode->ContainedGraph->Nodes.Add(ReturnNode);
+		ReturnNode->Position = FVector2D(Diagonal.X, -Diagonal.Y) * 0.5f + FVector2D(160.f, 0.f);
+		RefreshFunctionPins(ReturnNode);
+	}
+
+	// create the new nodes within the collapse node
+	if (bSetupUndoRedo)
+	{
+		TArray<FName> ContainedNodeNames;
+		{
+			FString TextContent = ExportNodesToText(NodeNames);
+
+			TGuardValue<URigVMGraph*> GraphGuard(Graph, CollapseNode->ContainedGraph);
+			ContainedNodeNames = ImportNodesFromText(TextContent, false);
+
+			// move the nodes to the right place
+			for (const FName& ContainedNodeName : ContainedNodeNames)
+			{
+				if (URigVMNode* ContainedNode = CollapseNode->ContainedGraph->FindNodeByName(ContainedNodeName))
+				{
+					ContainedNode->Position -= Center;
+				}
+			}
+
+			for (URigVMLink* LinkToRewire : LinksToRewire)
+			{
+				URigVMPin* SourcePin = LinkToRewire->GetSourcePin();
+				URigVMPin* TargetPin = LinkToRewire->GetTargetPin();
+
+				if (Nodes.Contains(SourcePin->GetNode()))
+				{
+					// if the parent pin of this was collapsed
+					// it's possible that the child pin wasn't.
+					if (!CollapsedPins.Contains(SourcePin))
+					{
+						continue;
+					}
+
+					URigVMPin* CollapsedPin = CollapsedPins.FindChecked(SourcePin);
+					SourcePin = CollapseNode->ContainedGraph->FindPin(SourcePin->GetPinPath());
+					TargetPin = ReturnNode->FindPin(CollapsedPin->GetName());
+				}
+				else
+				{
+					URigVMPin* CollapsedPin = CollapsedPins.FindChecked(TargetPin);
+					SourcePin = EntryNode->FindPin(CollapsedPin->GetName());
+					TargetPin = CollapseNode->ContainedGraph->FindPin(TargetPin->GetPinPath());
+				}
+
+				if (SourcePin && TargetPin)
+				{
+					if (!SourcePin->IsLinkedTo(TargetPin))
+					{
+						AddLink(SourcePin, TargetPin, false);
+					}
+				}
+			}
+		}
+
+		TArray<URigVMLink*> RewiredLinks;
+		for (URigVMLink* LinkToRewire : LinksToRewire)
+		{
+			if (RewiredLinks.Contains(LinkToRewire))
+			{
+				continue;
+			}
+
+			URigVMPin* SourcePin = LinkToRewire->GetSourcePin();
+			URigVMPin* TargetPin = LinkToRewire->GetTargetPin();
+
+			if (Nodes.Contains(SourcePin->GetNode()))
+			{
+				FString SegmentPath;
+				URigVMPin* PinToCheck = SourcePin;
+
+				URigVMPin** CollapsedPinPtr = CollapsedPins.Find(PinToCheck);
+				while (CollapsedPinPtr == nullptr)
+				{
+					if (SegmentPath.IsEmpty())
+					{
+						SegmentPath = PinToCheck->GetName();
+					}
+					else
+					{
+						SegmentPath = URigVMPin::JoinPinPath(PinToCheck->GetName(), SegmentPath);
+					}
+
+					PinToCheck = PinToCheck->GetParentPin();
+					check(PinToCheck);
+
+					CollapsedPinPtr = CollapsedPins.Find(PinToCheck);
+				}
+
+				URigVMPin* CollapsedPin = *CollapsedPinPtr;
+				check(CollapsedPin);
+
+				if (!SegmentPath.IsEmpty())
+				{
+					CollapsedPin = CollapsedPin->FindSubPin(SegmentPath);
+					check(CollapsedPin);
+				}
+
+				TArray<URigVMLink*> TargetLinks = SourcePin->GetTargetLinks(false);
+				for (URigVMLink* TargetLink : TargetLinks)
+				{
+					TargetPin = TargetLink->GetTargetPin();
+					if (!CollapsedPin->IsLinkedTo(TargetPin))
+					{
+						AddLink(CollapsedPin, TargetPin, false);
+					}
+				}
+				RewiredLinks.Append(TargetLinks);
+			}
+			else
+			{
+				URigVMPin* CollapsedPin = CollapsedPins.FindChecked(TargetPin);
+				if (!SourcePin->IsLinkedTo(CollapsedPin))
+				{
+					AddLink(SourcePin, CollapsedPin, false);
+				}
+			}
+
+			RewiredLinks.Add(LinkToRewire);
+		}
+
+
+		for (const FName& NodeToRemove : NodeNames)
+		{
+			RemoveNodeByName(NodeToRemove, false, true);
+		}
+	}
+
+	if (bSetupUndoRedo)
+	{
+		ActionStack->EndAction(CollapseAction);
+	}
+
+	return CollapseNode;
+}
+
+TArray<URigVMNode*> URigVMController::ExpandLibraryNode(URigVMLibraryNode* InNode, bool bSetupUndoRedo)
+{
+	if (!IsValidNodeForGraph(InNode))
+	{
+		return TArray<URigVMNode*>();
+	}
+
+	TArray<URigVMNode*> ContainedNodes = InNode->GetContainedNodes();
+	TArray<URigVMLink*> ContainedLinks = InNode->GetContainedLinks();
+	if (ContainedNodes.Num() == 0)
+	{
+		return TArray<URigVMNode*>();
+	}
+
+	FRigVMExpandNodeAction ExpandAction;
+	ExpandAction.Title = FString::Printf(TEXT("Expand '%s' Node"), *InNode->GetName());
+
+	if (bSetupUndoRedo)
+	{
+		ActionStack->BeginAction(ExpandAction);
+	}
+
+	TArray<FName> NodeNames;
+	FBox2D Bounds;
+	{
+		TArray<URigVMNode*> FilteredNodes;
+		for (URigVMNode* Node : ContainedNodes)
+		{
+			if (Cast<URigVMFunctionEntryNode>(Node) != nullptr ||
+				Cast<URigVMFunctionReturnNode>(Node) != nullptr)
+			{
+				continue;
+			}
+
+			NodeNames.Add(Node->GetFName());
+			FilteredNodes.Add(Node);
+			Bounds += Node->GetPosition();
+		}
+		ContainedNodes = FilteredNodes;
+	}
+
+	if (ContainedNodes.Num() == 0)
+	{
+		return TArray<URigVMNode*>();
+	}
+
+	FVector2D Diagonal = Bounds.Max - Bounds.Min;
+	FVector2D Center = (Bounds.Min + Bounds.Max) * 0.5f;
+
+	FString TextContent;
+	{
+		TGuardValue<URigVMGraph*> GraphGuard(Graph, InNode->GetContainedGraph());
+		TextContent = ExportNodesToText(NodeNames);
+	}
+
+	TArray<FName> ExpandedNodeNames = ImportNodesFromText(TextContent, false);
+	TArray<URigVMNode*> ExpandedNodes;
+	for (const FName& ExpandedNodeName : ExpandedNodeNames)
+	{
+		URigVMNode* ExpandedNode = Graph->FindNodeByName(ExpandedNodeName);
+		check(ExpandedNode);
+		ExpandedNodes.Add(ExpandedNode);
+	}
+
+	check(ExpandedNodeNames.Num() == NodeNames.Num());
+
+	TMap<FName, FName> NodeNameMap;
+	for (int32 NodeNameIndex = 0; NodeNameIndex < NodeNames.Num(); NodeNameIndex++)
+	{
+		NodeNameMap.Add(NodeNames[NodeNameIndex], ExpandedNodeNames[NodeNameIndex]);
+	}
+
+	// a) store all of the pin defaults off the library node
+	TMap<FString, FPinState> PinStates = GetPinStates(InNode);
+
+	// b) create a map of new links to create by following the links to / from the library node
+	TMap<FString, TArray<FString>> ToLibraryNode;
+	TMap<FString, TArray<FString>> FromLibraryNode;
+
+	TArray<URigVMLink*> LibraryLinks = InNode->GetLinks();
+	for (URigVMLink* Link : LibraryLinks)
+	{
+		if (Link->GetTargetPin()->GetNode() == InNode)
+		{
+			FString NodeName, PinPath;
+			URigVMPin::SplitPinPathAtStart(Link->GetTargetPin()->GetPinPath(), NodeName, PinPath);
+			ToLibraryNode.FindOrAdd(PinPath).Add(Link->GetSourcePin()->GetPinPath());
+		}
+		else
+		{
+			FString NodeName, PinPath;
+			URigVMPin::SplitPinPathAtStart(Link->GetSourcePin()->GetPinPath(), NodeName, PinPath);
+			FromLibraryNode.FindOrAdd(PinPath).Add(Link->GetTargetPin()->GetPinPath());
+		}
+	}
+
+	// c) create a map from the entry node to the contained graph
+	TMap<FString, TArray<FString>> FromEntryNode;
+	if (URigVMFunctionEntryNode* EntryNode = InNode->GetEntryNode())
+	{
+		TArray<URigVMLink*> EntryLinks = EntryNode->GetLinks();
+		for (URigVMLink* Link : EntryLinks)
+		{
+			if (Link->GetSourcePin()->GetNode() != EntryNode)
+			{
+				continue;
+			}
+
+			FString NodeName, PinPath;
+			URigVMPin::SplitPinPathAtStart(Link->GetSourcePin()->GetPinPath(), NodeName, PinPath);
+
+			TArray<FString>& LinkedPins = FromEntryNode.FindOrAdd(PinPath);
+
+			URigVMPin::SplitPinPathAtStart(Link->GetTargetPin()->GetPinPath(), NodeName, PinPath);
+			NodeName = NodeNameMap.FindChecked(*NodeName).ToString();
+			LinkedPins.Add(URigVMPin::JoinPinPath(NodeName, PinPath));
+		}
+	}
+
+	// d) create a map from the contained graph from to the return node
+	TMap<FString, TArray<FString>> ToReturnNode;
+	if (URigVMFunctionReturnNode* ReturnNode = InNode->GetReturnNode())
+	{
+		TArray<URigVMLink*> ReturnLinks = ReturnNode->GetLinks();
+		for (URigVMLink* Link : ReturnLinks)
+		{
+			if (Link->GetTargetPin()->GetNode() != ReturnNode)
+			{
+				continue;
+			}
+
+			FString NodeName, PinPath;
+			URigVMPin::SplitPinPathAtStart(Link->GetTargetPin()->GetPinPath(), NodeName, PinPath);
+
+			TArray<FString>& LinkedPins = ToReturnNode.FindOrAdd(PinPath);
+
+			URigVMPin::SplitPinPathAtStart(Link->GetSourcePin()->GetPinPath(), NodeName, PinPath);
+			NodeName = NodeNameMap.FindChecked(*NodeName).ToString();
+			LinkedPins.Add(URigVMPin::JoinPinPath(NodeName, PinPath));
+		}
+	}
+
+	// e) restore all pin states on pins linked to the entry node
+	for (const TPair<FString, TArray<FString>>& FromEntryPair : FromEntryNode)
+	{
+		FString EntryPinPath = FromEntryPair.Key;
+		const FPinState* CollapsedPinState = PinStates.Find(EntryPinPath);
+		if (CollapsedPinState == nullptr)
+		{
+			continue;
+		}
+
+		for (const FString& EntryTargetLinkPinPath : FromEntryPair.Value)
+		{
+			FString TargetNodeName, TargetPinPath;
+			URigVMPin::SplitPinPathAtStart(EntryTargetLinkPinPath, TargetNodeName, TargetPinPath);
+
+			FString ExpandedNodeName = NodeNameMap.FindChecked(*TargetNodeName).ToString();
+			FString ExpandedPinPath = URigVMPin::JoinPinPath(ExpandedNodeName, TargetPinPath);
+
+			if (URigVMPin* TargetPin = GetGraph()->FindPin(ExpandedPinPath))
+			{
+				ApplyPinState(TargetPin, *CollapsedPinState);
+			}
+		}
+	}
+
+	// f) remap all output / source pins and create a final list of links to create
+	TMap<FString, FString> RemappedSourcePins;;
+	TArray<URigVMPin*> LibraryPins = InNode->GetAllPinsRecursively();
+	for (URigVMPin* LibraryPin : LibraryPins)
+	{
+		FString LibraryPinPath = LibraryPin->GetPinPath();
+		FString LibraryNodeName;
+		URigVMPin::SplitPinPathAtStart(LibraryPinPath, LibraryNodeName, LibraryPinPath);
+
+		if (LibraryPin->GetDirection() == ERigVMPinDirection::Input)
+		{
+			if (const TArray<FString>* LibraryPinLinksPtr = ToLibraryNode.Find(LibraryPinPath))
+			{
+				const TArray<FString>& LibraryPinLinks = *LibraryPinLinksPtr;
+				ensure(LibraryPinLinks.Num() == 1);
+
+				RemappedSourcePins.FindOrAdd(LibraryPinPath) = LibraryPinLinks[0];
+			}
+		}
+		else if (LibraryPin->GetDirection() == ERigVMPinDirection::Output)
+		{
+			if (const TArray<FString>* LibraryPinLinksPtr = ToReturnNode.Find(LibraryPinPath))
+			{
+				const TArray<FString>& LibraryPinLinks = *LibraryPinLinksPtr;
+				ensure(LibraryPinLinks.Num() == 1);
+
+				FString SourcePinNodeName, SourcePinPath;
+				URigVMPin::SplitPinPathAtStart(LibraryPinLinks[0], SourcePinNodeName, SourcePinPath);
+
+				SourcePinNodeName = NodeNameMap.FindChecked(*SourcePinNodeName).ToString();
+
+				RemappedSourcePins.FindOrAdd(LibraryPinPath) = URigVMPin::JoinPinPath(SourcePinNodeName, SourcePinPath);
+			}
+		}
+	}
+
+	// g) re-establish all of the links going to the left of the library node
+	for (const TPair<FString, TArray<FString>>& FromEntryNodePair : FromEntryNode)
+	{
+		FString EntryPinPath = FromEntryNodePair.Key;
+		FString EntryPinPathSuffix;
+
+		const FString* RemappedSourcePin = RemappedSourcePins.Find(EntryPinPath);
+		while (RemappedSourcePin == nullptr)
+		{
+			FString LastSegment;
+			if (!URigVMPin::SplitPinPathAtEnd(EntryPinPath, EntryPinPath, LastSegment))
+			{
+				break;
+			}
+
+			if (EntryPinPathSuffix.IsEmpty())
+			{
+				EntryPinPathSuffix = LastSegment;
+			}
+			else
+			{
+				EntryPinPathSuffix = URigVMPin::JoinPinPath(LastSegment, EntryPinPathSuffix);
+			}
+
+			RemappedSourcePin = RemappedSourcePins.Find(EntryPinPath);
+		}
+
+		if (RemappedSourcePin == nullptr)
+		{
+			continue;
+		}
+
+		FString RemappedSourcePinPath = *RemappedSourcePin;
+		if (!EntryPinPathSuffix.IsEmpty())
+		{
+			RemappedSourcePinPath = URigVMPin::JoinPinPath(RemappedSourcePinPath, EntryPinPathSuffix);
+		}
+
+		for (const FString& FromEntryNodeTargetPinPath : FromEntryNodePair.Value)
+		{
+			FString TargetPinNodeName, TargetPinPath;
+			URigVMPin::SplitPinPathAtStart(FromEntryNodeTargetPinPath, TargetPinNodeName, TargetPinPath);
+
+			TargetPinNodeName = NodeNameMap.FindChecked(*TargetPinNodeName).ToString();
+
+			URigVMPin* SourcePin = GetGraph()->FindPin(*RemappedSourcePinPath);
+			URigVMPin* TargetPin = GetGraph()->FindPin(URigVMPin::JoinPinPath(TargetPinNodeName, TargetPinPath));
+			if (SourcePin && TargetPin)
+			{
+				if (!SourcePin->IsLinkedTo(TargetPin))
+				{
+					AddLink(SourcePin, TargetPin, false);
+				}
+			}
+		}
+	}
+
+	// h) re-establish all of the links going from the right of the library node
+	for (const TPair<FString, TArray<FString>>& FromLibraryNodePair : FromLibraryNode)
+	{
+		FString FromLibraryNodePinPath = FromLibraryNodePair.Key;
+		FString FromLibraryNodePinPathSuffix;
+
+		const FString* RemappedSourcePin = RemappedSourcePins.Find(FromLibraryNodePinPath);
+		while (RemappedSourcePin == nullptr)
+		{
+			FString LastSegment;
+			if (!URigVMPin::SplitPinPathAtEnd(FromLibraryNodePinPath, FromLibraryNodePinPath, LastSegment))
+			{
+				break;
+			}
+
+			if (FromLibraryNodePinPathSuffix.IsEmpty())
+			{
+				FromLibraryNodePinPathSuffix = LastSegment;
+			}
+			else
+			{
+				FromLibraryNodePinPathSuffix = URigVMPin::JoinPinPath(LastSegment, FromLibraryNodePinPathSuffix);
+			}
+
+			RemappedSourcePin = RemappedSourcePins.Find(FromLibraryNodePinPath);
+		}
+
+		if (RemappedSourcePin == nullptr)
+		{
+			continue;
+		}
+
+		FString RemappedSourcePinPath = *RemappedSourcePin;
+		if (!FromLibraryNodePinPathSuffix.IsEmpty())
+		{
+			RemappedSourcePinPath = URigVMPin::JoinPinPath(RemappedSourcePinPath, FromLibraryNodePinPathSuffix);
+		}
+
+		for (const FString& FromLibraryNodeTargetPinPath : FromLibraryNodePair.Value)
+		{
+			FString TargetPinNodeName, TargetPinPath;
+			URigVMPin::SplitPinPathAtStart(FromLibraryNodeTargetPinPath, TargetPinNodeName, TargetPinPath);
+
+			URigVMPin* SourcePin = GetGraph()->FindPin(*RemappedSourcePinPath);
+			URigVMPin* TargetPin = GetGraph()->FindPin(URigVMPin::JoinPinPath(TargetPinNodeName, TargetPinPath));
+			if (SourcePin && TargetPin)
+			{
+				if (!SourcePin->IsLinkedTo(TargetPin))
+				{
+					AddLink(SourcePin, TargetPin, false);
+				}
+			}
+		}
+	}
+
+	// i) remove the library node from the graph
+	RemoveNode(InNode, false, true);
+
+	if (bSetupUndoRedo)
+	{
+		ActionStack->EndAction(ExpandAction);
+	}
+
+	return ExpandedNodes;
+}
+
+void URigVMController::RefreshFunctionPins(URigVMNode* InNode)
+{
+	URigVMFunctionEntryNode* EntryNode = Cast<URigVMFunctionEntryNode>(InNode);
+	URigVMFunctionReturnNode* ReturnNode = Cast<URigVMFunctionReturnNode>(InNode);
+
+	if (EntryNode || ReturnNode)
+	{
+		TArray<URigVMLink*> Links = InNode->GetLinks();
+		DetachLinksFromPinObjects(&Links);
+		RepopulatePinsOnNode(InNode, false);
+		ReattachLinksToPinObjects(false, &Links);
+	}
+}
+
 bool URigVMController::RemoveNode(URigVMNode* InNode, bool bSetupUndoRedo, bool bRecursive)
 {
 	if (!IsValidNodeForGraph(InNode))
@@ -1655,6 +2323,16 @@ bool URigVMController::RemoveNode(URigVMNode* InNode, bool bSetupUndoRedo, bool
 		return false;
 	}
 
+	if (bSetupUndoRedo)
+	{
+		// don't allow deletion of function entry / return nodes
+		if (Cast<URigVMFunctionEntryNode>(InNode) != nullptr ||
+			Cast<URigVMFunctionReturnNode>(InNode) != nullptr)
+		{
+			return false;
+		}
+	}
+
 	TGuardValue<bool> GuardCompactness(bIgnoreRerouteCompactnessChanges, true);
 
 	FRigVMBaseAction Action;
@@ -1688,7 +2366,7 @@ bool URigVMController::RemoveNode(URigVMNode* InNode, bool bSetupUndoRedo, bool
 	{
 		SelectNode(InNode, false, bSetupUndoRedo);
 
-		for (URigVMPin* Pin : InNode->Pins)
+		for (URigVMPin* Pin : InNode->GetPins())
 		{
 			TArray<URigVMInjectionInfo*> InjectedNodes = Pin->GetInjectedNodes();
 			for (URigVMInjectionInfo* InjectedNode : InjectedNodes)
@@ -1705,7 +2383,7 @@ bool URigVMController::RemoveNode(URigVMNode* InNode, bool bSetupUndoRedo, bool
 
 	if (bSetupUndoRedo)
 	{
-		ActionStack->AddAction(FRigVMRemoveNodeAction(InNode));
+		ActionStack->AddAction(FRigVMRemoveNodeAction(InNode, this));
 	}
 
 	Graph->Nodes.Remove(InNode);
@@ -4105,7 +4783,7 @@ void URigVMController::ForEveryPinRecursively(URigVMPin* InPin, TFunction<void(U
 
 void URigVMController::ForEveryPinRecursively(URigVMNode* InNode, TFunction<void(URigVMPin*)> OnEachPinFunction)
 {
-	for (URigVMPin* Pin : InNode->Pins)
+	for (URigVMPin* Pin : InNode->GetPins())
 	{
 		ForEveryPinRecursively(Pin, OnEachPinFunction);
 	}
@@ -4570,12 +5248,22 @@ FProperty* URigVMController::FindPropertyForPin(const FString& InPinPath)
 	return nullptr;
 }
 
-int32 URigVMController::DetachLinksFromPinObjects()
+int32 URigVMController::DetachLinksFromPinObjects(const TArray<URigVMLink*>* InLinks)
 {
 	check(Graph);
 	TGuardValue<bool> SuspendNotifs(bSuspendNotifications, true);
 
-	for (URigVMLink* Link : Graph->Links)
+	TArray<URigVMLink*> Links;
+	if (InLinks)
+	{
+		Links = *InLinks;
+	}
+	else
+	{
+		Links = Graph->Links;
+	}
+
+	for (URigVMLink* Link : Links)
 	{
 		URigVMPin* SourcePin = Link->GetSourcePin();
 		URigVMPin* TargetPin = Link->GetTargetPin();
@@ -4596,19 +5284,31 @@ int32 URigVMController::DetachLinksFromPinObjects()
 		Link->TargetPin = nullptr;
 	}
 
-	return Graph->Links.Num();
+	return Links.Num();
 }
 
-int32 URigVMController::ReattachLinksToPinObjects(bool bFollowCoreRedirectors)
+int32 URigVMController::ReattachLinksToPinObjects(bool bFollowCoreRedirectors, const TArray<URigVMLink*>* InLinks)
 {
 	check(Graph);
 	TGuardValue<bool> SuspendNotifs(bSuspendNotifications, true);
 	FScopeLock Lock(&PinPathCoreRedirectorsLock);
 
+	bool bReplacingAllLinks = false;
+	TArray<URigVMLink*> Links;
+	if (InLinks)
+	{
+		Links = *InLinks;
+	}
+	else
+	{
+		Links = Graph->Links;
+		bReplacingAllLinks = true;
+	}
+
 	TMap<FString, FString> RedirectedPinPaths;
 	if (bFollowCoreRedirectors)
 	{
-		for (URigVMLink* Link : Graph->Links)
+		for (URigVMLink* Link : Links)
 		{
 			FString RedirectedSourcePinPath;
 			if (ShouldRedirectPin(Link->SourcePinPath, RedirectedSourcePinPath))
@@ -4626,7 +5326,7 @@ int32 URigVMController::ReattachLinksToPinObjects(bool bFollowCoreRedirectors)
 
 	// fix up the pin links based on the persisted data
 	TArray<URigVMLink*> NewLinks;
-	for (URigVMLink* Link : Graph->Links)
+	for (URigVMLink* Link : Links)
 	{
 		if (FString* RedirectedSourcePinPath = OutputPinRedirectors.Find(Link->SourcePinPath))
 		{
@@ -4665,7 +5365,24 @@ int32 URigVMController::ReattachLinksToPinObjects(bool bFollowCoreRedirectors)
 		TargetPin->Links.AddUnique(Link);
 		NewLinks.Add(Link);
 	}
-	Graph->Links = NewLinks;
+
+	if (bReplacingAllLinks)
+	{
+		Graph->Links = NewLinks;
+	}
+	else
+	{
+		// if we are running of a subset of links
+		// find the ones we weren't able to connect
+		// again and remove them.
+		for (URigVMLink* Link : Links)
+		{
+			if (!NewLinks.Contains(Link))
+			{
+				Graph->Links.Remove(Link);
+			}
+		}
+	}
 
 	InputPinRedirectors.Reset();
 	OutputPinRedirectors.Reset();
@@ -4851,7 +5568,7 @@ bool URigVMController::ShouldRedirectPin(const FString& InOldPinPath, FString& I
 	return false;
 }
 
-void URigVMController::RepopulatePinsOnNode(URigVMNode* InNode)
+void URigVMController::RepopulatePinsOnNode(URigVMNode* InNode, bool bFollowCoreRedirectors)
 {
 	if (InNode == nullptr)
 	{
@@ -4860,14 +5577,9 @@ void URigVMController::RepopulatePinsOnNode(URigVMNode* InNode)
 	}
 
 	URigVMStructNode* StructNode = Cast<URigVMStructNode>(InNode);
-
-	// reroute node may also contain a struct value pin that need to be refreshed
 	URigVMRerouteNode* RerouteNode = Cast<URigVMRerouteNode>(InNode);
-	
-	if (StructNode == nullptr && RerouteNode == nullptr)
-	{
-		return;
-	}
+	URigVMFunctionEntryNode* EntryNode = Cast<URigVMFunctionEntryNode>(InNode);
+	URigVMFunctionReturnNode* ReturnNode = Cast<URigVMFunctionReturnNode>(InNode);
 
 	TGuardValue<bool> SuspendNotifs(bSuspendNotifications, true);
 	FScopeLock Lock(&PinPathCoreRedirectorsLock);
@@ -4875,43 +5587,12 @@ void URigVMController::RepopulatePinsOnNode(URigVMNode* InNode)
 	check(Graph);
 
 	// step 1/3: keep a record of the current state of the node's pins
-	TMap<FString, FString> DefaultValues;
-	TArray<URigVMPin*> AllPins = InNode->GetAllPinsRecursively();
-	for (URigVMPin* Pin : AllPins)
-	{
-		FString DefaultValue = Pin->GetDefaultValue();
-		if (!DefaultValue.IsEmpty())
-		{
-			FString PinPath, NodeName;
-			URigVMPin::SplitPinPathAtStart(Pin->GetPinPath(), NodeName, PinPath);
-			DefaultValues.Add(PinPath, DefaultValue);
-		}
-	}
-
-	TMap<FString, bool> ExpansionStates;
-	for (URigVMPin* Pin : AllPins)
-	{
-		FString PinPath, NodeName;
-		URigVMPin::SplitPinPathAtStart(Pin->GetPinPath(), NodeName, PinPath);
-		ExpansionStates.Add(PinPath, Pin->IsExpanded());
-	}
-
-	TMap<FString, TArray<URigVMInjectionInfo*>> InjectionInfos;
-	for (URigVMPin* Pin : InNode->Pins)
-	{
-		if (Pin->HasInjectedNodes())
-		{
-			InjectionInfos.Add(Pin->GetName(), Pin->GetInjectedNodes());
-		}
-	}
-
-	TMap<FString, FString> BoundVariables;
-	for (URigVMPin* Pin : AllPins)
+	TMap<FString, FString> RedirectedPinPaths;
+	if (bFollowCoreRedirectors)
 	{
-		FString PinPath, NodeName;
-		URigVMPin::SplitPinPathAtStart(Pin->GetPinPath(), NodeName, PinPath);
-		BoundVariables.Add(PinPath, Pin->GetBoundVariablePath());
+		RedirectedPinPaths = GetRedirectedPinPaths(InNode);
 	}
+	TMap<FString, FPinState> PinStates = GetPinStates(InNode);
 
 	// also in case this node is part of an injection
 	FName InjectionInputPinName = NAME_None;
@@ -4922,61 +5603,10 @@ void URigVMController::RepopulatePinsOnNode(URigVMNode* InNode)
 		InjectionOutputPinName = InjectionInfo->OutputPin->GetFName();
 	}
 
-	// create a map for pin paths to their respective redirected pins
-	TMap<FString, FString> RedirectedPinPaths;
-	UScriptStruct* OwningStruct = nullptr;
-	if (StructNode)
-	{
-		OwningStruct = StructNode->GetScriptStruct();
-	}
-	else if (RerouteNode)
-	{
-		URigVMPin* ValuePin = RerouteNode->Pins[0];
-		if (ValuePin->IsStruct())
-		{
-			OwningStruct = ValuePin->GetScriptStruct();
-		}
-	}
-	else
-	{
-		checkNoEntry();
-	}
-
-	if (OwningStruct)
-	{
-		for (URigVMPin* Pin : AllPins)
-		{
-			FString NodeName, PinPath;
-			URigVMPin::SplitPinPathAtStart(Pin->GetPinPath(), NodeName, PinPath);
-
-			if (RerouteNode)
-			{
-				FString ValuePinName, SubPinPath;
-				if (URigVMPin::SplitPinPathAtStart(PinPath, ValuePinName, SubPinPath))
-				{
-					FString RedirectedSubPinPath;
-					if (ShouldRedirectPin(OwningStruct, SubPinPath, RedirectedSubPinPath))
-					{
-						FString RedirectedPinPath = URigVMPin::JoinPinPath(ValuePinName, RedirectedSubPinPath);
-						RedirectedPinPaths.Add(PinPath, RedirectedPinPath);
-					}
-				}
-			}
-			else
-			{
-				FString RedirectedPinPath;
-				if (ShouldRedirectPin(OwningStruct, PinPath, RedirectedPinPath))
-				{
-					RedirectedPinPaths.Add(PinPath, RedirectedPinPath);
-				}
-			}
-		}
-	}
-
 	// step 2/3: clear pins on the node and repopulate the node with new pins
 	if (StructNode != nullptr)
 	{
-		TArray<URigVMPin*> Pins = InNode->Pins;
+		TArray<URigVMPin*> Pins = InNode->GetPins();
 		for (URigVMPin* Pin : Pins)
 		{
 			RemovePin(Pin, false);
@@ -5010,7 +5640,7 @@ void URigVMController::RepopulatePinsOnNode(URigVMNode* InNode)
 	}
 	else if (RerouteNode != nullptr)
 	{
-		if (RerouteNode->Pins.Num() == 0)
+		if (RerouteNode->GetPins().Num() == 0)
 		{
 			return;
 		}
@@ -5046,78 +5676,73 @@ void URigVMController::RepopulatePinsOnNode(URigVMNode* InNode)
 			AddPinsForStruct(ScriptStruct, RerouteNode, ValuePin, ValuePin->Direction, ExportedDefaultValue, false);
 		}
 	}
-
-	// step 3/3: restore states for the pins
-	for (TPair<FString, TArray<URigVMInjectionInfo*>> InjectionPair : InjectionInfos)
+	else if (EntryNode || ReturnNode)
 	{
-		if (URigVMPin* Pin = InNode->FindPin(InjectionPair.Key))
+		if (URigVMLibraryNode* LibraryNode = Cast<URigVMLibraryNode>(InNode->GetGraph()->GetOuter()))
 		{
-			for (URigVMInjectionInfo* InjectionInfo : InjectionPair.Value)
+			bool bIsEntryNode = EntryNode != nullptr;
+
+			TArray<URigVMPin*> Pins = InNode->GetPins();
+			for (URigVMPin* Pin : Pins)
 			{
-				InjectionInfo->Rename(nullptr, Pin);
-				InjectionInfo->InputPin = InjectionInfo->StructNode->FindPin(InjectionInfo->InputPin->GetName());
-				InjectionInfo->OutputPin = InjectionInfo->StructNode->FindPin(InjectionInfo->OutputPin->GetName());
-				Pin->InjectionInfos.Add(InjectionInfo);
+				RemovePin(Pin, false);
 			}
-		}
-		else
-		{
-			for (URigVMInjectionInfo* InjectionInfo : InjectionPair.Value)
+			InNode->Pins.Reset();
+			Pins.Reset();
+
+			for (URigVMPin* LibraryPin : LibraryNode->GetPins())
 			{
-				InjectionInfo->StructNode->Rename(nullptr, InNode->GetGraph());
-				DestroyObject(InjectionInfo);
-			}
-		}
-	}
+				if (LibraryPin->GetDirection() == ERigVMPinDirection::Input && !bIsEntryNode)
+				{
+					continue;
+				}
+				else if (LibraryPin->GetDirection() == ERigVMPinDirection::Output && bIsEntryNode)
+				{
+					continue;
+				}
 
-	for (TPair<FString, FString> DefaultValuePair : DefaultValues)
-	{
-		FString PinPath = DefaultValuePair.Key;
-		if (RedirectedPinPaths.Contains(PinPath))
-		{
-			PinPath = RedirectedPinPaths.FindChecked(PinPath);
-		}
+				URigVMPin* ExposedPin = NewObject<URigVMPin>(InNode, LibraryPin->GetFName());
+				ConfigurePinFromPin(ExposedPin, LibraryPin);
 
-		if (URigVMPin* Pin = InNode->FindPin(PinPath))
-		{
-			SetPinDefaultValue(Pin, DefaultValuePair.Value, true, false, false);
-		}
-	}
+				if (ExposedPin->GetDirection() == ERigVMPinDirection::Input)
+				{
+					ExposedPin->Direction = ERigVMPinDirection::Output;
+				}
+				else if (ExposedPin->GetDirection() == ERigVMPinDirection::Output)
+				{
+					ExposedPin->Direction = ERigVMPinDirection::Input;
+				}
+				else if (ExposedPin->Direction != ERigVMPinDirection::IO)
+				{
+					checkNoEntry();
+				}
 
-	for (TPair<FString, bool> ExpansionStatePair : ExpansionStates)
-	{
-		FString PinPath = ExpansionStatePair.Key;
-		if (RedirectedPinPaths.Contains(PinPath))
-		{
-			PinPath = RedirectedPinPaths.FindChecked(PinPath);
-		}
+				InNode->Pins.Add(ExposedPin);
 
-		if (URigVMPin* Pin = InNode->FindPin(PinPath))
+				if (ExposedPin->IsStruct())
+				{
+					AddPinsForStruct(ExposedPin->GetScriptStruct(), InNode, ExposedPin, ExposedPin->GetDirection(), FString(), false);
+				}
+			}
+		}
+		else
 		{
-			SetPinExpansion(Pin, ExpansionStatePair.Value, false);
+			// in the future we'll likely have function libraries as outers here
+			checkNoEntry();
 		}
 	}
-
-	for (TPair<FString, FString> BoundVariablePair: BoundVariables)
+	else
 	{
-		FString PinPath = BoundVariablePair.Key;
-		if (RedirectedPinPaths.Contains(PinPath))
-		{
-			PinPath = RedirectedPinPaths.FindChecked(PinPath);
-		}
-
-		if (URigVMPin* Pin = InNode->FindPin(PinPath))
-		{
-			BindPinToVariable(Pin, BoundVariablePair.Value, false);
-		}
+		return;
 	}
 
+	ApplyPinStates(InNode, PinStates, RedirectedPinPaths);
+
 	if (URigVMInjectionInfo* InjectionInfo = InNode->GetInjectionInfo())
 	{
 		InjectionInfo->InputPin = InNode->FindPin(InjectionInputPinName.ToString());
 		InjectionInfo->OutputPin = InNode->FindPin(InjectionOutputPinName.ToString());
 	}
-
 }
 
 #endif
@@ -5179,6 +5804,133 @@ FLinearColor URigVMController::GetColorFromMetadata(const FString& InMetadata)
 	return Color;
 }
 
+TMap<FString, FString> URigVMController::GetRedirectedPinPaths(URigVMNode* InNode) const
+{
+	TMap<FString, FString> RedirectedPinPaths;
+	URigVMStructNode* StructNode = Cast<URigVMStructNode>(InNode);
+	URigVMRerouteNode* RerouteNode = Cast<URigVMRerouteNode>(InNode);
+
+	UScriptStruct* OwningStruct = nullptr;
+	if (StructNode)
+	{
+		OwningStruct = StructNode->GetScriptStruct();
+	}
+	else if (RerouteNode)
+	{
+		URigVMPin* ValuePin = RerouteNode->Pins[0];
+		if (ValuePin->IsStruct())
+		{
+			OwningStruct = ValuePin->GetScriptStruct();
+		}
+	}
+
+	if (OwningStruct)
+	{
+		TArray<URigVMPin*> AllPins = InNode->GetAllPinsRecursively();
+		for (URigVMPin* Pin : AllPins)
+		{
+			FString NodeName, PinPath;
+			URigVMPin::SplitPinPathAtStart(Pin->GetPinPath(), NodeName, PinPath);
+
+			if (RerouteNode)
+			{
+				FString ValuePinName, SubPinPath;
+				if (URigVMPin::SplitPinPathAtStart(PinPath, ValuePinName, SubPinPath))
+				{
+					FString RedirectedSubPinPath;
+					if (ShouldRedirectPin(OwningStruct, SubPinPath, RedirectedSubPinPath))
+					{
+						FString RedirectedPinPath = URigVMPin::JoinPinPath(ValuePinName, RedirectedSubPinPath);
+						RedirectedPinPaths.Add(PinPath, RedirectedPinPath);
+					}
+				}
+			}
+			else
+			{
+				FString RedirectedPinPath;
+				if (ShouldRedirectPin(OwningStruct, PinPath, RedirectedPinPath))
+				{
+					RedirectedPinPaths.Add(PinPath, RedirectedPinPath);
+				}
+			}
+		}
+	};
+	return RedirectedPinPaths;
+}
+
+URigVMController::FPinState URigVMController::GetPinState(URigVMPin* InPin) const
+{
+	FPinState State;
+	State.DefaultValue = InPin->GetDefaultValue();
+	State.BoundVariable = InPin->GetBoundVariablePath();
+	State.bIsExpanded = InPin->IsExpanded();
+	State.InjectionInfos = InPin->GetInjectedNodes();
+	return State;
+}
+
+TMap<FString, URigVMController::FPinState> URigVMController::GetPinStates(URigVMNode* InNode) const
+{
+	TMap<FString, FPinState> PinStates;
+
+	TArray<URigVMPin*> AllPins = InNode->GetAllPinsRecursively();
+	for (URigVMPin* Pin : AllPins)
+	{
+		FString PinPath, NodeName;
+		URigVMPin::SplitPinPathAtStart(Pin->GetPinPath(), NodeName, PinPath);
+
+		FPinState State = GetPinState(Pin);
+		PinStates.Add(PinPath, State);
+	}
+
+	return PinStates;
+}
+
+void URigVMController::ApplyPinState(URigVMPin* InPin, const FPinState& InPinState)
+{
+	for (URigVMInjectionInfo* InjectionInfo : InPinState.InjectionInfos)
+	{
+		InjectionInfo->Rename(nullptr, InPin);
+		InjectionInfo->InputPin = InjectionInfo->StructNode->FindPin(InjectionInfo->InputPin->GetName());
+		InjectionInfo->OutputPin = InjectionInfo->StructNode->FindPin(InjectionInfo->OutputPin->GetName());
+		InPin->InjectionInfos.Add(InjectionInfo);
+	}
+
+	if (!InPinState.DefaultValue.IsEmpty())
+	{
+		SetPinDefaultValue(InPin, InPinState.DefaultValue, true, false, false);
+	}
+
+	SetPinExpansion(InPin, InPinState.bIsExpanded, false);
+	BindPinToVariable(InPin, InPinState.BoundVariable, false);
+}
+
+void URigVMController::ApplyPinStates(URigVMNode* InNode, const TMap<FString, URigVMController::FPinState>& InPinStates, const TMap<FString, FString>& InRedirectedPinPaths)
+{
+	for (const TPair<FString, FPinState>& PinStatePair : InPinStates)
+	{
+		FString PinPath = PinStatePair.Key;
+		const FPinState& PinState = PinStatePair.Value;
+
+		if (InRedirectedPinPaths.Contains(PinPath))
+		{
+			PinPath = InRedirectedPinPaths.FindChecked(PinPath);
+		}
+
+		if (URigVMPin* Pin = InNode->FindPin(PinPath))
+		{
+			ApplyPinState(Pin, PinState);
+		}
+		else
+		{
+			for (URigVMInjectionInfo* InjectionInfo : PinState.InjectionInfos)
+			{
+				InjectionInfo->StructNode->Rename(nullptr, InNode->GetGraph());
+				DestroyObject(InjectionInfo);
+			}
+		}
+	}
+}
+
 void URigVMController::ReportWarning(const FString& InMessage)
 {
 	if(!bReportWarningsAndErrors)
@@ -5272,7 +6024,7 @@ void URigVMController::PotentiallyResolvePrototypeNode(URigVMPrototypeNode* InNo
 	NodesVisited.Add(InNode);
 
 	// propagate types first
-	for (URigVMPin* Pin : InNode->Pins)
+	for (URigVMPin* Pin : InNode->GetPins())
 	{
 		if (Pin->CPPType.IsEmpty())
 		{
@@ -5300,7 +6052,7 @@ void URigVMController::PotentiallyResolvePrototypeNode(URigVMPrototypeNode* InNo
 		TMap<FString, FString> DefaultValues;
 		TArray<TPair<FString, FString>> LinkPaths;
 
-		for (URigVMPin* Pin : InNode->Pins)
+		for (URigVMPin* Pin : InNode->GetPins())
 		{
 			FString DefaultValue = Pin->GetDefaultValue();
 			if (!DefaultValue.IsEmpty())
@@ -5343,7 +6095,7 @@ void URigVMController::PotentiallyResolvePrototypeNode(URigVMPrototypeNode* InNo
 	else
 	{
 		// update all of the pins that might have changed now as well!
-		for (URigVMPin* Pin : InNode->Pins)
+		for (URigVMPin* Pin : InNode->GetPins())
 		{
 			if (Pin->CPPType.IsEmpty())
 			{
diff --git a/Engine/Source/Developer/RigVMDeveloper/Private/RigVMModel/RigVMControllerActions.cpp b/Engine/Source/Developer/RigVMDeveloper/Private/RigVMModel/RigVMControllerActions.cpp
--- a/Engine/Source/Developer/RigVMDeveloper/Private/RigVMModel/RigVMControllerActions.cpp
+++ b/Engine/Source/Developer/RigVMDeveloper/Private/RigVMModel/RigVMControllerActions.cpp
@@ -646,7 +646,7 @@ bool FRigVMAddInjectedNodeAction::Redo(URigVMController* InController)
 	return false;
 }
 
-FRigVMRemoveNodeAction::FRigVMRemoveNodeAction(URigVMNode* InNode)
+FRigVMRemoveNodeAction::FRigVMRemoveNodeAction(URigVMNode* InNode, URigVMController* InController)
 {
 	FRigVMInverseAction InverseAction;
 
@@ -714,6 +714,10 @@ FRigVMRemoveNodeAction::FRigVMRemoveNodeAction(URigVMNode* InNode)
 	{
 		InverseAction.AddAction(FRigVMAddEnumNodeAction(EnumNode));
 	}
+	else if (URigVMLibraryNode* LibraryNode = Cast<URigVMLibraryNode>(InNode))
+	{
+		InverseAction.AddAction(FRigVMImportNodeFromTextAction(LibraryNode, InController));
+	}
 	else
 	{
 		ensure(false);
@@ -1278,3 +1282,121 @@ bool FRigVMSetPinBoundVariableAction::Redo(URigVMController* InController)
 	InController->BindPinToVariable(PinPath, NewBoundVariablePath, false);
 	return FRigVMBaseAction::Redo(InController);
 }
+
+FRigVMImportNodeFromTextAction::FRigVMImportNodeFromTextAction()
+	: Position(FVector2D::ZeroVector)
+	, NodePath()
+{
+}
+
+FRigVMImportNodeFromTextAction::FRigVMImportNodeFromTextAction(URigVMNode* InNode, URigVMController* InController)
+	: Position(InNode->GetPosition())
+	, NodePath(InNode->GetNodePath())
+	, ExportedText()
+{
+	TArray<FName> NodeNamesToExport;
+	NodeNamesToExport.Add(InNode->GetFName());
+	ExportedText = InController->ExportNodesToText(NodeNamesToExport);
+}
+
+bool FRigVMImportNodeFromTextAction::Undo(URigVMController* InController)
+{
+	if (!FRigVMBaseAction::Undo(InController))
+	{
+		return false;
+	}
+	return InController->RemoveNodeByName(*NodePath, false);
+}
+
+bool FRigVMImportNodeFromTextAction::Redo(URigVMController* InController)
+{
+#if WITH_EDITOR
+	TArray<FName> NodeNames = InController->ImportNodesFromText(ExportedText, false);
+	if (NodeNames.Num() == 1)
+	{
+		return FRigVMBaseAction::Redo(InController);
+	}
+#endif
+	return false;
+}
+
+FRigVMCollapseNodesAction::FRigVMCollapseNodesAction()
+	: LibraryNodePath()
+{
+}
+
+FRigVMCollapseNodesAction::FRigVMCollapseNodesAction(const TArray<URigVMNode*>& InNodes, const FString& InNodePath)
+	: LibraryNodePath(InNodePath)
+{
+	for (URigVMNode* InNode : InNodes)
+	{
+		CollapsedNodesPaths.Add(InNode->GetName());
+	}
+}
+
+bool FRigVMCollapseNodesAction::Undo(URigVMController* InController)
+{
+	if (!FRigVMBaseAction::Undo(InController))
+	{
+		return false;
+	}
+	TArray<URigVMNode*> ExpandedNodes = InController->ExpandLibraryNode(*LibraryNodePath, false);
+	return ExpandedNodes.Num() == CollapsedNodesPaths.Num();
+}
+
+bool FRigVMCollapseNodesAction::Redo(URigVMController* InController)
+{
+#if WITH_EDITOR
+	TArray<FName> NodeNames;
+	for (const FString& NodePath : CollapsedNodesPaths)
+	{
+		NodeNames.Add(*NodePath);
+	}
+
+	URigVMLibraryNode* LibraryNode = InController->CollapseNodes(NodeNames, LibraryNodePath, false);
+	if (LibraryNode)
+	{
+		return FRigVMBaseAction::Redo(InController);
+	}
+#endif
+	return false;
+}
+
+FRigVMExpandNodeAction::FRigVMExpandNodeAction()
+	: LibraryNodePath()
+{
+}
+
+FRigVMExpandNodeAction::FRigVMExpandNodeAction(URigVMLibraryNode* InLibraryNode)
+	: LibraryNodePath(InLibraryNode->GetName())
+{
+}
+
+bool FRigVMExpandNodeAction::Undo(URigVMController* InController)
+{
+	if (!FRigVMBaseAction::Undo(InController))
+	{
+		return false;
+	}
+
+	TArray<FName> NodeNames;
+	for (const FString& NodePath : ExpandedNodePaths)
+	{
+		NodeNames.Add(*NodePath);
+	}
+
+	URigVMLibraryNode* LibraryNode = InController->CollapseNodes(NodeNames, *LibraryNodePath, false);
+	return LibraryNode != nullptr;
+}
+
+bool FRigVMExpandNodeAction::Redo(URigVMController* InController)
+{
+#if WITH_EDITOR
+	TArray<URigVMNode*> ExpandedNodes = InController->ExpandLibraryNode(*LibraryNodePath, false);
+	if (ExpandedNodes.Num() == ExpandedNodePaths.Num())
+	{
+		return FRigVMBaseAction::Redo(InController);
+	}
+#endif
+	return false;
+}
\ No newline at end of file
diff --git a/Engine/Source/Developer/RigVMDeveloper/Private/RigVMModel/RigVMNode.cpp b/Engine/Source/Developer/RigVMDeveloper/Private/RigVMModel/RigVMNode.cpp
--- a/Engine/Source/Developer/RigVMDeveloper/Private/RigVMModel/RigVMNode.cpp
+++ b/Engine/Source/Developer/RigVMDeveloper/Private/RigVMModel/RigVMNode.cpp
@@ -69,7 +69,7 @@ TArray<URigVMPin*> URigVMNode::GetAllPinsRecursively() const
 	};
 
 	TArray<URigVMPin*> Result;
-	for (URigVMPin* Pin : Pins)
+	for (URigVMPin* Pin : GetPins())
 	{
 		Local::VisitPinRecursively(Pin, Result);
 	}
@@ -84,7 +84,7 @@ URigVMPin* URigVMNode::FindPin(const FString& InPinPath) const
 		Left = InPinPath;
 	}
 
-	for (URigVMPin* Pin : Pins)
+	for (URigVMPin* Pin : GetPins())
 	{
 		if (Pin->GetName() == Left)
 		{
@@ -177,7 +177,7 @@ bool URigVMNode::IsPure() const
 		return false;
 	}
 
-	for (URigVMPin* Pin : Pins)
+	for (URigVMPin* Pin : GetPins())
 	{
 		if(Pin->GetDirection() == ERigVMPinDirection::Hidden)
 		{
@@ -245,7 +245,7 @@ bool URigVMNode::HasOutputPin(bool bIncludeIO) const
 
 bool URigVMNode::HasPinOfDirection(ERigVMPinDirection InDirection) const
 {
-	for (URigVMPin* Pin : Pins)
+	for (URigVMPin* Pin : GetPins())
 	{
 		if (Pin->GetDirection() == InDirection)
 		{
@@ -269,7 +269,7 @@ bool URigVMNode::IsLinkedTo(URigVMNode* InNode) const
 	{
 		return false;
 	}
-	for (URigVMPin* Pin : Pins)
+	for (URigVMPin* Pin : GetPins())
 	{
 		if (IsLinkedToRecursive(Pin, InNode))
 		{
@@ -305,10 +305,34 @@ bool URigVMNode::IsLinkedToRecursive(URigVMPin* InPin, URigVMNode* InNode) const
 	return false;
 }
 
+TArray<URigVMLink*> URigVMNode::GetLinks() const
+{
+	TArray<URigVMLink*> Links;
+
+	struct Local
+	{
+		static void Traverse(URigVMPin* InPin, TArray<URigVMLink*>& Links)
+		{
+			Links.Append(InPin->GetLinks());
+			for (URigVMPin* SubPin : InPin->GetSubPins())
+			{
+				Local::Traverse(SubPin, Links);
+			}
+		}
+	};
+
+	for (URigVMPin* Pin : GetPins())
+	{
+		Local::Traverse(Pin, Links);
+	}
+
+	return Links;
+}
+
 TArray<URigVMNode*> URigVMNode::GetLinkedSourceNodes() const
 {
 	TArray<URigVMNode*> Nodes;
-	for (URigVMPin* Pin : Pins)
+	for (URigVMPin* Pin : GetPins())
 	{
 		GetLinkedNodesRecursive(Pin, true, Nodes);
 	}
@@ -318,7 +342,7 @@ TArray<URigVMNode*> URigVMNode::GetLinkedSourceNodes() const
 TArray<URigVMNode*> URigVMNode::GetLinkedTargetNodes() const
 {
 	TArray<URigVMNode*> Nodes;
-	for (URigVMPin* Pin : Pins)
+	for (URigVMPin* Pin : GetPins())
 	{
 		GetLinkedNodesRecursive(Pin, false, Nodes);
 	}
@@ -350,7 +374,7 @@ int32 URigVMNode::GetNumSlices(const FRigVMUserDataArray& InUserData)
 
 int32 URigVMNode::GetNumSlicesForContext(const FName& InContextName, const FRigVMUserDataArray& InUserData)
 {
-	for (URigVMPin* RootPin : Pins)
+	for (URigVMPin* RootPin : GetPins())
 	{
 		if (RootPin->GetFName() == InContextName)
 		{
@@ -364,7 +388,7 @@ int32 URigVMNode::GetNumSlicesForContext(const FName& InContextName, const FRigV
 	{
 		TGuardValue<int32> ReentrantGuard(GetSliceContextBracket, GetSliceContextBracket + 1);
 
-		for (URigVMPin* Pin : Pins)
+		for (URigVMPin* Pin : GetPins())
 		{
 			TArray<URigVMPin*> SourcePins = Pin->GetLinkedSourcePins(true /* recursive */);
 			if (SourcePins.Num() > 0)
diff --git a/Engine/Source/Developer/RigVMDeveloper/Public/RigVMModel/Nodes/RigVMCollapseNode.h b/Engine/Source/Developer/RigVMDeveloper/Public/RigVMModel/Nodes/RigVMCollapseNode.h
new file mode 100644
--- /dev/null
+++ b/Engine/Source/Developer/RigVMDeveloper/Public/RigVMModel/Nodes/RigVMCollapseNode.h
@@ -0,0 +1,34 @@
+// Copyright Epic Games, Inc. All Rights Reserved.
+
+#pragma once
+
+#include "RigVMModel/Nodes/RigVMLibraryNode.h"
+#include "RigVMModel/RigVMGraph.h"
+#include "RigVMCollapseNode.generated.h"
+
+/**
+ * The Collapse Node is a library node which stores the 
+ * function and its nodes directly within the node itself.
+ */
+UCLASS(BlueprintType)
+class RIGVMDEVELOPER_API URigVMCollapseNode : public URigVMLibraryNode
+{
+	GENERATED_BODY()
+
+public:
+
+	// Override node functions
+	virtual FString GetNodeTitle() const override;
+	virtual FText GetToolTipText() const override;
+
+	// Library node interface
+	virtual URigVMGraph* GetContainedGraph() const override { return ContainedGraph; }
+
+private:
+
+	UPROPERTY()
+	URigVMGraph* ContainedGraph;
+
+	friend class URigVMController;
+};
+
diff --git a/Engine/Source/Developer/RigVMDeveloper/Public/RigVMModel/Nodes/RigVMFunctionEntryNode.h b/Engine/Source/Developer/RigVMDeveloper/Public/RigVMModel/Nodes/RigVMFunctionEntryNode.h
new file mode 100644
--- /dev/null
+++ b/Engine/Source/Developer/RigVMDeveloper/Public/RigVMModel/Nodes/RigVMFunctionEntryNode.h
@@ -0,0 +1,26 @@
+// Copyright Epic Games, Inc. All Rights Reserved.
+
+#pragma once
+
+#include "RigVMModel/RigVMNode.h"
+#include "RigVMFunctionEntryNode.generated.h"
+
+/**
+ * The Function Entry node is used to provide access to the 
+ * input pins of the library node for links within.
+ */
+UCLASS(BlueprintType)
+class RIGVMDEVELOPER_API URigVMFunctionEntryNode : public URigVMNode
+{
+	GENERATED_BODY()
+
+public:
+
+	// Override node functions
+	virtual bool IsDefinedAsVarying() const override;
+
+private:
+
+	friend class URigVMController;
+};
+
diff --git a/Engine/Source/Developer/RigVMDeveloper/Public/RigVMModel/Nodes/RigVMFunctionReturnNode.h b/Engine/Source/Developer/RigVMDeveloper/Public/RigVMModel/Nodes/RigVMFunctionReturnNode.h
new file mode 100644
--- /dev/null
+++ b/Engine/Source/Developer/RigVMDeveloper/Public/RigVMModel/Nodes/RigVMFunctionReturnNode.h
@@ -0,0 +1,26 @@
+// Copyright Epic Games, Inc. All Rights Reserved.
+
+#pragma once
+
+#include "RigVMModel/RigVMNode.h"
+#include "RigVMFunctionReturnNode.generated.h"
+
+/**
+ * The Function Return node is used to provide access to the 
+ * output pins of the library node for links within.
+ */
+UCLASS(BlueprintType)
+class RIGVMDEVELOPER_API URigVMFunctionReturnNode : public URigVMNode
+{
+	GENERATED_BODY()
+
+public:
+
+	// Override node functions
+	virtual bool IsDefinedAsVarying() const override;
+
+private:
+
+	friend class URigVMController;
+};
+
diff --git a/Engine/Source/Developer/RigVMDeveloper/Public/RigVMModel/Nodes/RigVMLibraryNode.h b/Engine/Source/Developer/RigVMDeveloper/Public/RigVMModel/Nodes/RigVMLibraryNode.h
new file mode 100644
--- /dev/null
+++ b/Engine/Source/Developer/RigVMDeveloper/Public/RigVMModel/Nodes/RigVMLibraryNode.h
@@ -0,0 +1,43 @@
+// Copyright Epic Games, Inc. All Rights Reserved.
+
+#pragma once
+
+#include "RigVMModel/RigVMNode.h"
+#include "RigVMLibraryNode.generated.h"
+
+class URigVMGraph;
+class URigVMFunctionEntryNode;
+class URigVMFunctionReturnNode;
+
+/**
+ * The Library Node represents a function invocation of a
+ * function specified somewhere else. The function can be 
+ * expressed as a sub-graph (RigVMGroupNode) or as a 
+ * referenced function (RigVMFunctionNode).
+ */
+UCLASS(BlueprintType)
+class RIGVMDEVELOPER_API URigVMLibraryNode : public URigVMNode
+{
+	GENERATED_BODY()
+
+public:
+
+	// Override node functions
+	virtual bool IsDefinedAsConstant() const override;
+	virtual bool IsDefinedAsVarying() const override;
+
+	// Library node interface
+	virtual URigVMGraph* GetContainedGraph() const { return nullptr; }
+	virtual TArray<URigVMNode*> GetContainedNodes() const;
+	virtual TArray<URigVMLink*> GetContainedLinks() const;
+	virtual URigVMFunctionEntryNode* GetEntryNode() const;
+	virtual URigVMFunctionReturnNode* GetReturnNode() const;
+
+private:
+
+	UPROPERTY()
+	bool bDefinedAsVarying = false;
+
+	friend class URigVMController;
+};
+
diff --git a/Engine/Source/Developer/RigVMDeveloper/Public/RigVMModel/RigVMController.h b/Engine/Source/Developer/RigVMDeveloper/Public/RigVMModel/RigVMController.h
--- a/Engine/Source/Developer/RigVMDeveloper/Public/RigVMModel/RigVMController.h
+++ b/Engine/Source/Developer/RigVMDeveloper/Public/RigVMModel/RigVMController.h
@@ -13,6 +13,7 @@
 #include "RigVMModel/Nodes/RigVMSelectNode.h"
 #include "RigVMModel/Nodes/RigVMPrototypeNode.h"
 #include "RigVMModel/Nodes/RigVMEnumNode.h"
+#include "RigVMModel/Nodes/RigVMCollapseNode.h"
 #include "RigVMController.generated.h"
 
 class URigVMActionStack;
@@ -262,6 +263,14 @@ public:
 	// Returns a unique name
 	static FName GetUniqueName(const FName& InName, TFunction<bool(const FName&)> IsNameAvailableFunction);
 
+	// Turns a series of nodes into a Collapse node
+	UFUNCTION(BlueprintCallable, Category = RigVMController)
+	URigVMCollapseNode* CollapseNodes(const TArray<FName>& InNodeNames, const FString& InCollapseNodeName = TEXT(""), bool bSetupUndoRedo = true);
+
+	// Turns a library node into its contained nodes
+	UFUNCTION(BlueprintCallable, Category = RigVMController)
+	TArray<URigVMNode*> ExpandLibraryNode(const FName& InNodeName, bool bSetupUndoRedo = true);
+
 #endif
 
 	// Removes a node from the graph
@@ -454,8 +463,8 @@ public:
 	// A delegate to retrieve the list of external variables
 	FRigVMController_GetExternalVariablesDelegate GetExternalVariablesDelegate;
 
-	int32 DetachLinksFromPinObjects();
-	int32 ReattachLinksToPinObjects(bool bFollowCoreRedirectors = false);
+	int32 DetachLinksFromPinObjects(const TArray<URigVMLink*>* InLinks = nullptr);
+	int32 ReattachLinksToPinObjects(bool bFollowCoreRedirectors = false, const TArray<URigVMLink*>* InLinks = nullptr);
 	void AddPinRedirector(bool bInput, bool bOutput, const FString& OldPinPath, const FString& NewPinPath);
 
 	// Removes nodes which went stale.
@@ -466,7 +475,7 @@ public:
 	bool ShouldRedirectPin(UScriptStruct* InOwningStruct, const FString& InOldRelativePinPath, FString& InOutNewRelativePinPath) const;
 	bool ShouldRedirectPin(const FString& InOldPinPath, FString& InOutNewPinPath) const;
 
-	void RepopulatePinsOnNode(URigVMNode* InNode);
+	void RepopulatePinsOnNode(URigVMNode* InNode, bool bFollowCoreRedirectors = true);
 #endif
 
 	FRigVMStructNodeCreatedContext& GetStructNodeCreatedContext() { return StructNodeCreatedContext; }
@@ -516,6 +525,23 @@ private:
 	bool SetParameterName(URigVMParameterNode* InParameterNode, const FName& InParameterName, bool bSetupUndoRedo);
 	static void ForEveryPinRecursively(URigVMPin* InPin, TFunction<void(URigVMPin*)> OnEachPinFunction);
 	static void ForEveryPinRecursively(URigVMNode* InNode, TFunction<void(URigVMPin*)> OnEachPinFunction);
+	URigVMCollapseNode* CollapseNodes(const TArray<URigVMNode*>& InNodes, const FString& InCollapseNodeName, bool bSetupUndoRedo);
+	TArray<URigVMNode*> ExpandLibraryNode(URigVMLibraryNode* InNode, bool bSetupUndoRedo);
+	void RefreshFunctionPins(URigVMNode* InNode);
+
+	struct FPinState
+	{
+		FString DefaultValue;
+		FString BoundVariable;
+		bool bIsExpanded;
+		TArray<URigVMInjectionInfo*> InjectionInfos;
+	};
+
+	TMap<FString, FString> GetRedirectedPinPaths(URigVMNode* InNode) const;
+	FPinState GetPinState(URigVMPin* InPin) const;
+	TMap<FString, FPinState> GetPinStates(URigVMNode* InNode) const;
+	void ApplyPinState(URigVMPin* InPin, const FPinState& InPinState);
+	void ApplyPinStates(URigVMNode* InNode, const TMap<FString, FPinState>& InPinStates, const TMap<FString, FString>& InRedirectedPinPaths = TMap<FString, FString>());
 
 	void ReportWarning(const FString& InMessage);
 	void ReportError(const FString& InMessage);
diff --git a/Engine/Source/Developer/RigVMDeveloper/Public/RigVMModel/RigVMControllerActions.h b/Engine/Source/Developer/RigVMDeveloper/Public/RigVMModel/RigVMControllerActions.h
--- a/Engine/Source/Developer/RigVMDeveloper/Public/RigVMModel/RigVMControllerActions.h
+++ b/Engine/Source/Developer/RigVMDeveloper/Public/RigVMModel/RigVMControllerActions.h
@@ -657,7 +657,7 @@ struct FRigVMRemoveNodeAction : public FRigVMBaseAction
 public:
 
 	FRigVMRemoveNodeAction() {}
-	FRigVMRemoveNodeAction(URigVMNode* InNode);
+	FRigVMRemoveNodeAction(URigVMNode* InNode, URigVMController* InController);
 	virtual ~FRigVMRemoveNodeAction() {};
 	virtual bool Undo(URigVMController* InController) override;
 	virtual bool Redo(URigVMController* InController) override;
@@ -1127,3 +1127,75 @@ public:
 	UPROPERTY()
 	FString NewBoundVariablePath;
 };
+
+/**
+ * An action to add a node from a text buffer
+ */
+USTRUCT()
+struct FRigVMImportNodeFromTextAction : public FRigVMBaseAction
+{
+	GENERATED_BODY()
+
+public:
+
+	FRigVMImportNodeFromTextAction();
+	FRigVMImportNodeFromTextAction(URigVMNode* InNode, URigVMController* InController);
+	virtual ~FRigVMImportNodeFromTextAction() {};
+	virtual bool Undo(URigVMController* InController) override;
+	virtual bool Redo(URigVMController* InController) override;
+
+	UPROPERTY()
+	FVector2D Position;
+
+	UPROPERTY()
+	FString NodePath;
+
+	UPROPERTY()
+	FString ExportedText;
+};
+
+/**
+ * An action to collapse a selection of nodes
+ */
+USTRUCT()
+struct FRigVMCollapseNodesAction : public FRigVMBaseAction
+{
+	GENERATED_BODY()
+
+public:
+
+	FRigVMCollapseNodesAction();
+	FRigVMCollapseNodesAction(const TArray<URigVMNode*>& InNodes, const FString& InNodePath);
+	virtual ~FRigVMCollapseNodesAction() {};
+	virtual bool Undo(URigVMController* InController) override;
+	virtual bool Redo(URigVMController* InController) override;
+
+	UPROPERTY()
+	FString LibraryNodePath;
+
+	UPROPERTY()
+	TArray<FString> CollapsedNodesPaths;
+};
+
+/**
+ * An action to expand a library node into its content
+ */
+USTRUCT()
+struct FRigVMExpandNodeAction : public FRigVMBaseAction
+{
+	GENERATED_BODY()
+
+public:
+
+	FRigVMExpandNodeAction();
+	FRigVMExpandNodeAction(URigVMLibraryNode* InLibraryNode);
+	virtual ~FRigVMExpandNodeAction() {};
+	virtual bool Undo(URigVMController* InController) override;
+	virtual bool Redo(URigVMController* InController) override;
+
+	UPROPERTY()
+	FString LibraryNodePath;
+
+	UPROPERTY()
+	TArray<FString> ExpandedNodePaths;
+};
\ No newline at end of file
diff --git a/Engine/Source/Developer/RigVMDeveloper/Public/RigVMModel/RigVMNode.h b/Engine/Source/Developer/RigVMDeveloper/Public/RigVMModel/RigVMNode.h
--- a/Engine/Source/Developer/RigVMDeveloper/Public/RigVMModel/RigVMNode.h
+++ b/Engine/Source/Developer/RigVMDeveloper/Public/RigVMModel/RigVMNode.h
@@ -49,7 +49,7 @@ public:
 
 	// Returns all of the top-level Pins of this Node.
 	UFUNCTION(BlueprintCallable, Category = RigVMNode)
-	const TArray<URigVMPin*>& GetPins() const;
+	virtual const TArray<URigVMPin*>& GetPins() const;
 
 	// Returns all of the Pins of this Node (including SubPins).
 	UFUNCTION(BlueprintCallable, Category = RigVMNode)
@@ -150,6 +150,10 @@ public:
 	UFUNCTION(BlueprintCallable, Category = RigVMNode)
 	bool IsLinkedTo(URigVMNode* InNode) const;
 
+	// Returns all links to any pin on this node
+	UFUNCTION(BlueprintCallable, Category = RigVMNode)
+	TArray<URigVMLink*> GetLinks() const;
+
 	// Returns a list of Nodes connected as sources to
 	// this Node as the target.
 	UFUNCTION(BlueprintCallable, Category = RigVMNode)
