Control Rig: Encapsulation support, next steps

diff --git a/Engine/Source/Developer/RigVMDeveloper/Private/RigVMCompiler/RigVMAST.cpp b/Engine/Source/Developer/RigVMDeveloper/Private/RigVMCompiler/RigVMAST.cpp
--- a/Engine/Source/Developer/RigVMDeveloper/Private/RigVMCompiler/RigVMAST.cpp
+++ b/Engine/Source/Developer/RigVMDeveloper/Private/RigVMCompiler/RigVMAST.cpp
@@ -11,14 +11,14 @@
 #include "RigVMModel/Nodes/RigVMIfNode.h"
 #include "RigVMModel/Nodes/RigVMSelectNode.h"
 #include "RigVMModel/Nodes/RigVMEnumNode.h"
+#include "RigVMModel/Nodes/RigVMFunctionReturnNode.h"
+#include "RigVMModel/Nodes/RigVMFunctionEntryNode.h"
 #include "RigVMModel/RigVMGraph.h"
 #include "RigVMModel/RigVMController.h"
 #include "RigVMCore/RigVMExecuteContext.h"
 #include "Stats/StatsHierarchical.h"
 #include "RigVMDeveloperModule.h"
 
-// #define FRIGVMASTPARSER_DEBUG_TRAVERSAL
-
 FRigVMExprAST::FRigVMExprAST(EType InType, UObject* InSubject)
 	: Name(NAME_None)
 	, Type(InType)
@@ -2484,32 +2484,118 @@ void FRigVMParserAST::Inline(URigVMGraph* InGraph, const TArray<URigVMNode*>& In
 		TMap<URigVMPin*, URigVMPin*> SourcePins;
 		TMap<URigVMPin*, TArray<URigVMPin*>>* TargetLinks;
 		TMap<URigVMPin*, TArray<URigVMPin*>>* SourceLinks;
-#ifdef FRIGVMASTPARSER_DEBUG_TRAVERSAL
-		FString Indentation;
-#endif
+		TArray<URigVMLibraryNode*> LibraryNodeStack;
+
+		class FLibraryNodeGuard
+		{
+		public:
+			FLibraryNodeGuard(URigVMLibraryNode* InNode, TArray<URigVMLibraryNode*>& InStack)
+				: Stack(InStack)
+			{
+				Stack.Push(InNode);
+			}
+
+			~FLibraryNodeGuard()
+			{
+				Stack.Pop();
+			}
+
+		private:
+
+			TArray<URigVMLibraryNode*>& Stack;
+		};
+
+		static bool ShouldRecursePin(URigVMPin* InPin)
+		{
+			URigVMNode* Node = InPin->GetNode();
+			return Node->IsA<URigVMRerouteNode>() ||
+				Node->IsA<URigVMLibraryNode>() ||
+				Node->IsA<URigVMFunctionEntryNode>() ||
+				Node->IsA<URigVMFunctionReturnNode>();
+		}
+
+		static bool IsValidPinForAST(URigVMPin* InPin)
+		{
+			return !ShouldRecursePin(InPin);
+		}
+
+		static bool IsValidLinkForAST(URigVMPin* InSourcePin, URigVMPin* InTargetPin)
+		{
+			return IsValidPinForAST(InSourcePin) && IsValidPinForAST(InTargetPin);
+		}
 
 		static URigVMPin* FindSourcePin(URigVMPin* InPin, LocalPinTraversalInfo& OutTraversalInfo)
 		{
-#ifdef FRIGVMASTPARSER_DEBUG_TRAVERSAL
-			FString& Indent = OutTraversalInfo.Indentation;
-#endif
+			// if this pin is a root on a library
+			if (InPin->GetParentPin() == nullptr)
+			{
+				if (InPin->GetDirection() == ERigVMPinDirection::Output ||
+					InPin->GetDirection() == ERigVMPinDirection::IO)
+				{
+					URigVMNode* Node = InPin->GetNode();
+					if (URigVMLibraryNode* LibraryNode = Cast<URigVMLibraryNode>(Node))
+					{
+						if (!OutTraversalInfo.LibraryNodeStack.Contains(LibraryNode))
+						{
+							if (URigVMFunctionReturnNode* ReturnNode = LibraryNode->GetReturnNode())
+							{
+								if (URigVMPin* ReturnPin = ReturnNode->FindPin(InPin->GetName()))
+								{
+									FLibraryNodeGuard NodeGuard(LibraryNode, OutTraversalInfo.LibraryNodeStack);
+									URigVMPin* SourcePin = FindSourcePin(ReturnPin, OutTraversalInfo);
+									SourcePin = SourcePin == nullptr ? ReturnPin : SourcePin;
+									OutTraversalInfo.SourcePins.FindOrAdd(InPin) = SourcePin;
+									return SourcePin;
+
+								}
+							}
+						}
+					}
+					else if (URigVMFunctionEntryNode* EntryNode = Cast<URigVMFunctionEntryNode>(Node))
+					{
+						for (int32 LibraryNodeIndex = OutTraversalInfo.LibraryNodeStack.Num() - 1; LibraryNodeIndex >= 0; LibraryNodeIndex--)
+						{
+							URigVMLibraryNode* LastLibraryNode = OutTraversalInfo.LibraryNodeStack[LibraryNodeIndex];
+							if(LastLibraryNode->GetEntryNode() == EntryNode)
+							{
+								if (URigVMPin* LibraryPin = LastLibraryNode->FindPin(InPin->GetName()))
+								{
+									URigVMPin* SourcePin = FindSourcePin(LibraryPin, OutTraversalInfo);
+									SourcePin = SourcePin == nullptr ? LibraryPin : SourcePin;
+									OutTraversalInfo.SourcePins.FindOrAdd(InPin) = SourcePin;
+									return SourcePin;
+								}
+							}
+						}
+					}
+				}
+			}
+
 			if (InPin->GetDirection() != ERigVMPinDirection::Input &&
 				InPin->GetDirection() != ERigVMPinDirection::IO)
 			{
 				return nullptr;
 			}
 
-			if (URigVMPin* const* SourcePin = OutTraversalInfo.SourcePins.Find(InPin))
+			bool bIOPinOnLeftOfLibraryNode = false;
+			if (InPin->GetDirection() == ERigVMPinDirection::IO)
 			{
-#ifdef FRIGVMASTPARSER_DEBUG_TRAVERSAL
-				UE_LOG(LogRigVMDeveloper, Display, TEXT("%s= %s (%s)"), *Indent, *InPin->GetPinPath(), *InPin->GetCPPType());
-#endif
-				return *SourcePin;
+				if (URigVMLibraryNode* LibraryNode = Cast<URigVMLibraryNode>(InPin->GetNode()))
+				{
+					bIOPinOnLeftOfLibraryNode = OutTraversalInfo.LibraryNodeStack.Contains(LibraryNode);
+				}
+			}
+
+			if (!bIOPinOnLeftOfLibraryNode)
+			{
+				// note: this map isn't going to work for functions which are referenced.
+				// (since the pin objects are shared between multiple invocation nodes)
+				if (URigVMPin* const* SourcePin = OutTraversalInfo.SourcePins.Find(InPin))
+				{
+					return *SourcePin;
+				}
 			}
 
-#ifdef FRIGVMASTPARSER_DEBUG_TRAVERSAL
-			UE_LOG(LogRigVMDeveloper, Display, TEXT("%s| %s (%s)"), *Indent, *InPin->GetPinPath(), *InPin->GetCPPType());
-#endif
 			TArray<FString> SegmentPath;
 			URigVMPin* SourcePin = nullptr;
 
@@ -2522,18 +2608,12 @@ void FRigVMParserAST::Inline(URigVMGraph* InGraph, const TArray<URigVMNode*>& In
 					SourcePin = SourceLinks[0]->GetSourcePin();
 
 					// only continue the recursion on reroutes
-					if (SourcePin->GetNode()->IsA<URigVMRerouteNode>())
+					if (ShouldRecursePin(SourcePin))
 					{
-#ifdef FRIGVMASTPARSER_DEBUG_TRAVERSAL
-						Indent += TEXT("  ");
-#endif
 						if (URigVMPin* SourceSourcePin = FindSourcePin(SourcePin, OutTraversalInfo))
 						{
 							SourcePin = SourceSourcePin;
 						}
-#ifdef FRIGVMASTPARSER_DEBUG_TRAVERSAL
-						Indent = Indent.LeftChop(2);
-#endif
 					}
 
 					break;
@@ -2548,7 +2628,7 @@ void FRigVMParserAST::Inline(URigVMGraph* InGraph, const TArray<URigVMNode*>& In
 						URigVMPin* ParentSourcePin = *ParentSourcePinPtr;
 						if (ParentSourcePin)
 						{
-							if (!ParentSourcePin->GetNode()->IsA<URigVMRerouteNode>())
+							if (!ShouldRecursePin(ParentSourcePin))
 							{
 								SourcePin = nullptr;
 								break;
@@ -2557,10 +2637,6 @@ void FRigVMParserAST::Inline(URigVMGraph* InGraph, const TArray<URigVMNode*>& In
 					}
 
 					SegmentPath.Push(ChildPin->GetName());
-#ifdef FRIGVMASTPARSER_DEBUG_TRAVERSAL
-					UE_LOG(LogRigVMDeveloper, Display, TEXT("%sPath %s"), *Indent, *URigVMPin::JoinPinPath(SegmentPath));
-					UE_LOG(LogRigVMDeveloper, Display, TEXT("%s> %s (%s)"), *Indent, *ParentPin->GetPinPath(), *ParentPin->GetCPPType());
-#endif
 				}
 				ChildPin = ParentPin;
 			}
@@ -2571,33 +2647,17 @@ void FRigVMParserAST::Inline(URigVMGraph* InGraph, const TArray<URigVMNode*>& In
 				{
 					FString Segment = SegmentPath.Pop();
 
-#ifdef FRIGVMASTPARSER_DEBUG_TRAVERSAL
-					if (!SegmentPath.IsEmpty())
-					{
-						UE_LOG(LogRigVMDeveloper, Display, TEXT("%sPath %s"), *Indent, *URigVMPin::JoinPinPath(SegmentPath));
-					}
-					UE_LOG(LogRigVMDeveloper, Display, TEXT("%sChecking segment %s"), *Indent, *Segment);
-#endif
 					if (URigVMPin* SourceSubPin = SourcePin->FindSubPin(Segment))
 					{
-#ifdef FRIGVMASTPARSER_DEBUG_TRAVERSAL
-						UE_LOG(LogRigVMDeveloper, Display, TEXT("%s< %s (%s)"), *Indent, *SourceSubPin->GetPinPath(), *SourceSubPin->GetCPPType());
-#endif
 						SourcePin = SourceSubPin;
 
 						// only continue the recursion on reroutes
-						if (SourcePin->GetNode()->IsA<URigVMRerouteNode>())
+						if (ShouldRecursePin(SourcePin))
 						{
-#ifdef FRIGVMASTPARSER_DEBUG_TRAVERSAL
-							Indent += TEXT("  ");
-#endif
 							if (URigVMPin* SourceSourceSubPin = FindSourcePin(SourcePin, OutTraversalInfo))
 							{
 								SourcePin = SourceSourceSubPin;
 							}
-#ifdef FRIGVMASTPARSER_DEBUG_TRAVERSAL
-							Indent = Indent.LeftChop(2);
-#endif
 						}
 					}
 					else
@@ -2608,13 +2668,10 @@ void FRigVMParserAST::Inline(URigVMGraph* InGraph, const TArray<URigVMNode*>& In
 				}
 			}
 
-			OutTraversalInfo.SourcePins.FindOrAdd(InPin) = SourcePin;
-#ifdef FRIGVMASTPARSER_DEBUG_TRAVERSAL
-			if (SourcePin)
+			if (!bIOPinOnLeftOfLibraryNode)
 			{
-				UE_LOG(LogRigVMDeveloper, Display, TEXT("%s%s --> %s"), *Indent, *InPin->GetPinPath(), *SourcePin->GetPinPath());
+				OutTraversalInfo.SourcePins.FindOrAdd(InPin) = SourcePin;
 			}
-#endif
 			return SourcePin;
 		}
 
@@ -2627,15 +2684,23 @@ void FRigVMParserAST::Inline(URigVMGraph* InGraph, const TArray<URigVMNode*>& In
 				// If the source pin is on a reroute node, this means that
 				// we only care about the default value - since it is a 
 				// "hanging" reroute without any live input.
+				// same goes for library nodes or return nodes - we'll 
+				// just use the default pin value in that case.
 
-				if (SourcePin->GetNode()->IsA<URigVMRerouteNode>())
+				URigVMNode* SourceNode = SourcePin->GetNode();
+				if (SourceNode->IsA<URigVMRerouteNode>() ||
+					SourceNode->IsA<URigVMLibraryNode>() ||
+					SourceNode->IsA<URigVMFunctionReturnNode>())
 				{
 					OutTraversalInfo.PinDefaultValueOverrides->FindOrAdd(InPin) = SourcePin->GetDefaultValue();
 				}
 				else
 				{
-					OutTraversalInfo.SourceLinks->FindOrAdd(InPin).Add(SourcePin);
-					OutTraversalInfo.TargetLinks->FindOrAdd(SourcePin).Add(InPin);
+					if (IsValidLinkForAST(SourcePin, InPin))
+					{
+						OutTraversalInfo.SourceLinks->FindOrAdd(InPin).Add(SourcePin);
+						OutTraversalInfo.TargetLinks->FindOrAdd(SourcePin).Add(InPin);
+					}
 				}
 			}
 
@@ -2644,18 +2709,45 @@ void FRigVMParserAST::Inline(URigVMGraph* InGraph, const TArray<URigVMNode*>& In
 				VisitPin(SubPin, OutTraversalInfo);
 			}
 		}
+
+		static void VisitNode(URigVMNode* InNode, LocalPinTraversalInfo& OutTraversalInfo)
+		{
+			if (InNode->IsA<URigVMRerouteNode>() ||
+				InNode->IsA<URigVMFunctionEntryNode>() || 
+				InNode->IsA<URigVMFunctionReturnNode>())
+			{
+				return;
+			}
+
+			if (URigVMLibraryNode* LibraryNode = Cast<URigVMLibraryNode>(InNode))
+			{
+				FLibraryNodeGuard NodeGuard(LibraryNode, OutTraversalInfo.LibraryNodeStack);
+
+				TArray<URigVMNode*> ContainedNodes = LibraryNode->GetContainedNodes();
+				for (URigVMNode* ContainedNode : ContainedNodes)
+				{
+					VisitNode(ContainedNode, OutTraversalInfo);
+				}
+			}
+			else
+			{
+				for (URigVMPin* Pin : InNode->GetPins())
+				{
+					LocalPinTraversalInfo::VisitPin(Pin, OutTraversalInfo);
+				}
+			}
+		}
 	};
 
 	Nodes.Reset();
 	SourceLinks.Reset();
 	TargetLinks.Reset();
 
-	// a) find all of the relevant nodes
-	//    once we support encapsulation we'll have to traverse into library nodes
+	// a) find all of the relevant nodes,
+	//    inline and traverse into library nodes
 	Nodes = InNodes;
 
 	// c) flatten links from an entry node / to a return node
-	//    once we have encapsulation implement this
 	//    also traverse links along reroutes and flatten them
 	LocalPinTraversalInfo TraversalInfo;
 	TraversalInfo.PinDefaultValueOverrides = &PinDefaultValueOverrides;
@@ -2664,15 +2756,7 @@ void FRigVMParserAST::Inline(URigVMGraph* InGraph, const TArray<URigVMNode*>& In
 
 	for (URigVMNode* Node : Nodes)
 	{
-		if (Node->IsA<URigVMRerouteNode>())
-		{
-			continue;
-		}
-
-		for (URigVMPin* Pin : Node->GetPins())
-		{
-			LocalPinTraversalInfo::VisitPin(Pin, TraversalInfo);
-		}
+		LocalPinTraversalInfo::VisitNode(Node, TraversalInfo);
 	}
 }
 
diff --git a/Engine/Source/Developer/RigVMDeveloper/Private/RigVMModel/RigVMController.cpp b/Engine/Source/Developer/RigVMDeveloper/Private/RigVMModel/RigVMController.cpp
--- a/Engine/Source/Developer/RigVMDeveloper/Private/RigVMModel/RigVMController.cpp
+++ b/Engine/Source/Developer/RigVMDeveloper/Private/RigVMModel/RigVMController.cpp
@@ -1781,6 +1781,25 @@ URigVMCollapseNode* URigVMController::CollapseNodes(const TArray<URigVMNode*>& I
 			continue;
 		}
 
+		if (PinToCollapse->IsExecuteContext())
+		{
+			bool bFoundExistingPin = false;
+			for (URigVMPin* ExistingPin : CollapseNode->Pins)
+			{
+				if (ExistingPin->IsExecuteContext())
+				{
+					CollapsedPins.Add(PinToCollapse, ExistingPin);
+					bFoundExistingPin = true;
+					break;
+				}
+			}
+
+			if (bFoundExistingPin)
+			{
+				continue;
+			}
+		}
+
 		// for links that connect to the right side of the collapse
 		// node, we need to skip sub pins of already exposed pins
 		if (bSourceToBeCollapsed)
@@ -1810,6 +1829,16 @@ URigVMCollapseNode* URigVMController::CollapseNodes(const TArray<URigVMNode*>& I
 		URigVMPin* CollapsedPin = NewObject<URigVMPin>(CollapseNode, PinName);
 		ConfigurePinFromPin(CollapsedPin, PinToCollapse);
 
+		if (CollapsedPin->IsExecuteContext())
+		{
+			CollapsedPin->Direction = ERigVMPinDirection::IO;
+			bContainsOutputs = true;
+		}
+		else if (CollapsedPin->GetDirection() == ERigVMPinDirection::IO)
+		{
+			CollapsedPin->Direction = ERigVMPinDirection::Input;
+		}
+
 		if (CollapsedPin->IsStruct())
 		{
 			AddPinsForStruct(CollapsedPin->GetScriptStruct(), CollapseNode, CollapsedPin, CollapsedPin->GetDirection(), FString(), false);
@@ -5692,30 +5721,37 @@ void URigVMController::RepopulatePinsOnNode(URigVMNode* InNode, bool bFollowCore
 
 			for (URigVMPin* LibraryPin : LibraryNode->GetPins())
 			{
-				if (LibraryPin->GetDirection() == ERigVMPinDirection::Input && !bIsEntryNode)
+				if (LibraryPin->GetDirection() == ERigVMPinDirection::IO && !LibraryPin->IsExecuteContext())
 				{
 					continue;
 				}
-				else if (LibraryPin->GetDirection() == ERigVMPinDirection::Output && bIsEntryNode)
+
+				if (bIsEntryNode)
 				{
-					continue;
+					if (LibraryPin->GetDirection() == ERigVMPinDirection::Output)
+					{
+						continue;
+					}
+				}
+				else
+				{
+					if (LibraryPin->GetDirection() == ERigVMPinDirection::Input)
+					{
+						continue;
+					}
 				}
 
 				URigVMPin* ExposedPin = NewObject<URigVMPin>(InNode, LibraryPin->GetFName());
 				ConfigurePinFromPin(ExposedPin, LibraryPin);
 
-				if (ExposedPin->GetDirection() == ERigVMPinDirection::Input)
+				if (bIsEntryNode)
 				{
 					ExposedPin->Direction = ERigVMPinDirection::Output;
 				}
-				else if (ExposedPin->GetDirection() == ERigVMPinDirection::Output)
+				else
 				{
 					ExposedPin->Direction = ERigVMPinDirection::Input;
 				}
-				else if (ExposedPin->Direction != ERigVMPinDirection::IO)
-				{
-					checkNoEntry();
-				}
 
 				InNode->Pins.Add(ExposedPin);
 
diff --git a/Engine/Source/Developer/RigVMDeveloper/Private/RigVMModel/RigVMPin.cpp b/Engine/Source/Developer/RigVMDeveloper/Private/RigVMModel/RigVMPin.cpp
--- a/Engine/Source/Developer/RigVMDeveloper/Private/RigVMModel/RigVMPin.cpp
+++ b/Engine/Source/Developer/RigVMDeveloper/Private/RigVMModel/RigVMPin.cpp
@@ -84,22 +84,32 @@ URigVMPin::URigVMPin()
 	, DefaultValue(FString())
 	, BoundVariablePath()
 {
+#if WITH_EDITORONLY_DATA
+	CachedPinPath = GetPinPath();
+#endif
 }
 
 FString URigVMPin::GetPinPath() const
 {
+	FString PinPath;
+
 	URigVMPin* ParentPin = GetParentPin();
 	if(ParentPin)
 	{
-		return FString::Printf(TEXT("%s.%s"), *ParentPin->GetPinPath(), *GetName());
+		PinPath = FString::Printf(TEXT("%s.%s"), *ParentPin->GetPinPath(), *GetName());
 	}
-
-	URigVMNode* Node = GetNode();
-	if (Node == nullptr)
+	else
 	{
-		return FString();
+		URigVMNode* Node = GetNode();
+		if (Node != nullptr)
+		{
+			PinPath = FString::Printf(TEXT("%s.%s"), *Node->GetNodePath(), *GetName());
+		}
 	}
-	return FString::Printf(TEXT("%s.%s"), *Node->GetNodePath(), *GetName());
+#if WITH_EDITORONLY_DATA
+	CachedPinPath = PinPath;
+#endif
+	return PinPath;
 }
 
 FString URigVMPin::GetSegmentPath() const
diff --git a/Engine/Source/Developer/RigVMDeveloper/Public/RigVMModel/RigVMPin.h b/Engine/Source/Developer/RigVMDeveloper/Public/RigVMModel/RigVMPin.h
--- a/Engine/Source/Developer/RigVMDeveloper/Public/RigVMModel/RigVMPin.h
+++ b/Engine/Source/Developer/RigVMDeveloper/Public/RigVMModel/RigVMPin.h
@@ -328,6 +328,11 @@ private:
 	UPROPERTY()
 	FName DisplayName;
 
+#if WITH_EDITORONLY_DATA
+	UPROPERTY(transient)
+	mutable FString CachedPinPath;
+#endif
+
 	UPROPERTY()
 	ERigVMPinDirection Direction;
 
