RigVM: Literal consolidation optimization

diff --git a/Engine/Source/Developer/RigVMDeveloper/Private/RigVMCompiler/RigVMAST.cpp b/Engine/Source/Developer/RigVMDeveloper/Private/RigVMCompiler/RigVMAST.cpp
--- a/Engine/Source/Developer/RigVMDeveloper/Private/RigVMCompiler/RigVMAST.cpp
+++ b/Engine/Source/Developer/RigVMDeveloper/Private/RigVMCompiler/RigVMAST.cpp
@@ -2048,7 +2048,22 @@ void FRigVMParserAST::FoldLiterals()
 			FString DefaultValue = LiteralExpr->GetDefaultValue();
 			if (DefaultValue.IsEmpty())
 			{
-				continue;
+				if (LiteralExpr->GetCPPType() == TEXT("bool"))
+				{
+					DefaultValue = TEXT("False");
+				}
+				else if (LiteralExpr->GetCPPType() == TEXT("float"))
+				{
+					DefaultValue = TEXT("0.000000");
+				}
+				else if (LiteralExpr->GetCPPType() == TEXT("int32"))
+				{
+					DefaultValue = TEXT("0");
+				}
+				else
+				{
+					continue;
+				}
 			}
 
 			FString Hash = FString::Printf(TEXT("[%s] %s"), *LiteralExpr->GetCPPType(), *DefaultValue);
