Control Rig: Rename Struct Node to Unit Node

diff --git a/Engine/Config/BaseEngine.ini b/Engine/Config/BaseEngine.ini
--- a/Engine/Config/BaseEngine.ini
+++ b/Engine/Config/BaseEngine.ini
@@ -842,8 +842,10 @@ NearClipPlane=10.0
 +PropertyRedirects=(OldName="CameraShakeSourceComponent.bAutoPlay",NewName="bAutoStart")
 
 ;5.0
++ClassRedirects=(OldName="RigVMStructNode",NewName="/Script/RigVMDeveloper.RigVMUnitNode")
 +EnumRedirects=(OldName="/Script/Slate.ETextTransformPolicy", NewName="/Script/SlateCore.ETextTransformPolicy")
 +PropertyRedirects=(OldName="BlendProfile.BlendProfileMode", NewName="Mode")
++PropertyRedirects=(OldName="RigVMInjectionInfo.StructNode", NewName="UnitNode")
 
 [CoreUObject.Metadata]
 ; Note: UnrealHeaderTool should be rerun after making changes to MetadataRedirects to catch any keys specified in class headers
diff --git a/Engine/Plugins/Experimental/ControlRig/Source/ControlRigDeveloper/Private/ControlRigBlueprint.cpp b/Engine/Plugins/Experimental/ControlRig/Source/ControlRigDeveloper/Private/ControlRigBlueprint.cpp
--- a/Engine/Plugins/Experimental/ControlRig/Source/ControlRigDeveloper/Private/ControlRigBlueprint.cpp
+++ b/Engine/Plugins/Experimental/ControlRig/Source/ControlRigDeveloper/Private/ControlRigBlueprint.cpp
@@ -115,7 +115,7 @@ void UControlRigBlueprint::InitializeModelIfRequired(bool bRecompileVM)
 		// c) create external variable (mapped to the passed in tfunction)
 		// the last one is defined within the blueprint since the controller
 		// doesn't own the variables and can't create one itself.
-		Controller->SetupDefaultStructNodeDelegates(TDelegate<FName(FRigVMExternalVariable)>::CreateLambda(
+		Controller->SetupDefaultUnitNodeDelegates(TDelegate<FName(FRigVMExternalVariable)>::CreateLambda(
 			[WeakThis](FRigVMExternalVariable InVariableToCreate) -> FName {
 				if (WeakThis.IsValid())
 				{
@@ -648,9 +648,9 @@ FName UControlRigBlueprint::AddTransientControl(URigVMPin* InPin)
 	UControlRig* CDO = Cast<UControlRig>(RigClass->GetDefaultObject(true /* create if needed */));
 
 	FRigElementKey SpaceKey;
-	if (URigVMStructNode* StructNode = Cast<URigVMStructNode>(InPin->GetPinForLink()->GetNode()))
+	if (URigVMUnitNode* UnitNode = Cast<URigVMUnitNode>(InPin->GetPinForLink()->GetNode()))
 	{
-		if (TSharedPtr<FStructOnScope> DefaultStructScope = StructNode->ConstructStructInstance())
+		if (TSharedPtr<FStructOnScope> DefaultStructScope = UnitNode->ConstructStructInstance())
 		{
 			FRigUnit* DefaultStruct = (FRigUnit*)DefaultStructScope->GetStructMemory();
 
@@ -912,7 +912,7 @@ void UControlRigBlueprint::PopulateModelFromGraphForBackwardsCompatibility(UCont
 				UScriptStruct* UnitStruct = URigVMPin::FindObjectFromCPPTypeObjectPath<UScriptStruct>(StructPath);
 				if (UnitStruct && UnitStruct->IsChildOf(FRigVMStruct::StaticStruct()))
 				{ 
-					ModelNode = Controller->AddStructNode(UnitStruct, TEXT("Execute"), NodePosition, PropertyName.ToString(), false);
+					ModelNode = Controller->AddUnitNode(UnitStruct, TEXT("Execute"), NodePosition, PropertyName.ToString(), false);
 				}
 				else if (PropertyName != NAME_None) // check if this is a variable
 				{
@@ -1003,9 +1003,9 @@ void UControlRigBlueprint::PopulateModelFromGraphForBackwardsCompatibility(UCont
 						FString PinPath = LocalHelpers::FixUpPinPath(Pin->GetName());
 
 						// check the material + mesh pins for deprecated control nodes
-						if (URigVMStructNode* ModelStructNode = Cast<URigVMStructNode>(ModelNode))
+						if (URigVMUnitNode* ModelUnitNode = Cast<URigVMUnitNode>(ModelNode))
 						{
-							if (ModelStructNode->GetScriptStruct()->IsChildOf(FRigUnit_Control::StaticStruct()))
+							if (ModelUnitNode->GetScriptStruct()->IsChildOf(FRigUnit_Control::StaticStruct()))
 							{
 								if (Pin->GetName().EndsWith(TEXT(".StaticMesh")) || Pin->GetName().EndsWith(TEXT(".Materials")))
 								{
@@ -1084,13 +1084,13 @@ void UControlRigBlueprint::SetupPinRedirectorsForBackwardsCompatibility()
 {
 	for (URigVMNode* Node : Model->GetNodes())
 	{
-		if (URigVMStructNode* StructNode = Cast<URigVMStructNode>(Node))
+		if (URigVMUnitNode* UnitNode = Cast<URigVMUnitNode>(Node))
 		{
-			UScriptStruct* Struct = StructNode->GetScriptStruct();
+			UScriptStruct* Struct = UnitNode->GetScriptStruct();
 			if (Struct == FRigUnit_SetBoneTransform::StaticStruct())
 			{
-				URigVMPin* TransformPin = StructNode->FindPin(TEXT("Transform"));
-				URigVMPin* ResultPin = StructNode->FindPin(TEXT("Result"));
+				URigVMPin* TransformPin = UnitNode->FindPin(TEXT("Transform"));
+				URigVMPin* ResultPin = UnitNode->FindPin(TEXT("Result"));
 				Controller->AddPinRedirector(false, true, TransformPin->GetPinPath(), ResultPin->GetPinPath());
 			}
 		}
diff --git a/Engine/Plugins/Experimental/ControlRig/Source/ControlRigDeveloper/Private/Graph/ControlRigGraph.cpp b/Engine/Plugins/Experimental/ControlRig/Source/ControlRigDeveloper/Private/Graph/ControlRigGraph.cpp
--- a/Engine/Plugins/Experimental/ControlRig/Source/ControlRigDeveloper/Private/Graph/ControlRigGraph.cpp
+++ b/Engine/Plugins/Experimental/ControlRig/Source/ControlRigDeveloper/Private/Graph/ControlRigGraph.cpp
@@ -464,7 +464,7 @@ void UControlRigGraph::HandleModifiedEvent(ERigVMGraphNotifType InNotifType, URi
 							RigNode->ReconstructNode_Internal(true);
 						}
 					}
-					else if (Cast<URigVMStructNode>(ModelPin->GetNode()))
+					else if (Cast<URigVMUnitNode>(ModelPin->GetNode()))
 					{
 						RigNode->InvalidateNodeTitle();
 					}
diff --git a/Engine/Plugins/Experimental/ControlRig/Source/ControlRigDeveloper/Private/Graph/ControlRigGraphNode.cpp b/Engine/Plugins/Experimental/ControlRig/Source/ControlRigDeveloper/Private/Graph/ControlRigGraphNode.cpp
--- a/Engine/Plugins/Experimental/ControlRig/Source/ControlRigDeveloper/Private/Graph/ControlRigGraphNode.cpp
+++ b/Engine/Plugins/Experimental/ControlRig/Source/ControlRigDeveloper/Private/Graph/ControlRigGraphNode.cpp
@@ -46,11 +46,11 @@ FText UControlRigGraphNode::GetNodeTitle(ENodeTitleType::Type TitleType) const
 	{
 		if(URigVMNode* ModelNode = GetModelNode())
 		{
-			if (URigVMStructNode* StructNode = Cast<URigVMStructNode>(ModelNode))
+			if (URigVMUnitNode* UnitNode = Cast<URigVMUnitNode>(ModelNode))
 			{
-				if (StructNode->GetScriptStruct()->IsChildOf(FRigUnit::StaticStruct()))
+				if (UnitNode->GetScriptStruct()->IsChildOf(FRigUnit::StaticStruct()))
 				{
-					if (TSharedPtr<FStructOnScope> StructOnScope = StructNode->ConstructStructInstance())
+					if (TSharedPtr<FStructOnScope> StructOnScope = UnitNode->ConstructStructInstance())
 					{
 						FRigUnit* RigUnit = (FRigUnit*)StructOnScope->GetStructMemory();
 						NodeTitle = FText::FromString(RigUnit->GetUnitLabel());
@@ -124,7 +124,7 @@ bool UControlRigGraphNode::IsDeprecated() const
 {
 	if(URigVMNode* ModelNode = GetModelNode())
 	{
-		if(URigVMStructNode* StructModelNode = Cast<URigVMStructNode>(ModelNode))
+		if(URigVMUnitNode* StructModelNode = Cast<URigVMUnitNode>(ModelNode))
 		{
 			return StructModelNode->IsDeprecated();
 		}
@@ -138,7 +138,7 @@ FEdGraphNodeDeprecationResponse UControlRigGraphNode::GetDeprecationResponse(EEd
 
 	if(URigVMNode* ModelNode = GetModelNode())
 	{
-		if(URigVMStructNode* StructModelNode = Cast<URigVMStructNode>(ModelNode))
+		if(URigVMUnitNode* StructModelNode = Cast<URigVMUnitNode>(ModelNode))
 		{
 			FString DeprecatedMetadata = StructModelNode->GetDeprecatedMetadata();
 			if (!DeprecatedMetadata.IsEmpty())
diff --git a/Engine/Plugins/Experimental/ControlRig/Source/ControlRigEditor/Private/ControlRigEditorModule.cpp b/Engine/Plugins/Experimental/ControlRig/Source/ControlRigEditor/Private/ControlRigEditorModule.cpp
--- a/Engine/Plugins/Experimental/ControlRig/Source/ControlRigEditor/Private/ControlRigEditorModule.cpp
+++ b/Engine/Plugins/Experimental/ControlRig/Source/ControlRigEditor/Private/ControlRigEditorModule.cpp
@@ -1176,7 +1176,7 @@ void FControlRigEditorModule::GetContextMenuActions(const UControlRigGraphSchema
 						}
 					}
 
-					if(Cast<URigVMStructNode>(ModelPin->GetNode()))
+					if(Cast<URigVMUnitNode>(ModelPin->GetNode()))
 					{
 						if (ModelPin->GetDirection() == ERigVMPinDirection::Input && 
 							!ModelPin->IsExecuteContext() &&
@@ -1224,7 +1224,7 @@ void FControlRigEditorModule::GetContextMenuActions(const UControlRigGraphSchema
 						));
 					}
 
-					if (ModelPin->GetRootPin() == ModelPin && Cast<URigVMStructNode>(ModelPin->GetNode()) != nullptr)
+					if (ModelPin->GetRootPin() == ModelPin && Cast<URigVMUnitNode>(ModelPin->GetNode()) != nullptr)
 					{
 						if (ModelPin->HasInjectedNodes())
 						{
@@ -1249,11 +1249,11 @@ void FControlRigEditorModule::GetContextMenuActions(const UControlRigGraphSchema
 							for (URigVMInjectionInfo* Injection : ModelPin->GetInjectedNodes())
 							{
 								FString PrototypeName;
-								if (Injection->StructNode->GetScriptStruct()->GetStringMetaDataHierarchical(TEXT("PrototypeName"), &PrototypeName))
+								if (Injection->UnitNode->GetScriptStruct()->GetStringMetaDataHierarchical(TEXT("PrototypeName"), &PrototypeName))
 								{
 									if (PrototypeName == TEXT("AlphaInterp"))
 									{
-										InterpNode = Injection->StructNode;
+										InterpNode = Injection->UnitNode;
 										break;
 									}
 								}
@@ -1286,7 +1286,7 @@ void FControlRigEditorModule::GetContextMenuActions(const UControlRigGraphSchema
 										if (Injection)
 										{
 											TArray<FName> NodeNames;
-											NodeNames.Add(Injection->StructNode->GetFName());
+											NodeNames.Add(Injection->UnitNode->GetFName());
 											RigBlueprint->Controller->SetNodeSelection(NodeNames);
 										}
 									})
@@ -1327,11 +1327,11 @@ void FControlRigEditorModule::GetContextMenuActions(const UControlRigGraphSchema
 							for (URigVMInjectionInfo* Injection : ModelPin->GetInjectedNodes())
 							{
 								FString PrototypeName;
-								if (Injection->StructNode->GetScriptStruct()->GetStringMetaDataHierarchical(TEXT("PrototypeName"), &PrototypeName))
+								if (Injection->UnitNode->GetScriptStruct()->GetStringMetaDataHierarchical(TEXT("PrototypeName"), &PrototypeName))
 								{
 									if (PrototypeName == TEXT("VisualDebug"))
 									{
-										VisualDebugNode = Injection->StructNode;
+										VisualDebugNode = Injection->UnitNode;
 										break;
 									}
 								}
@@ -1368,12 +1368,12 @@ void FControlRigEditorModule::GetContextMenuActions(const UControlRigGraphSchema
 										if (Injection)
 										{
 											TArray<FName> NodeNames;
-											NodeNames.Add(Injection->StructNode->GetFName());
+											NodeNames.Add(Injection->UnitNode->GetFName());
 											RigBlueprint->Controller->SetNodeSelection(NodeNames);
 
-											if (URigVMStructNode* StructNode = Cast<URigVMStructNode>(ModelPin->GetNode()))
+											if (URigVMUnitNode* UnitNode = Cast<URigVMUnitNode>(ModelPin->GetNode()))
 											{
-												if (TSharedPtr<FStructOnScope> DefaultStructScope = StructNode->ConstructStructInstance())
+												if (TSharedPtr<FStructOnScope> DefaultStructScope = UnitNode->ConstructStructInstance())
 												{
 													FRigUnit* DefaultStruct = (FRigUnit*)DefaultStructScope->GetStructMemory();
 
@@ -1388,7 +1388,7 @@ void FControlRigEditorModule::GetContextMenuActions(const UControlRigGraphSchema
 
 													if (SpaceKey.IsValid())
 													{
-														if (URigVMPin* SpacePin = Injection->StructNode->FindPin(TEXT("Space")))
+														if (URigVMPin* SpacePin = Injection->UnitNode->FindPin(TEXT("Space")))
 														{
 															if(URigVMPin* SpaceTypePin = SpacePin->FindSubPin(TEXT("Type")))
 															{
@@ -1463,10 +1463,10 @@ void FControlRigEditorModule::GetContextMenuActions(const UControlRigGraphSchema
 						FRigHierarchyContainer TemporaryHierarchy = RigBlueprint->HierarchyContainer;
 						FRigUnit* StructMemory = nullptr;
 						UScriptStruct* ScriptStruct = nullptr;
-						if (URigVMStructNode* StructNode = Cast<URigVMStructNode>(ModelNode))
+						if (URigVMUnitNode* UnitNode = Cast<URigVMUnitNode>(ModelNode))
 						{
-							ScriptStruct = StructNode->GetScriptStruct();
-							StructOnScope = StructNode->ConstructStructInstance(false /* default */);
+							ScriptStruct = UnitNode->GetScriptStruct();
+							StructOnScope = UnitNode->ConstructStructInstance(false /* default */);
 							StructMemory = (FRigUnit*)StructOnScope->GetStructMemory();
 
 							FRigUnitContext RigUnitContext;
@@ -1622,7 +1622,7 @@ void FControlRigEditorModule::GetContextMenuActions(const UControlRigGraphSchema
 
 				if (const UControlRigGraphNode* RigNode = Cast<const UControlRigGraphNode>(Context->Node))
 				{
-					if (URigVMStructNode* StructNode = Cast<URigVMStructNode>(RigNode->GetModelNode()))
+					if (URigVMUnitNode* UnitNode = Cast<URigVMUnitNode>(RigNode->GetModelNode()))
 					{
 						FToolMenuSection& SettingsSection = Menu->AddSection("EdGraphSchemaSettings", LOCTEXT("SettingsHeader", "Settings"));
 						SettingsSection.AddMenuEntry(
@@ -1630,17 +1630,17 @@ void FControlRigEditorModule::GetContextMenuActions(const UControlRigGraphSchema
 							LOCTEXT("SaveDefaultExpansionState", "Save Default Expansion State"),
 							LOCTEXT("SaveDefaultExpansionState_Tooltip", "Saves the expansion state of all pins of the node as the default."),
 							FSlateIcon(),
-							FUIAction(FExecuteAction::CreateLambda([StructNode]() {
+							FUIAction(FExecuteAction::CreateLambda([UnitNode]() {
 
 #if WITH_EDITORONLY_DATA
 
 								FScopedTransaction Transaction(LOCTEXT("RigUnitDefaultExpansionStateChanged", "Changed Rig Unit Default Expansion State"));
 								UControlRigSettings::Get()->Modify();
 
-								FControlRigSettingsPerPinBool& ExpansionMap = UControlRigSettings::Get()->RigUnitPinExpansion.FindOrAdd(StructNode->GetScriptStruct()->GetName());
+								FControlRigSettingsPerPinBool& ExpansionMap = UControlRigSettings::Get()->RigUnitPinExpansion.FindOrAdd(UnitNode->GetScriptStruct()->GetName());
 								ExpansionMap.Values.Empty();
 
-								TArray<URigVMPin*> Pins = StructNode->GetAllPinsRecursively();
+								TArray<URigVMPin*> Pins = UnitNode->GetAllPinsRecursively();
 								for (URigVMPin* Pin : Pins)
 								{
 									if (Pin->GetSubPins().Num() == 0)
diff --git a/Engine/Plugins/Experimental/ControlRig/Source/ControlRigEditor/Private/Editor/ControlRigEditor.cpp b/Engine/Plugins/Experimental/ControlRig/Source/ControlRigEditor/Private/Editor/ControlRigEditor.cpp
--- a/Engine/Plugins/Experimental/ControlRig/Source/ControlRigEditor/Private/Editor/ControlRigEditor.cpp
+++ b/Engine/Plugins/Experimental/ControlRig/Source/ControlRigEditor/Private/Editor/ControlRigEditor.cpp
@@ -290,7 +290,7 @@ void FControlRigEditor::InitControlRigEditor(const EToolkitMode::Type Mode, cons
 		{
 			if (Model->GetNodes().Num() == 0)
 			{
-				URigVMNode* Node = InControlRigBlueprint->Controller->AddStructNode(FRigUnit_BeginExecution::StaticStruct(), TEXT("Execute"), FVector2D::ZeroVector, FString(), false);
+				URigVMNode* Node = InControlRigBlueprint->Controller->AddUnitNode(FRigUnit_BeginExecution::StaticStruct(), TEXT("Execute"), FVector2D::ZeroVector, FString(), false);
 				if (Node)
 				{
 					TArray<FName> NodeNames;
@@ -987,18 +987,18 @@ void FControlRigEditor::SetDetailObjects(const TArray<UObject*>& InObjects)
 
 void FControlRigEditor::SetDetailObject(UObject* Obj)
 {
-	if (URigVMStructNode* StructNode = Cast<URigVMStructNode>(Obj))
+	if (URigVMUnitNode* UnitNode = Cast<URigVMUnitNode>(Obj))
 	{
 		ClearDetailObject();
-		NodeDetailStruct = StructNode->GetScriptStruct();
-		NodeDetailName = StructNode->GetFName();
+		NodeDetailStruct = UnitNode->GetScriptStruct();
+		NodeDetailName = UnitNode->GetFName();
 
 		if (NodeDetailStruct)
 		{
 			NodeDetailBuffer.AddUninitialized(NodeDetailStruct->GetStructureSize());
 			NodeDetailStruct->InitializeDefaultValue(NodeDetailBuffer.GetData());
 
-			FString StructDefaultValue = StructNode->GetStructDefaultValue();
+			FString StructDefaultValue = UnitNode->GetStructDefaultValue();
 			NodeDetailStruct->ImportText(*StructDefaultValue, NodeDetailBuffer.GetData(), nullptr, PPF_None, nullptr, NodeDetailStruct->GetName());
 
 			StructToDisplay = MakeShareable(new FStructOnScope(NodeDetailStruct, (uint8*)NodeDetailBuffer.GetData()));
@@ -1026,7 +1026,7 @@ void FControlRigEditor::SetDetailObject(UObject* Obj)
 				}
 
 				bool bEditable = true;
-				if (URigVMPin* Pin = StructNode->FindPin(Property->GetName()))
+				if (URigVMPin* Pin = UnitNode->FindPin(Property->GetName()))
 				{
 					Pin = Pin->GetPinForLink();
 					if (Pin->GetDirection() == ERigVMPinDirection::Output)
@@ -3497,7 +3497,7 @@ void FControlRigEditor::OnGraphNodeDropToPerform(TSharedPtr<FGraphNodeDragDropOp
 								{
 									Controller->OpenUndoBracket(TEXT("Create Collection from Items"));
 
-									if (URigVMNode* ItemsNode = Controller->AddStructNode(FRigUnit_CollectionItems::StaticStruct(), TEXT("Execute"), NodePosition))
+									if (URigVMNode* ItemsNode = Controller->AddUnitNode(FRigUnit_CollectionItems::StaticStruct(), TEXT("Execute"), NodePosition))
 									{
 										if (URigVMPin* ItemsPin = ItemsNode->FindPin(TEXT("Items")))
 										{
@@ -3804,7 +3804,7 @@ void FControlRigEditor::HandleMakeElementGetterSetter(ERigElementGetterSetterTyp
 		}
 
 		FName Name = FControlRigBlueprintUtils::ValidateName(Blueprint, StructTemplate->GetName());
-		if (URigVMStructNode* ModelNode = Blueprint->Controller->AddStructNode(StructTemplate, TEXT("Execute"), NodePosition))
+		if (URigVMUnitNode* ModelNode = Blueprint->Controller->AddUnitNode(StructTemplate, TEXT("Execute"), NodePosition))
 		{
 			FString ItemTypeStr = StaticEnum<ERigElementType>()->GetDisplayNameTextByValue((int64)Key.Type).ToString();
 			NewNode.Name = ModelNode->GetFName();
diff --git a/Engine/Plugins/Experimental/ControlRig/Source/ControlRigEditor/Private/Editor/SControlRigStackView.cpp b/Engine/Plugins/Experimental/ControlRig/Source/ControlRigEditor/Private/Editor/SControlRigStackView.cpp
--- a/Engine/Plugins/Experimental/ControlRig/Source/ControlRigEditor/Private/Editor/SControlRigStackView.cpp
+++ b/Engine/Plugins/Experimental/ControlRig/Source/ControlRigEditor/Private/Editor/SControlRigStackView.cpp
@@ -313,17 +313,17 @@ void SControlRigStackView::PopulateStackView(URigVM* InVM)
 		{
 			URigVMNode* Node = Nodes[NodeIndex];
 			// only struct nodes among all nodes has StaticExecute() that generates actual instructions
-			if (URigVMStructNode* StructNode = Cast<URigVMStructNode>(Node))
+			if (URigVMUnitNode* UnitNode = Cast<URigVMUnitNode>(Node))
 			{
-				int32 InstructionIndex = StructNode->GetInstructionIndex();
+				int32 InstructionIndex = UnitNode->GetInstructionIndex();
 				// only active nodes are assigned a valid instruction index, hence the check here
 				if (InstructionIndex >= 0 && InstructionIndex < InstructionIndexToNodeIndex.Num())
 				{
 					InstructionIndexToNodeIndex[InstructionIndex] = NodeIndex;
 				}
-				FString DisplayName = StructNode->GetNodeTitle();
+				FString DisplayName = UnitNode->GetNodeTitle();
 #if WITH_EDITOR
-				UScriptStruct* Struct = StructNode->GetScriptStruct();
+				UScriptStruct* Struct = UnitNode->GetScriptStruct();
 				FString MenuDescSuffixMetadata;
 				if (Struct)
 				{
@@ -331,11 +331,11 @@ void SControlRigStackView::PopulateStackView(URigVM* InVM)
 				}
 				if (!MenuDescSuffixMetadata.IsEmpty())
 				{
-					DisplayName = FString::Printf(TEXT("%s %s"), *StructNode->GetNodeTitle(), *MenuDescSuffixMetadata);
+					DisplayName = FString::Printf(TEXT("%s %s"), *UnitNode->GetNodeTitle(), *MenuDescSuffixMetadata);
 				}
 #endif
 				// this is needed for name replacement later
-				NodeNameToDisplayName.Add(StructNode->GetName(), DisplayName);
+				NodeNameToDisplayName.Add(UnitNode->GetName(), DisplayName);
 			}
 		}
 
diff --git a/Engine/Plugins/Experimental/ControlRig/Source/ControlRigEditor/Private/Graph/ControlRigGraphPanelPinFactory.cpp b/Engine/Plugins/Experimental/ControlRig/Source/ControlRigEditor/Private/Graph/ControlRigGraphPanelPinFactory.cpp
--- a/Engine/Plugins/Experimental/ControlRig/Source/ControlRigEditor/Private/Graph/ControlRigGraphPanelPinFactory.cpp
+++ b/Engine/Plugins/Experimental/ControlRig/Source/ControlRigEditor/Private/Graph/ControlRigGraphPanelPinFactory.cpp
@@ -14,7 +14,7 @@
 #include "NodeFactory.h"
 #include "EdGraphSchema_K2.h"
 #include "Curves/CurveFloat.h"
-#include "RigVMModel/Nodes/RigVMStructNode.h"
+#include "RigVMModel/Nodes/RigVMUnitNode.h"
 #include "RigVMCore/RigVMExecuteContext.h"
 #include "IPropertyAccessEditor.h"
 
diff --git a/Engine/Plugins/Experimental/ControlRig/Source/ControlRigEditor/Private/Graph/NodeSpawners/ControlRigUnitNodeSpawner.cpp b/Engine/Plugins/Experimental/ControlRig/Source/ControlRigEditor/Private/Graph/NodeSpawners/ControlRigUnitNodeSpawner.cpp
--- a/Engine/Plugins/Experimental/ControlRig/Source/ControlRigEditor/Private/Graph/NodeSpawners/ControlRigUnitNodeSpawner.cpp
+++ b/Engine/Plugins/Experimental/ControlRig/Source/ControlRigEditor/Private/Graph/NodeSpawners/ControlRigUnitNodeSpawner.cpp
@@ -141,10 +141,10 @@ UControlRigGraphNode* UControlRigUnitNodeSpawner::SpawnNode(UEdGraph* ParentGrap
 			Controller->OpenUndoBracket(FString::Printf(TEXT("Add '%s' Node"), *Name.ToString()));
 		}
 
-		FRigVMStructNodeCreatedContext& StructNodeCreatedContext = Controller->GetStructNodeCreatedContext();
-		FRigVMStructNodeCreatedContext::FScope ReasonScope(StructNodeCreatedContext, ERigVMNodeCreatedReason::NodeSpawner);
+		FRigVMUnitNodeCreatedContext& UnitNodeCreatedContext = Controller->GetUnitNodeCreatedContext();
+		FRigVMUnitNodeCreatedContext::FScope ReasonScope(UnitNodeCreatedContext, ERigVMNodeCreatedReason::NodeSpawner);
 
-		if (URigVMStructNode* ModelNode = Controller->AddStructNode(StructTemplate, TEXT("Execute"), Location, Name.ToString(), bUndo))
+		if (URigVMUnitNode* ModelNode = Controller->AddUnitNode(StructTemplate, TEXT("Execute"), Location, Name.ToString(), bUndo))
 		{
 			NewNode = Cast<UControlRigGraphNode>(RigGraph->FindNodeForModelNodeName(ModelNode->GetFName()));
 
diff --git a/Engine/Plugins/Experimental/ControlRig/Source/ControlRigEditor/Private/Graph/NodeSpawners/ControlRigUnitNodeSpawner.h b/Engine/Plugins/Experimental/ControlRig/Source/ControlRigEditor/Private/Graph/NodeSpawners/ControlRigUnitNodeSpawner.h
--- a/Engine/Plugins/Experimental/ControlRig/Source/ControlRigEditor/Private/Graph/NodeSpawners/ControlRigUnitNodeSpawner.h
+++ b/Engine/Plugins/Experimental/ControlRig/Source/ControlRigEditor/Private/Graph/NodeSpawners/ControlRigUnitNodeSpawner.h
@@ -15,7 +15,7 @@
 #include "ControlRigBlueprint.h"
 #include "Graph/ControlRigGraph.h"
 #include "Graph/ControlRigGraphNode.h"
-#include "RigVMModel/Nodes/RigVMStructNode.h"
+#include "RigVMModel/Nodes/RigVMUnitNode.h"
 #include "ControlRigUnitNodeSpawner.generated.h"
 
 class UControlRigGraphNode;
diff --git a/Engine/Plugins/Experimental/ControlRig/Source/ControlRigEditor/Private/Graph/SControlRigGraphNode.cpp b/Engine/Plugins/Experimental/ControlRig/Source/ControlRigEditor/Private/Graph/SControlRigGraphNode.cpp
--- a/Engine/Plugins/Experimental/ControlRig/Source/ControlRigEditor/Private/Graph/SControlRigGraphNode.cpp
+++ b/Engine/Plugins/Experimental/ControlRig/Source/ControlRigEditor/Private/Graph/SControlRigGraphNode.cpp
@@ -947,7 +947,7 @@ TArray<FOverlayWidgetInfo> SControlRigGraphNode::GetOverlayWidgets(bool bSelecte
 				{
 					for (URigVMInjectionInfo* Injection : ModelPin->GetInjectedNodes())
 					{
-						URigVMStructNode* VisualDebugNode = Injection->StructNode;
+						URigVMUnitNode* VisualDebugNode = Injection->UnitNode;
 
 						FString PrototypeName;
 						if (VisualDebugNode->GetScriptStruct()->GetStringMetaDataHierarchical(TEXT("PrototypeName"), &PrototypeName))
diff --git a/Engine/Plugins/Experimental/ControlRig/Source/ControlRigEditor/Private/Graph/SControlRigGraphPinCurveFloat.cpp b/Engine/Plugins/Experimental/ControlRig/Source/ControlRigEditor/Private/Graph/SControlRigGraphPinCurveFloat.cpp
--- a/Engine/Plugins/Experimental/ControlRig/Source/ControlRigEditor/Private/Graph/SControlRigGraphPinCurveFloat.cpp
+++ b/Engine/Plugins/Experimental/ControlRig/Source/ControlRigEditor/Private/Graph/SControlRigGraphPinCurveFloat.cpp
@@ -11,7 +11,7 @@
 #include "ScopedTransaction.h"
 #include "DetailLayoutBuilder.h"
 #include "PropertyPathHelpers.h"
-#include "RigVMModel/Nodes/RigVMStructNode.h"
+#include "RigVMModel/Nodes/RigVMUnitNode.h"
 #include "RigVMModel/RigVMController.h"
 
 #include "IControlRigEditorModule.h"
@@ -109,7 +109,7 @@ bool SControlRigGraphPinCurveFloat::IsValidCurve(FRichCurveEditInfo CurveInfo)
 	{
 		if (UControlRigGraphNode* Node = Cast<UControlRigGraphNode>(Pin->GetOwningNode()))
 		{
-			if (URigVMStructNode* StructModelNode = Cast<URigVMStructNode>(Node->GetModelNode()))
+			if (URigVMUnitNode* StructModelNode = Cast<URigVMUnitNode>(Node->GetModelNode()))
 			{
 				FString NodeName, PropertyName;
 				Pin->PinName.ToString().Split(TEXT("."), &NodeName, &PropertyName);
diff --git a/Engine/Source/Developer/RigVMDeveloper/Private/RigVMCompiler/RigVMAST.cpp b/Engine/Source/Developer/RigVMDeveloper/Private/RigVMCompiler/RigVMAST.cpp
--- a/Engine/Source/Developer/RigVMDeveloper/Private/RigVMCompiler/RigVMAST.cpp
+++ b/Engine/Source/Developer/RigVMDeveloper/Private/RigVMCompiler/RigVMAST.cpp
@@ -2,7 +2,7 @@
 
 #include "RigVMCompiler/RigVMAST.h"
 #include "RigVMCompiler/RigVMCompiler.h"
-#include "RigVMModel/Nodes/RigVMStructNode.h"
+#include "RigVMModel/Nodes/RigVMUnitNode.h"
 #include "RigVMModel/Nodes/RigVMParameterNode.h"
 #include "RigVMModel/Nodes/RigVMVariableNode.h"
 #include "RigVMModel/Nodes/RigVMCommentNode.h"
@@ -402,7 +402,7 @@ FString FRigVMExprAST::DumpDot(TArray<bool>& OutExpressionDefined, const FString
 			}
 			case EType::CallExtern:
 			{
-				if (URigVMStructNode* Node = Cast<URigVMStructNode>(To<FRigVMCallExternExprAST>()->GetNode()))
+				if (URigVMUnitNode* Node = Cast<URigVMUnitNode>(To<FRigVMCallExternExprAST>()->GetNode()))
 				{
 					Label = Node->GetScriptStruct()->GetName();
 				}
@@ -680,9 +680,9 @@ bool FRigVMVarExprAST::IsEnumValue() const
 
 bool FRigVMVarExprAST::SupportsSoftLinks() const
 {
-	if (URigVMStructNode* StructNode = Cast<URigVMStructNode>(Pin->GetNode()))
+	if (URigVMUnitNode* UnitNode = Cast<URigVMUnitNode>(Pin->GetNode()))
 	{
-		if (StructNode->IsLoopNode())
+		if (UnitNode->IsLoopNode())
 		{
 			if (Pin->GetFName() != FRigVMStruct::ExecuteContextName &&
 				Pin->GetFName() != FRigVMStruct::ForLoopCompletedPinName)
@@ -966,9 +966,9 @@ FRigVMExprAST* FRigVMParserAST::TraverseMutableNode(URigVMNode* InNode, FRigVMEx
 				if (SourcePin->IsExecuteContext())
 				{
 					bool bIsForLoop = false;
-					if (URigVMStructNode* StructNode = Cast<URigVMStructNode>(InNode))
+					if (URigVMUnitNode* UnitNode = Cast<URigVMUnitNode>(InNode))
 					{
-						bIsForLoop = StructNode->IsLoopNode();
+						bIsForLoop = UnitNode->IsLoopNode();
 					}
 
 					FRigVMExprAST* ParentExpr = InParentExpr;
@@ -1592,9 +1592,9 @@ void FRigVMParserAST::FoldAssignments()
 		}
 
 		// if this node is a loop node - let's skip the folding
-		if (URigVMStructNode* StructNode = Cast<URigVMStructNode>(AssignExpr->GetTargetPin()->GetNode()))
+		if (URigVMUnitNode* UnitNode = Cast<URigVMUnitNode>(AssignExpr->GetTargetPin()->GetNode()))
 		{
-			if (StructNode->IsLoopNode())
+			if (UnitNode->IsLoopNode())
 			{
 				continue;
 			}
diff --git a/Engine/Source/Developer/RigVMDeveloper/Private/RigVMCompiler/RigVMCompiler.cpp b/Engine/Source/Developer/RigVMDeveloper/Private/RigVMCompiler/RigVMCompiler.cpp
--- a/Engine/Source/Developer/RigVMDeveloper/Private/RigVMCompiler/RigVMCompiler.cpp
+++ b/Engine/Source/Developer/RigVMDeveloper/Private/RigVMCompiler/RigVMCompiler.cpp
@@ -322,9 +322,9 @@ void URigVMCompiler::TraverseExpression(const FRigVMExprAST* InExpr, FRigVMCompi
 		case FRigVMExprAST::EType::CallExtern:
 		{
 			const FRigVMCallExternExprAST* CallExternExpr = InExpr->To<FRigVMCallExternExprAST>();
-			if (URigVMStructNode* StructNode = Cast<URigVMStructNode>(CallExternExpr->GetNode()))
+			if (URigVMUnitNode* UnitNode = Cast<URigVMUnitNode>(CallExternExpr->GetNode()))
 			{
-				if (StructNode->IsLoopNode())
+				if (UnitNode->IsLoopNode())
 				{
 					TraverseForLoop(CallExternExpr, WorkData);
 					break;
@@ -416,11 +416,11 @@ void URigVMCompiler::TraverseBlock(const FRigVMBlockExprAST* InExpr, FRigVMCompi
 
 void URigVMCompiler::TraverseEntry(const FRigVMEntryExprAST* InExpr, FRigVMCompilerWorkData& WorkData)
 {
-	URigVMStructNode* StructNode = Cast<URigVMStructNode>(InExpr->GetNode());
+	URigVMUnitNode* UnitNode = Cast<URigVMUnitNode>(InExpr->GetNode());
 
 	if (WorkData.bSetupMemory)
 	{
-		TSharedPtr<FStructOnScope> DefaultStruct = StructNode->ConstructStructInstance();
+		TSharedPtr<FStructOnScope> DefaultStruct = UnitNode->ConstructStructInstance();
 		WorkData.DefaultStructs.Add(DefaultStruct);
 		TraverseChildren(InExpr, WorkData);
 		WorkData.DefaultStructs.Pop();
@@ -442,11 +442,11 @@ void URigVMCompiler::TraverseEntry(const FRigVMEntryExprAST* InExpr, FRigVMCompi
 		}
 
 		// setup the instruction
-		int32 FunctionIndex = WorkData.VM->AddRigVMFunction(StructNode->GetScriptStruct(), StructNode->GetMethodName());
+		int32 FunctionIndex = WorkData.VM->AddRigVMFunction(UnitNode->GetScriptStruct(), UnitNode->GetMethodName());
 		WorkData.VM->GetByteCode().AddExecuteOp(FunctionIndex, Operands);
 		
 		int32 EntryInstructionIndex = WorkData.VM->GetByteCode().GetNumInstructions() - 1;
-		FName Entryname = StructNode->GetEventName();
+		FName Entryname = UnitNode->GetEventName();
 
 		if (WorkData.VM->GetByteCode().FindEntryIndex(Entryname) == INDEX_NONE)
 		{
@@ -458,7 +458,7 @@ void URigVMCompiler::TraverseEntry(const FRigVMEntryExprAST* InExpr, FRigVMCompi
 
 		if (Settings.SetupNodeInstructionIndex)
 		{
-			StructNode->InstructionIndex = EntryInstructionIndex;
+			UnitNode->InstructionIndex = EntryInstructionIndex;
 		}
 
 		TraverseChildren(InExpr, WorkData);
@@ -467,20 +467,20 @@ void URigVMCompiler::TraverseEntry(const FRigVMEntryExprAST* InExpr, FRigVMCompi
 
 int32 URigVMCompiler::TraverseCallExtern(const FRigVMCallExternExprAST* InExpr, FRigVMCompilerWorkData& WorkData)
 {
-	URigVMStructNode* StructNode = Cast<URigVMStructNode>(InExpr->GetNode());
+	URigVMUnitNode* UnitNode = Cast<URigVMUnitNode>(InExpr->GetNode());
 
 	int32 InstructionIndex = INDEX_NONE;
 
 	if (WorkData.bSetupMemory)
 	{
-		TSharedPtr<FStructOnScope> DefaultStruct = StructNode->ConstructStructInstance();
+		TSharedPtr<FStructOnScope> DefaultStruct = UnitNode->ConstructStructInstance();
 		WorkData.DefaultStructs.Add(DefaultStruct);
 		TraverseChildren(InExpr, WorkData);
 		WorkData.DefaultStructs.Pop();
 
 		if (Settings.EnablePinWatches)
 		{
-			for (URigVMPin* Pin : StructNode->Pins)
+			for (URigVMPin* Pin : UnitNode->Pins)
 			{
 				if (Pin->RequiresWatch())
 				{
@@ -535,18 +535,18 @@ int32 URigVMCompiler::TraverseCallExtern(const FRigVMCallExternExprAST* InExpr,
 		TraverseChildren(InExpr, WorkData);
 
 		// setup the instruction
-		int32 FunctionIndex = WorkData.VM->AddRigVMFunction(StructNode->GetScriptStruct(), StructNode->GetMethodName());
+		int32 FunctionIndex = WorkData.VM->AddRigVMFunction(UnitNode->GetScriptStruct(), UnitNode->GetMethodName());
 		WorkData.VM->GetByteCode().AddExecuteOp(FunctionIndex, Operands);
 		InstructionIndex = WorkData.VM->GetByteCode().GetNumInstructions() - 1;
 		if (Settings.SetupNodeInstructionIndex)
 		{
-			StructNode->InstructionIndex = InstructionIndex;
+			UnitNode->InstructionIndex = InstructionIndex;
 		}
 
-		ensure(InExpr->NumChildren() == StructNode->Pins.Num());
-		for (int32 PinIndex = 0; PinIndex < StructNode->Pins.Num(); PinIndex++)
+		ensure(InExpr->NumChildren() == UnitNode->Pins.Num());
+		for (int32 PinIndex = 0; PinIndex < UnitNode->Pins.Num(); PinIndex++)
 		{
-			URigVMPin* Pin = StructNode->Pins[PinIndex];
+			URigVMPin* Pin = UnitNode->Pins[PinIndex];
 
 			// ensure to copy the debug values
 			if (Pin->RequiresWatch() && Settings.EnablePinWatches)
@@ -578,7 +578,7 @@ void URigVMCompiler::TraverseForLoop(const FRigVMCallExternExprAST* InExpr, FRig
 		return;
 	}
 
-	URigVMStructNode* StructNode = Cast<URigVMStructNode>(InExpr->GetNode());
+	URigVMUnitNode* UnitNode = Cast<URigVMUnitNode>(InExpr->GetNode());
 
 	const FRigVMVarExprAST* CompletedExpr = InExpr->FindVarWithPinName(FRigVMStruct::ForLoopCompletedPinName);
 	check(CompletedExpr);
@@ -1387,9 +1387,9 @@ FRigVMOperand URigVMCompiler::FindOrAddRegister(const FRigVMVarExprAST* InVarExp
 		}
 		if (bRegisterIsMultiUse)
 		{
-			if (URigVMStructNode* StructNode = Cast<URigVMStructNode>(Pin->GetNode()))
+			if (URigVMUnitNode* UnitNode = Cast<URigVMUnitNode>(Pin->GetNode()))
 			{
-				if (StructNode->IsLoopNode())
+				if (UnitNode->IsLoopNode())
 				{
 					bRegisterIsMultiUse = false;
 				}
diff --git a/Engine/Source/Developer/RigVMDeveloper/Private/RigVMModel/Nodes/RigVMStructNode.cpp b/Engine/Source/Developer/RigVMDeveloper/Private/RigVMModel/Nodes/RigVMStructNode.cpp
deleted file mode 100644
diff --git a/Engine/Source/Developer/RigVMDeveloper/Private/RigVMModel/Nodes/RigVMUnitNode.cpp b/Engine/Source/Developer/RigVMDeveloper/Private/RigVMModel/Nodes/RigVMUnitNode.cpp
new file mode 100644
--- /dev/null
+++ b/Engine/Source/Developer/RigVMDeveloper/Private/RigVMModel/Nodes/RigVMUnitNode.cpp
@@ -0,0 +1,237 @@
+// Copyright Epic Games, Inc. All Rights Reserved.
+
+#include "RigVMModel/Nodes/RigVMUnitNode.h"
+#include "RigVMCore/RigVMStruct.h"
+
+const FName URigVMUnitNode::LoopSliceContextName = TEXT("LoopContext");
+
+FString URigVMUnitNode::GetNodeTitle() const
+{
+	if (UScriptStruct* Struct = GetScriptStruct())
+	{
+		return Struct->GetDisplayNameText().ToString();
+	}
+	return Super::GetNodeTitle();
+}
+
+FText URigVMUnitNode::GetToolTipText() const
+{
+	if (UScriptStruct* Struct = GetScriptStruct())
+	{
+		return Struct->GetToolTipText();
+	}
+	return URigVMNode::GetToolTipText();
+}
+
+bool URigVMUnitNode::IsDefinedAsConstant() const
+{
+	if (UScriptStruct* Struct = GetScriptStruct())
+	{
+		return Struct->HasMetaData(FRigVMStruct::ConstantMetaName);
+	}
+	return false;
+}
+
+bool URigVMUnitNode::IsDefinedAsVarying() const
+{
+	if (UScriptStruct* Struct = GetScriptStruct())
+	{
+		return Struct->HasMetaData(FRigVMStruct::VaryingMetaName);
+	}
+	return false;
+}
+
+FName URigVMUnitNode::GetEventName() const
+{
+	TSharedPtr<FStructOnScope> StructOnScope = ConstructStructInstance(true);
+	if (StructOnScope.IsValid())
+	{
+		const FRigVMStruct* StructMemory = (FRigVMStruct*)StructOnScope->GetStructMemory();
+		return StructMemory->GetEventName();
+	}
+	return NAME_None;
+}
+
+FName URigVMUnitNode::GetSliceContextForPin(URigVMPin* InRootPin, const FRigVMUserDataArray& InUserData)
+{
+	TSharedPtr<FStructOnScope> StructOnScope = ConstructStructInstance(false);
+	if (StructOnScope.IsValid())
+	{
+		const FRigVMStruct* StructMemory = (FRigVMStruct*)StructOnScope->GetStructMemory();
+		if (StructMemory->IsForLoop())
+		{
+			// if we are on any of the pins returning something from the loop
+			if (InRootPin->GetFName() == FRigVMStruct::ExecuteContextName ||
+				InRootPin->GetFName() == FRigVMStruct::ForLoopIndexPinName)
+			{
+				return LoopSliceContextName;
+			}
+			else if (InRootPin->GetFName() == FRigVMStruct::ForLoopCompletedPinName||
+				InRootPin->GetFName() == FRigVMStruct::ForLoopCountPinName ||
+				InRootPin->GetFName() == FRigVMStruct::ForLoopContinuePinName)
+			{
+				return NAME_None;
+			}
+		}
+	}
+
+	if (UScriptStruct* Struct = GetScriptStruct())
+	{
+		if (FProperty* Property = Struct->FindPropertyByName(InRootPin->GetFName()))
+		{
+			FString SliceContextMetaData = Property->GetMetaData(FRigVMStruct::SliceContextMetaName);
+			if (!SliceContextMetaData.IsEmpty())
+			{
+				return *SliceContextMetaData;
+			}
+		}
+	}
+
+	return Super::GetSliceContextForPin(InRootPin, InUserData);
+}
+
+int32 URigVMUnitNode::GetNumSlicesForContext(const FName& InContextName, const FRigVMUserDataArray& InUserData)
+{
+	int32 NumSlices = Super::GetNumSlicesForContext(InContextName, InUserData);
+
+	if (InContextName == LoopSliceContextName)
+	{
+		TSharedPtr<FStructOnScope> StructOnScope = ConstructStructInstance(false);
+		if (StructOnScope.IsValid())
+		{
+			const FRigVMStruct* StructMemory = (FRigVMStruct*)StructOnScope->GetStructMemory();
+			if (StructMemory->IsForLoop())
+			{
+				NumSlices = NumSlices * FMath::Max<int32>(StructMemory->GetNumSlices(), 1);
+			}
+		}
+	}
+
+	return NumSlices;
+}
+
+FText URigVMUnitNode::GetToolTipTextForPin(const URigVMPin* InPin) const
+{
+	if(UScriptStruct* Struct = GetScriptStruct())
+	{
+		TArray<FString> Parts;
+		URigVMPin::SplitPinPath(InPin->GetPinPath(), Parts);
+
+		for (int32 PartIndex = 1; PartIndex < Parts.Num(); PartIndex++)
+		{
+			FProperty* Property = Struct->FindPropertyByName(*Parts[PartIndex]);
+			if (!Property)
+			{
+				break;
+			}
+
+			if (FArrayProperty* ArrayProperty = CastField<FArrayProperty>(Property))
+			{
+				if (PartIndex < Parts.Num() - 1)
+				{
+					Property = ArrayProperty->Inner;
+					PartIndex++;
+				}
+			}
+
+			if (PartIndex == Parts.Num() - 1)
+			{
+				return Property->GetToolTipText();
+			}
+
+			if (FStructProperty* StructProperty = CastField<FStructProperty>(Property))
+			{
+				Struct = StructProperty->Struct;
+			}
+		}
+
+	}
+	return URigVMNode::GetToolTipTextForPin(InPin);
+}
+
+bool URigVMUnitNode::IsDeprecated() const
+{
+	return !GetDeprecatedMetadata().IsEmpty();
+}
+
+FString URigVMUnitNode::GetDeprecatedMetadata() const
+{
+	if (UScriptStruct* Struct = GetScriptStruct())
+	{
+		FString DeprecatedMetadata;
+		if(Struct->GetStringMetaDataHierarchical(FRigVMStruct::DeprecatedMetaName, &DeprecatedMetadata))
+		{
+			return DeprecatedMetadata;
+		}
+	}
+	return FString();
+}
+
+UScriptStruct* URigVMUnitNode::GetScriptStruct() const
+{
+	return ScriptStruct;
+}
+
+bool URigVMUnitNode::IsLoopNode() const
+{
+	TSharedPtr<FStructOnScope> StructOnScope = ConstructStructInstance(true);
+	if (StructOnScope.IsValid())
+	{
+		const FRigVMStruct* StructMemory = (FRigVMStruct*)StructOnScope->GetStructMemory();
+		return StructMemory->IsForLoop();
+	}
+	return false;
+}
+
+FName URigVMUnitNode::GetMethodName() const
+{
+	return MethodName;
+}
+
+FString URigVMUnitNode::GetStructDefaultValue() const
+{
+	TArray<FString> PinDefaultValues;
+	for (URigVMPin* Pin : GetPins())
+	{
+		if (Pin->GetDirection() == ERigVMPinDirection::Hidden)
+		{
+			continue;
+		}
+		FString PinDefaultValue = Pin->GetDefaultValue();
+		if (Pin->IsStringType())
+		{
+			PinDefaultValue = TEXT("\"") + PinDefaultValue + TEXT("\"");
+		}
+		else if (PinDefaultValue.IsEmpty() || PinDefaultValue == TEXT("()"))
+		{
+			continue;
+		}
+		PinDefaultValues.Add(FString::Printf(TEXT("%s=%s"), *Pin->GetName(), *PinDefaultValue));
+	}
+	if (PinDefaultValues.Num() == 0)
+	{
+		return TEXT("()");
+	}
+	return FString::Printf(TEXT("(%s)"), *FString::Join(PinDefaultValues, TEXT(",")));
+
+}
+
+TSharedPtr<FStructOnScope> URigVMUnitNode::ConstructStructInstance(bool bUseDefault) const
+{
+	if (ScriptStruct)
+	{
+		TSharedPtr<FStructOnScope> StructOnScope = MakeShareable(new FStructOnScope(ScriptStruct));
+		FRigVMStruct* StructMemory = (FRigVMStruct*)StructOnScope->GetStructMemory();
+		if (bUseDefault)
+		{
+			ScriptStruct->InitializeDefaultValue((uint8*)StructMemory);
+		}
+		else
+		{
+			FString StructDefaultValue = GetStructDefaultValue();
+			ScriptStruct->ImportText(*StructDefaultValue, StructMemory, nullptr, PPF_None, nullptr, ScriptStruct->GetName());
+		}
+		return StructOnScope;
+	}
+	return nullptr;
+}
diff --git a/Engine/Source/Developer/RigVMDeveloper/Private/RigVMModel/RigVMController.cpp b/Engine/Source/Developer/RigVMDeveloper/Private/RigVMModel/RigVMController.cpp
--- a/Engine/Source/Developer/RigVMDeveloper/Private/RigVMModel/RigVMController.cpp
+++ b/Engine/Source/Developer/RigVMDeveloper/Private/RigVMModel/RigVMController.cpp
@@ -161,7 +161,7 @@ void URigVMController::HandleModifiedEvent(ERigVMGraphNotifType InNotifType, URi
 
 #if WITH_EDITOR
 
-URigVMStructNode* URigVMController::AddStructNode(UScriptStruct* InScriptStruct, const FName& InMethodName, const FVector2D& InPosition, const FString& InNodeName, bool bSetupUndoRedo)
+URigVMUnitNode* URigVMController::AddUnitNode(UScriptStruct* InScriptStruct, const FName& InMethodName, const FVector2D& InPosition, const FString& InNodeName, bool bSetupUndoRedo)
 {
 	if(!IsValidGraph())
 	{
@@ -194,7 +194,7 @@ URigVMStructNode* URigVMController::AddStructNode(UScriptStruct* InScriptStruct,
 	}
 
 	FString Name = GetValidNodeName(InNodeName.IsEmpty() ? InScriptStruct->GetName() : InNodeName);
-	URigVMStructNode* Node = NewObject<URigVMStructNode>(Graph, *Name);
+	URigVMUnitNode* Node = NewObject<URigVMUnitNode>(Graph, *Name);
 	Node->ScriptStruct = InScriptStruct;
 	Node->MethodName = InMethodName;
 	Node->Position = InPosition;
@@ -217,23 +217,23 @@ URigVMStructNode* URigVMController::AddStructNode(UScriptStruct* InScriptStruct,
 		Graph->MarkPackageDirty();
 	}
 
-	FRigVMAddStructNodeAction Action;
+	FRigVMAddUnitNodeAction Action;
 	if (bSetupUndoRedo)
 	{
-		Action = FRigVMAddStructNodeAction(Node);
+		Action = FRigVMAddUnitNodeAction(Node);
 		Action.Title = FString::Printf(TEXT("Add %s Node"), *Node->GetNodeTitle());
 		ActionStack->BeginAction(Action);
 	}
 
 	Notify(ERigVMGraphNotifType::NodeAdded, Node);
 
-	if (StructNodeCreatedContext.IsValid())
+	if (UnitNodeCreatedContext.IsValid())
 	{
 		if (TSharedPtr<FStructOnScope> StructScope = Node->ConstructStructInstance())
 		{
-			TGuardValue<FName> NodeNameScope(StructNodeCreatedContext.NodeName, Node->GetFName());
+			TGuardValue<FName> NodeNameScope(UnitNodeCreatedContext.NodeName, Node->GetFName());
 			FRigVMStruct* StructInstance = (FRigVMStruct*)StructScope->GetStructMemory();
-			StructInstance->OnStructNodeCreated(StructNodeCreatedContext);
+			StructInstance->OnUnitNodeCreated(UnitNodeCreatedContext);
 		}
 	}
 
@@ -245,7 +245,7 @@ URigVMStructNode* URigVMController::AddStructNode(UScriptStruct* InScriptStruct,
 	return Node;
 }
 
-URigVMStructNode* URigVMController::AddStructNodeFromStructPath(const FString& InScriptStructPath, const FName& InMethodName, const FVector2D& InPosition, const FString& InNodeName, bool bSetupUndoRedo)
+URigVMUnitNode* URigVMController::AddUnitNodeFromStructPath(const FString& InScriptStructPath, const FName& InMethodName, const FVector2D& InPosition, const FString& InNodeName, bool bSetupUndoRedo)
 {
 	if(!IsValidGraph())
 	{
@@ -259,7 +259,7 @@ URigVMStructNode* URigVMController::AddStructNodeFromStructPath(const FString& I
 		return nullptr;
 	}
 
-	return AddStructNode(ScriptStruct, InMethodName, InPosition, InNodeName, bSetupUndoRedo);
+	return AddUnitNode(ScriptStruct, InMethodName, InPosition, InNodeName, bSetupUndoRedo);
 }
 
 URigVMVariableNode* URigVMController::AddVariableNode(const FName& InVariableName, const FString& InCPPType, UObject* InCPPTypeObject, bool bIsGetter, const FString& InDefaultValue, const FVector2D& InPosition, const FString& InNodeName, bool bSetupUndoRedo)
@@ -1078,12 +1078,12 @@ URigVMInjectionInfo* URigVMController::AddInjectedNode(const FString& InPinPath,
 		ActionStack->BeginAction(Action);
 	}
 
-	URigVMStructNode* StructNode = nullptr;
+	URigVMUnitNode* UnitNode = nullptr;
 	{
 		TGuardValue<bool> GuardNotifications(bSuspendNotifications, true);
-		StructNode = AddStructNode(InScriptStruct, InMethodName, FVector2D::ZeroVector, InNodeName, false);
+		UnitNode = AddUnitNode(InScriptStruct, InMethodName, FVector2D::ZeroVector, InNodeName, false);
 	}
-	if (StructNode == nullptr)
+	if (UnitNode == nullptr)
 	{
 		if (bSetupUndoRedo)
 		{
@@ -1091,10 +1091,10 @@ URigVMInjectionInfo* URigVMController::AddInjectedNode(const FString& InPinPath,
 		}
 		return nullptr;
 	}
-	else if (StructNode->IsMutable())
+	else if (UnitNode->IsMutable())
 	{
 		ReportErrorf(TEXT("Injected node %s is mutable."), *InScriptStruct->GetName());
-		RemoveNode(StructNode, false);
+		RemoveNode(UnitNode, false);
 		if (bSetupUndoRedo)
 		{
 			ActionStack->CancelAction(Action);
@@ -1102,16 +1102,16 @@ URigVMInjectionInfo* URigVMController::AddInjectedNode(const FString& InPinPath,
 		return nullptr;
 	}
 
-	URigVMPin* InputPin = StructNode->FindPin(InInputPinName.ToString());
+	URigVMPin* InputPin = UnitNode->FindPin(InInputPinName.ToString());
 	check(InputPin);
-	URigVMPin* OutputPin = StructNode->FindPin(InOutputPinName.ToString());
+	URigVMPin* OutputPin = UnitNode->FindPin(InOutputPinName.ToString());
 	check(OutputPin);
 
 	if (InputPin->GetCPPType() != OutputPin->GetCPPType() ||
 		InputPin->IsArray() != OutputPin->IsArray())
 	{
 		ReportErrorf(TEXT("Injected node %s is using incompatible input and output pins."), *InScriptStruct->GetName());
-		RemoveNode(StructNode, false);
+		RemoveNode(UnitNode, false);
 		if (bSetupUndoRedo)
 		{
 			ActionStack->CancelAction(Action);
@@ -1123,7 +1123,7 @@ URigVMInjectionInfo* URigVMController::AddInjectedNode(const FString& InPinPath,
 		InputPin->IsArray() != Pin->IsArray())
 	{
 		ReportErrorf(TEXT("Injected node %s is using incompatible pin."), *InScriptStruct->GetName());
-		RemoveNode(StructNode, false);
+		RemoveNode(UnitNode, false);
 		if (bSetupUndoRedo)
 		{
 			ActionStack->CancelAction(Action);
@@ -1133,10 +1133,10 @@ URigVMInjectionInfo* URigVMController::AddInjectedNode(const FString& InPinPath,
 
 	URigVMInjectionInfo* InjectionInfo = NewObject<URigVMInjectionInfo>(Pin);
 
-	// re-parent the struct node to be under the injection info
-	StructNode->Rename(nullptr, InjectionInfo);
+	// re-parent the unit node to be under the injection info
+	UnitNode->Rename(nullptr, InjectionInfo);
 
-	InjectionInfo->StructNode = StructNode;
+	InjectionInfo->UnitNode = UnitNode;
 	InjectionInfo->bInjectedAsInput = bAsInput;
 	InjectionInfo->InputPin = InputPin;
 	InjectionInfo->OutputPin = OutputPin;
@@ -1155,7 +1155,7 @@ URigVMInjectionInfo* URigVMController::AddInjectedNode(const FString& InPinPath,
 	}
 	Pin->InjectionInfos.Add(InjectionInfo);
 
-	Notify(ERigVMGraphNotifType::NodeAdded, StructNode);
+	Notify(ERigVMGraphNotifType::NodeAdded, UnitNode);
 
 	// now update all of the links
 	if (bAsInput)
@@ -1232,14 +1232,14 @@ URigVMNode* URigVMController::EjectNodeFromPin(const FString& InPinPath, bool bS
 
 	URigVMInjectionInfo* Injection = Pin->InjectionInfos.Last();
 
-	UScriptStruct* ScriptStruct = Injection->StructNode->GetScriptStruct();
-	FName StructNodeName = Injection->StructNode->GetFName();
-	FName MethodName = Injection->StructNode->GetMethodName();
+	UScriptStruct* ScriptStruct = Injection->UnitNode->GetScriptStruct();
+	FName UnitNodeName = Injection->UnitNode->GetFName();
+	FName MethodName = Injection->UnitNode->GetMethodName();
 	FName InputPinName = Injection->InputPin->GetFName();
 	FName OutputPinName = Injection->OutputPin->GetFName();
 
 	TMap<FName, FString> DefaultValues;
-	for (URigVMPin* PinOnNode : Injection->StructNode->GetPins())
+	for (URigVMPin* PinOnNode : Injection->UnitNode->GetPins())
 	{
 		if (PinOnNode->GetDirection() == ERigVMPinDirection::Input ||
 			PinOnNode->GetDirection() == ERigVMPinDirection::Visible ||
@@ -1268,7 +1268,7 @@ URigVMNode* URigVMController::EjectNodeFromPin(const FString& InPinPath, bool bS
 		Position -= FVector2D(250.f, 0.f);
 	}
 
-	URigVMNode* EjectedNode = AddStructNode(ScriptStruct, MethodName, Position, FString(), bSetupUndoRedo);
+	URigVMNode* EjectedNode = AddUnitNode(ScriptStruct, MethodName, Position, FString(), bSetupUndoRedo);
 
 	for (const TPair<FName, FString>& Pair : DefaultValues)
 	{
@@ -1290,9 +1290,9 @@ URigVMNode* URigVMController::EjectNodeFromPin(const FString& InPinPath, bool bS
 		PreviousLink->SourcePin = PreviousLink->TargetPin = nullptr;
 	}
 
-	RemoveNode(Injection->StructNode, bSetupUndoRedo);
+	RemoveNode(Injection->UnitNode, bSetupUndoRedo);
 
-	FString OldNodeNamePrefix = StructNodeName.ToString() + TEXT(".");
+	FString OldNodeNamePrefix = UnitNodeName.ToString() + TEXT(".");
 	FString NewNodeNamePrefix = EjectedNode->GetName() + TEXT(".");
 
 	for (URigVMLink* PreviousLink : PreviousLinks)
@@ -1396,7 +1396,7 @@ FString URigVMController::ExportNodesToText(const TArray<FName>& InNodeNames)
 			{
 				for (URigVMInjectionInfo* Injection : Pin->GetInjectedNodes())
 				{
-					AllNodeNames.AddUnique(Injection->StructNode->GetFName());
+					AllNodeNames.AddUnique(Injection->UnitNode->GetFName());
 				}
 			}
 		}
@@ -1492,13 +1492,13 @@ protected:
 			{
 				for (URigVMInjectionInfo* Injection : Pin->GetInjectedNodes())
 				{
-					ProcessConstructedObject(Injection->StructNode);
+					ProcessConstructedObject(Injection->UnitNode);
 
-					FName NewName = Injection->StructNode->GetFName();
+					FName NewName = Injection->UnitNode->GetFName();
 					UpdateObjectName(URigVMNode::StaticClass(), NewName);
-					Injection->StructNode->Rename(*NewName.ToString(), nullptr);
-					Injection->InputPin = Injection->StructNode->FindPin(Injection->InputPin->GetName());
-					Injection->OutputPin = Injection->StructNode->FindPin(Injection->OutputPin->GetName());
+					Injection->UnitNode->Rename(*NewName.ToString(), nullptr);
+					Injection->InputPin = Injection->UnitNode->FindPin(Injection->InputPin->GetName());
+					Injection->OutputPin = Injection->UnitNode->FindPin(Injection->OutputPin->GetName());
 				}
 			}
 		}
@@ -1546,7 +1546,7 @@ TArray<FName> URigVMController::ImportNodesFromText(const FString& InText, bool
 		ActionStack->BeginAction(AddNodesAction);
 	}
 
-	FRigVMStructNodeCreatedContext::FScope StructNodeCreatedScope(StructNodeCreatedContext, ERigVMNodeCreatedReason::Paste);
+	FRigVMUnitNodeCreatedContext::FScope UnitNodeCreatedScope(UnitNodeCreatedContext, ERigVMNodeCreatedReason::Paste);
 	for (URigVMNode* CreatedNode : Factory.CreatedNodes)
 	{
 		Graph->Nodes.Add(CreatedNode);
@@ -1556,15 +1556,15 @@ TArray<FName> URigVMController::ImportNodesFromText(const FString& InText, bool
 			ActionStack->AddAction(FRigVMRemoveNodeAction(CreatedNode, this));
 		}
 
-		if (URigVMStructNode* StructNode = Cast<URigVMStructNode>(CreatedNode))
+		if (URigVMUnitNode* UnitNode = Cast<URigVMUnitNode>(CreatedNode))
 		{
-			if (StructNodeCreatedContext.IsValid())
+			if (UnitNodeCreatedContext.IsValid())
 			{
-				if (TSharedPtr<FStructOnScope> StructScope = StructNode->ConstructStructInstance())
+				if (TSharedPtr<FStructOnScope> StructScope = UnitNode->ConstructStructInstance())
 				{
-					TGuardValue<FName> NodeNameScope(StructNodeCreatedContext.NodeName, StructNode->GetFName());
+					TGuardValue<FName> NodeNameScope(UnitNodeCreatedContext.NodeName, UnitNode->GetFName());
 					FRigVMStruct* StructInstance = (FRigVMStruct*)StructScope->GetStructMemory();
-					StructInstance->OnStructNodeCreated(StructNodeCreatedContext);
+					StructInstance->OnUnitNodeCreated(UnitNodeCreatedContext);
 				}
 			}
 		}
@@ -2400,7 +2400,7 @@ bool URigVMController::RemoveNode(URigVMNode* InNode, bool bSetupUndoRedo, bool
 			TArray<URigVMInjectionInfo*> InjectedNodes = Pin->GetInjectedNodes();
 			for (URigVMInjectionInfo* InjectedNode : InjectedNodes)
 			{
-				RemoveNode(InjectedNode->StructNode, bSetupUndoRedo);
+				RemoveNode(InjectedNode->UnitNode, bSetupUndoRedo);
 			}
 
 			BreakAllLinks(Pin, true, bSetupUndoRedo);
@@ -3242,9 +3242,9 @@ bool URigVMController::ResetPinDefaultValue(const FString& InPinPath, bool bSetu
 		return false;
 	}
 
-	if (Cast<URigVMStructNode>(Pin->GetNode()) == nullptr)
+	if (Cast<URigVMUnitNode>(Pin->GetNode()) == nullptr)
 	{
-		ReportErrorf(TEXT("Pin '%s' is not on a struct node."), *InPinPath);
+		ReportErrorf(TEXT("Pin '%s' is not on a unit node."), *InPinPath);
 		return false;
 	}
 
@@ -3255,9 +3255,9 @@ bool URigVMController::ResetPinDefaultValue(URigVMPin* InPin, bool bSetupUndoRed
 {
 	check(InPin);
 
-	if (URigVMStructNode* StructNode = Cast<URigVMStructNode>(InPin->GetNode()))
+	if (URigVMUnitNode* UnitNode = Cast<URigVMUnitNode>(InPin->GetNode()))
 	{
-		TSharedPtr<FStructOnScope> StructOnScope = StructNode->ConstructStructInstance(true /* use default */);
+		TSharedPtr<FStructOnScope> StructOnScope = UnitNode->ConstructStructInstance(true /* use default */);
 
 		TArray<FString> Parts;
 		if (!URigVMPin::SplitPinPath(InPin->GetPinPath(), Parts))
@@ -3267,7 +3267,7 @@ bool URigVMController::ResetPinDefaultValue(URigVMPin* InPin, bool bSetupUndoRed
 
 		int32 PartIndex = 1; // cut off the first one since it's the node
 
-		UStruct* Struct = StructNode->ScriptStruct;
+		UStruct* Struct = UnitNode->ScriptStruct;
 		FProperty* Property = Struct->FindPropertyByName(*Parts[PartIndex++]);
 		check(Property);
 
@@ -3935,13 +3935,13 @@ bool URigVMController::PromotePinToVariable(URigVMPin* InPin, bool bCreateVariab
 	}
 	else
 	{
-		if (!StructNodeCreatedContext.GetCreateExternalVariableDelegate().IsBound())
+		if (!UnitNodeCreatedContext.GetCreateExternalVariableDelegate().IsBound())
 		{
 			return false;
 		}
 
 		VariableForPin = InPin->ToExternalVariable();
-		FName VariableName = StructNodeCreatedContext.GetCreateExternalVariableDelegate().Execute(VariableForPin);
+		FName VariableName = UnitNodeCreatedContext.GetCreateExternalVariableDelegate().Execute(VariableForPin);
 		if (VariableName.IsNone())
 		{
 			return false;
@@ -5241,12 +5241,12 @@ FProperty* URigVMController::FindPropertyForPin(const FString& InPinPath)
 
 	URigVMNode* Node = Pin->GetNode();
 
-	URigVMStructNode* StructNode = Cast<URigVMStructNode>(Node);
-	if (StructNode)
+	URigVMUnitNode* UnitNode = Cast<URigVMUnitNode>(Node);
+	if (UnitNode)
 	{
 		int32 PartIndex = 1; // cut off the first one since it's the node
 
-		UStruct* Struct = StructNode->ScriptStruct;
+		UStruct* Struct = UnitNode->ScriptStruct;
 		FProperty* Property = Struct->FindPropertyByName(*Parts[PartIndex++]);
 
 		while (PartIndex < Parts.Num() && Property != nullptr)
@@ -5558,10 +5558,10 @@ bool URigVMController::ShouldRedirectPin(const FString& InOldPinPath, FString& I
 	URigVMPin::SplitPinPathAtStart(InOldPinPath, NodeName, PinPathInNode);
 
 	URigVMNode* Node = Graph->FindNode(NodeName);
-	if (URigVMStructNode* StructNode = Cast<URigVMStructNode>(Node))
+	if (URigVMUnitNode* UnitNode = Cast<URigVMUnitNode>(Node))
 	{
 		FString NewPinPathInNode;
-		if (ShouldRedirectPin(StructNode->GetScriptStruct(), PinPathInNode, NewPinPathInNode))
+		if (ShouldRedirectPin(UnitNode->GetScriptStruct(), PinPathInNode, NewPinPathInNode))
 		{
 			InOutNewPinPath = URigVMPin::JoinPinPath(NodeName, NewPinPathInNode);
 			return true;
@@ -5605,7 +5605,7 @@ void URigVMController::RepopulatePinsOnNode(URigVMNode* InNode, bool bFollowCore
 		return;
 	}
 
-	URigVMStructNode* StructNode = Cast<URigVMStructNode>(InNode);
+	URigVMUnitNode* UnitNode = Cast<URigVMUnitNode>(InNode);
 	URigVMRerouteNode* RerouteNode = Cast<URigVMRerouteNode>(InNode);
 	URigVMFunctionEntryNode* EntryNode = Cast<URigVMFunctionEntryNode>(InNode);
 	URigVMFunctionReturnNode* ReturnNode = Cast<URigVMFunctionReturnNode>(InNode);
@@ -5633,7 +5633,7 @@ void URigVMController::RepopulatePinsOnNode(URigVMNode* InNode, bool bFollowCore
 	}
 
 	// step 2/3: clear pins on the node and repopulate the node with new pins
-	if (StructNode != nullptr)
+	if (UnitNode != nullptr)
 	{
 		TArray<URigVMPin*> Pins = InNode->GetPins();
 		for (URigVMPin* Pin : Pins)
@@ -5643,16 +5643,16 @@ void URigVMController::RepopulatePinsOnNode(URigVMNode* InNode, bool bFollowCore
 		InNode->Pins.Reset();
 		Pins.Reset();
 
-		UScriptStruct* ScriptStruct = StructNode->GetScriptStruct();
+		UScriptStruct* ScriptStruct = UnitNode->GetScriptStruct();
 		if (ScriptStruct == nullptr)
 		{
 			ReportWarningf(
 				TEXT("Control Rig '%s', Node '%s' has no struct assigned. Do you have a broken redirect?"),
-				*StructNode->GetOutermost()->GetPathName(),
-				*StructNode->GetName()
+				*UnitNode->GetOutermost()->GetPathName(),
+				*UnitNode->GetName()
 				);
 
-			RemoveNode(StructNode, false, true);
+			RemoveNode(UnitNode, false, true);
 			return;
 		}
 
@@ -5660,12 +5660,12 @@ void URigVMController::RepopulatePinsOnNode(URigVMNode* InNode, bool bFollowCore
 		ScriptStruct->GetStringMetaDataHierarchical(*URigVMNode::NodeColorName, &NodeColorMetadata);
 		if (!NodeColorMetadata.IsEmpty())
 		{
-			StructNode->NodeColor = GetColorFromMetadata(NodeColorMetadata);
+			UnitNode->NodeColor = GetColorFromMetadata(NodeColorMetadata);
 		}
 
 		FString ExportedDefaultValue;
 		CreateDefaultValueForStructIfRequired(ScriptStruct, ExportedDefaultValue);
-		AddPinsForStruct(ScriptStruct, StructNode, nullptr, ERigVMPinDirection::Invalid, ExportedDefaultValue, false);
+		AddPinsForStruct(ScriptStruct, UnitNode, nullptr, ERigVMPinDirection::Invalid, ExportedDefaultValue, false);
 	}
 	else if (RerouteNode != nullptr)
 	{
@@ -5783,11 +5783,11 @@ void URigVMController::RepopulatePinsOnNode(URigVMNode* InNode, bool bFollowCore
 
 #endif
 
-void URigVMController::SetupDefaultStructNodeDelegates(TDelegate<FName(FRigVMExternalVariable)> InCreateExternalVariableDelegate)
+void URigVMController::SetupDefaultUnitNodeDelegates(TDelegate<FName(FRigVMExternalVariable)> InCreateExternalVariableDelegate)
 {
 	TWeakObjectPtr<URigVMController> WeakThis(this);
 
-	StructNodeCreatedContext.GetAllExternalVariablesDelegate().BindLambda([WeakThis]() -> TArray<FRigVMExternalVariable> {
+	UnitNodeCreatedContext.GetAllExternalVariablesDelegate().BindLambda([WeakThis]() -> TArray<FRigVMExternalVariable> {
 		if (WeakThis.IsValid())
 		{
 			return WeakThis->GetExternalVariables();
@@ -5795,7 +5795,7 @@ void URigVMController::SetupDefaultStructNodeDelegates(TDelegate<FName(FRigVMExt
 		return TArray<FRigVMExternalVariable>();
 	});
 
-	StructNodeCreatedContext.GetBindPinToExternalVariableDelegate().BindLambda([WeakThis](FString InPinPath, FString InVariablePath) -> bool {
+	UnitNodeCreatedContext.GetBindPinToExternalVariableDelegate().BindLambda([WeakThis](FString InPinPath, FString InVariablePath) -> bool {
 		if (WeakThis.IsValid())
 		{
 			return WeakThis->BindPinToVariable(InPinPath, InVariablePath, true);
@@ -5803,14 +5803,14 @@ void URigVMController::SetupDefaultStructNodeDelegates(TDelegate<FName(FRigVMExt
 		return false;
 	});
 
-	StructNodeCreatedContext.GetCreateExternalVariableDelegate() = InCreateExternalVariableDelegate;
+	UnitNodeCreatedContext.GetCreateExternalVariableDelegate() = InCreateExternalVariableDelegate;
 }
 
-void URigVMController::ResetStructNodeDelegates()
+void URigVMController::ResetUnitNodeDelegates()
 {
-	StructNodeCreatedContext.GetAllExternalVariablesDelegate().Unbind();
-	StructNodeCreatedContext.GetBindPinToExternalVariableDelegate().Unbind();
-	StructNodeCreatedContext.GetCreateExternalVariableDelegate().Unbind();
+	UnitNodeCreatedContext.GetAllExternalVariablesDelegate().Unbind();
+	UnitNodeCreatedContext.GetBindPinToExternalVariableDelegate().Unbind();
+	UnitNodeCreatedContext.GetCreateExternalVariableDelegate().Unbind();
 }
 
 FLinearColor URigVMController::GetColorFromMetadata(const FString& InMetadata)
@@ -5843,13 +5843,13 @@ FLinearColor URigVMController::GetColorFromMetadata(const FString& InMetadata)
 TMap<FString, FString> URigVMController::GetRedirectedPinPaths(URigVMNode* InNode) const
 {
 	TMap<FString, FString> RedirectedPinPaths;
-	URigVMStructNode* StructNode = Cast<URigVMStructNode>(InNode);
+	URigVMUnitNode* UnitNode = Cast<URigVMUnitNode>(InNode);
 	URigVMRerouteNode* RerouteNode = Cast<URigVMRerouteNode>(InNode);
 
 	UScriptStruct* OwningStruct = nullptr;
-	if (StructNode)
+	if (UnitNode)
 	{
-		OwningStruct = StructNode->GetScriptStruct();
+		OwningStruct = UnitNode->GetScriptStruct();
 	}
 	else if (RerouteNode)
 	{
@@ -5926,8 +5926,8 @@ void URigVMController::ApplyPinState(URigVMPin* InPin, const FPinState& InPinSta
 	for (URigVMInjectionInfo* InjectionInfo : InPinState.InjectionInfos)
 	{
 		InjectionInfo->Rename(nullptr, InPin);
-		InjectionInfo->InputPin = InjectionInfo->StructNode->FindPin(InjectionInfo->InputPin->GetName());
-		InjectionInfo->OutputPin = InjectionInfo->StructNode->FindPin(InjectionInfo->OutputPin->GetName());
+		InjectionInfo->InputPin = InjectionInfo->UnitNode->FindPin(InjectionInfo->InputPin->GetName());
+		InjectionInfo->OutputPin = InjectionInfo->UnitNode->FindPin(InjectionInfo->OutputPin->GetName());
 		InPin->InjectionInfos.Add(InjectionInfo);
 	}
 
@@ -5960,7 +5960,7 @@ void URigVMController::ApplyPinStates(URigVMNode* InNode, const TMap<FString, UR
 		{
 			for (URigVMInjectionInfo* InjectionInfo : PinState.InjectionInfos)
 			{
-				InjectionInfo->StructNode->Rename(nullptr, InNode->GetGraph());
+				InjectionInfo->UnitNode->Rename(nullptr, InNode->GetGraph());
 				DestroyObject(InjectionInfo);
 			}
 		}
@@ -6111,7 +6111,7 @@ void URigVMController::PotentiallyResolvePrototypeNode(URigVMPrototypeNode* InNo
 
 		RemoveNode(InNode, bSetupUndoRedo);
 
-		if (URigVMNode* NewNode = AddStructNode(Function.Struct, Function.GetMethodName(), NodePosition, NodeName, bSetupUndoRedo))
+		if (URigVMNode* NewNode = AddUnitNode(Function.Struct, Function.GetMethodName(), NodePosition, NodeName, bSetupUndoRedo))
 		{
 			// set default values again
 			for (TPair<FString, FString> Pair : DefaultValues)
diff --git a/Engine/Source/Developer/RigVMDeveloper/Private/RigVMModel/RigVMControllerActions.cpp b/Engine/Source/Developer/RigVMDeveloper/Private/RigVMModel/RigVMControllerActions.cpp
--- a/Engine/Source/Developer/RigVMDeveloper/Private/RigVMModel/RigVMControllerActions.cpp
+++ b/Engine/Source/Developer/RigVMDeveloper/Private/RigVMModel/RigVMControllerActions.cpp
@@ -201,7 +201,7 @@ bool FRigVMInverseAction::Redo(URigVMController* InController)
 	return FRigVMBaseAction::Undo(InController);
 }
 
-FRigVMAddStructNodeAction::FRigVMAddStructNodeAction()
+FRigVMAddUnitNodeAction::FRigVMAddUnitNodeAction()
 	: ScriptStructPath()
 	, MethodName(NAME_None)
 	, Position(FVector2D::ZeroVector)
@@ -209,7 +209,7 @@ FRigVMAddStructNodeAction::FRigVMAddStructNodeAction()
 {
 }
 
-FRigVMAddStructNodeAction::FRigVMAddStructNodeAction(URigVMStructNode* InNode)
+FRigVMAddUnitNodeAction::FRigVMAddUnitNodeAction(URigVMUnitNode* InNode)
 	: ScriptStructPath(InNode->GetScriptStruct()->GetPathName())
 	, MethodName(InNode->GetMethodName())
 	, Position(InNode->GetPosition())
@@ -217,7 +217,7 @@ FRigVMAddStructNodeAction::FRigVMAddStructNodeAction(URigVMStructNode* InNode)
 {
 }
 
-bool FRigVMAddStructNodeAction::Undo(URigVMController* InController)
+bool FRigVMAddUnitNodeAction::Undo(URigVMController* InController)
 {
 	if (!FRigVMBaseAction::Undo(InController))
 	{
@@ -226,10 +226,10 @@ bool FRigVMAddStructNodeAction::Undo(URigVMController* InController)
 	return InController->RemoveNodeByName(*NodePath, false);
 }
 
-bool FRigVMAddStructNodeAction::Redo(URigVMController* InController)
+bool FRigVMAddUnitNodeAction::Redo(URigVMController* InController)
 {
 #if WITH_EDITOR
-	if (URigVMStructNode* Node = InController->AddStructNodeFromStructPath(ScriptStructPath, MethodName, Position, NodePath, false))
+	if (URigVMUnitNode* Node = InController->AddUnitNodeFromStructPath(ScriptStructPath, MethodName, Position, NodePath, false))
 	{
 		return FRigVMBaseAction::Redo(InController);
 	}
@@ -618,11 +618,11 @@ FRigVMAddInjectedNodeAction::FRigVMAddInjectedNodeAction()
 FRigVMAddInjectedNodeAction::FRigVMAddInjectedNodeAction(URigVMInjectionInfo* InInjectionInfo)
 	: PinPath(InInjectionInfo->GetPin()->GetPinPath())
 	, bAsInput(InInjectionInfo->bInjectedAsInput)
-	, ScriptStructPath(InInjectionInfo->StructNode->GetScriptStruct()->GetPathName())
-	, MethodName(InInjectionInfo->StructNode->GetMethodName())
+	, ScriptStructPath(InInjectionInfo->UnitNode->GetScriptStruct()->GetPathName())
+	, MethodName(InInjectionInfo->UnitNode->GetMethodName())
 	, InputPinName(InInjectionInfo->InputPin->GetFName())
 	, OutputPinName(InInjectionInfo->OutputPin->GetFName())
-	, NodePath(InInjectionInfo->StructNode->GetName())
+	, NodePath(InInjectionInfo->UnitNode->GetName())
 {
 }
 
@@ -662,10 +662,10 @@ FRigVMRemoveNodeAction::FRigVMRemoveNodeAction(URigVMNode* InNode, URigVMControl
 			}
 		}
 	}
-	else if (URigVMStructNode* StructNode = Cast<URigVMStructNode>(InNode))
+	else if (URigVMUnitNode* UnitNode = Cast<URigVMUnitNode>(InNode))
 	{
-		InverseAction.AddAction(FRigVMAddStructNodeAction(StructNode));
-		for (URigVMPin* Pin : StructNode->GetPins())
+		InverseAction.AddAction(FRigVMAddUnitNodeAction(UnitNode));
+		for (URigVMPin* Pin : UnitNode->GetPins())
 		{
 			if (Pin->GetDirection() == ERigVMPinDirection::Input ||
 				Pin->GetDirection() == ERigVMPinDirection::Visible)
diff --git a/Engine/Source/Developer/RigVMDeveloper/Private/RigVMModel/RigVMPin.cpp b/Engine/Source/Developer/RigVMDeveloper/Private/RigVMModel/RigVMPin.cpp
--- a/Engine/Source/Developer/RigVMDeveloper/Private/RigVMModel/RigVMPin.cpp
+++ b/Engine/Source/Developer/RigVMDeveloper/Private/RigVMModel/RigVMPin.cpp
@@ -140,7 +140,7 @@ FName URigVMPin::GetDisplayName() const
 
 		for (URigVMInjectionInfo* Injection : InjectionInfos)
 		{
-			if (TSharedPtr<FStructOnScope> DefaultStructScope = Injection->StructNode->ConstructStructInstance())
+			if (TSharedPtr<FStructOnScope> DefaultStructScope = Injection->UnitNode->ConstructStructInstance())
 			{
 				FRigVMStruct* DefaultStruct = (FRigVMStruct*)DefaultStructScope->GetStructMemory();
 				ProcessedDisplayName = DefaultStruct->ProcessPinLabelForInjection(ProcessedDisplayName);
@@ -492,11 +492,11 @@ int32 URigVMPin::GetNumSlices(const FRigVMUserDataArray& InUserData)
 bool URigVMPin::ShowInDetailsPanelOnly() const
 {
 #if WITH_EDITOR
-	if (URigVMStructNode* StructNode = Cast<URigVMStructNode>(GetNode()))
+	if (URigVMUnitNode* UnitNode = Cast<URigVMUnitNode>(GetNode()))
 	{
 		if (GetParentPin() == nullptr)
 		{
-			if (UScriptStruct* ScriptStruct = StructNode->GetScriptStruct())
+			if (UScriptStruct* ScriptStruct = UnitNode->GetScriptStruct())
 			{
 				if (FProperty* Property = ScriptStruct->FindPropertyByName(GetFName()))
 				{
@@ -1001,11 +1001,11 @@ bool URigVMPin::CanLink(URigVMPin* InSourcePin, URigVMPin* InTargetPin, FString*
 			bool bNodeCanLinkAnywhere = SourceNodes[SourceNodeIndex]->IsA<URigVMRerouteNode>();
 			if (!bNodeCanLinkAnywhere)
 			{
-				if (URigVMStructNode* StructNode = Cast<URigVMStructNode>(SourceNodes[SourceNodeIndex]))
+				if (URigVMUnitNode* UnitNode = Cast<URigVMUnitNode>(SourceNodes[SourceNodeIndex]))
 				{
 					// pure / immutable nodes can be connected to any input in any order.
 					// since a new link is going to change the abstract syntax tree 
-					if (!StructNode->IsMutable())
+					if (!UnitNode->IsMutable())
 					{
 						bNodeCanLinkAnywhere = true;
 					}
diff --git a/Engine/Source/Developer/RigVMDeveloper/Public/RigVMCompiler/RigVMCompiler.h b/Engine/Source/Developer/RigVMDeveloper/Public/RigVMCompiler/RigVMCompiler.h
--- a/Engine/Source/Developer/RigVMDeveloper/Public/RigVMCompiler/RigVMCompiler.h
+++ b/Engine/Source/Developer/RigVMDeveloper/Public/RigVMCompiler/RigVMCompiler.h
@@ -5,7 +5,7 @@
 #include "CoreMinimal.h"
 
 #include "RigVMModel/RigVMGraph.h"
-#include "RigVMModel/Nodes/RigVMStructNode.h"
+#include "RigVMModel/Nodes/RigVMUnitNode.h"
 #include "RigVMModel/Nodes/RigVMVariableNode.h"
 #include "RigVMModel/Nodes/RigVMParameterNode.h"
 #include "RigVMModel/Nodes/RigVMCommentNode.h"
diff --git a/Engine/Source/Developer/RigVMDeveloper/Public/RigVMModel/Nodes/RigVMStructNode.h b/Engine/Source/Developer/RigVMDeveloper/Public/RigVMModel/Nodes/RigVMStructNode.h
deleted file mode 100644
diff --git a/Engine/Source/Developer/RigVMDeveloper/Public/RigVMModel/Nodes/RigVMUnitNode.h b/Engine/Source/Developer/RigVMDeveloper/Public/RigVMModel/Nodes/RigVMUnitNode.h
new file mode 100644
--- /dev/null
+++ b/Engine/Source/Developer/RigVMDeveloper/Public/RigVMModel/Nodes/RigVMUnitNode.h
@@ -0,0 +1,70 @@
+// Copyright Epic Games, Inc. All Rights Reserved.
+
+#pragma once
+
+#include "RigVMModel/RigVMNode.h"
+#include "RigVMCore/RigVMStruct.h"
+#include "UObject/StructOnScope.h"
+#include "RigVMUnitNode.generated.h"
+
+/**
+ * The Struct Node represents a Function Invocation of a RIGVM_METHOD
+ * declared on a USTRUCT. Struct Nodes have input / output pins for all
+ * struct UPROPERTY members.
+ */
+UCLASS(BlueprintType)
+class RIGVMDEVELOPER_API URigVMUnitNode : public URigVMNode
+{
+	GENERATED_BODY()
+
+public:
+
+	// Override node functions
+	virtual FString GetNodeTitle() const override;
+	virtual FText GetToolTipText() const override;
+	virtual bool IsDefinedAsConstant() const override;
+	virtual bool IsDefinedAsVarying() const override;
+	virtual FName GetEventName() const override;
+	virtual FName GetSliceContextForPin(URigVMPin* InRootPin, const FRigVMUserDataArray& InUserData) override;
+	virtual int32 GetNumSlicesForContext(const FName& InContextName, const FRigVMUserDataArray& InUserData) override;
+
+	bool IsDeprecated() const;
+	FString GetDeprecatedMetadata() const;
+
+	// Returns the UStruct for this unit node
+	// (the struct declaring the RIGVM_METHOD)
+	UFUNCTION(BlueprintCallable, Category = RigVMUnitNode)
+	UScriptStruct* GetScriptStruct() const;
+
+	// return true if this node is a loop node
+	bool IsLoopNode() const;
+
+	// Returns the name of the declared RIGVM_METHOD
+	UFUNCTION(BlueprintCallable, Category = RigVMUnitNode)
+	FName GetMethodName() const;
+
+	// Returns the default value for the struct as text
+	UFUNCTION(BlueprintCallable, Category = RigVMUnitNode)
+	FString GetStructDefaultValue() const;
+
+	// Returns an instance of the struct with the current values.
+	// @param bUseDefault If set to true the default struct will be created - otherwise the struct will contains the values from the node
+	TSharedPtr<FStructOnScope> ConstructStructInstance(bool bUseDefault = false) const;
+
+	static const FName LoopSliceContextName;
+
+protected:
+
+	virtual FText GetToolTipTextForPin(const URigVMPin* InPin) const override;
+
+private:
+
+	UPROPERTY()
+	UScriptStruct* ScriptStruct;
+
+	UPROPERTY()
+	FName MethodName;
+
+	friend class URigVMController;
+};
+
diff --git a/Engine/Source/Developer/RigVMDeveloper/Public/RigVMModel/RigVMController.h b/Engine/Source/Developer/RigVMDeveloper/Public/RigVMModel/RigVMController.h
--- a/Engine/Source/Developer/RigVMDeveloper/Public/RigVMModel/RigVMController.h
+++ b/Engine/Source/Developer/RigVMDeveloper/Public/RigVMModel/RigVMController.h
@@ -3,7 +3,7 @@
 #pragma once
 
 #include "RigVMGraph.h"
-#include "RigVMModel/Nodes/RigVMStructNode.h"
+#include "RigVMModel/Nodes/RigVMUnitNode.h"
 #include "RigVMModel/Nodes/RigVMVariableNode.h"
 #include "RigVMModel/Nodes/RigVMParameterNode.h"
 #include "RigVMModel/Nodes/RigVMCommentNode.h"
@@ -82,16 +82,16 @@ public:
 	// on the metadata of the USTRUCT - which is only available in the editor.
 
 	// Adds a Function / Struct Node to the edited Graph.
-	// StructNode represent a RIGVM_METHOD declaration on a USTRUCT.
+	// UnitNode represent a RIGVM_METHOD declaration on a USTRUCT.
 	// This causes a NodeAdded modified event.
 	UFUNCTION(BlueprintCallable, Category = RigVMController)
-	URigVMStructNode* AddStructNode(UScriptStruct* InScriptStruct, const FName& InMethodName, const FVector2D& InPosition = FVector2D::ZeroVector, const FString& InNodeName = TEXT(""), bool bSetupUndoRedo = true);
+	URigVMUnitNode* AddUnitNode(UScriptStruct* InScriptStruct, const FName& InMethodName, const FVector2D& InPosition = FVector2D::ZeroVector, const FString& InNodeName = TEXT(""), bool bSetupUndoRedo = true);
 
 	// Adds a Function / Struct Node to the edited Graph given its struct object path name.
-	// StructNode represent a RIGVM_METHOD declaration on a USTRUCT.
+	// UnitNode represent a RIGVM_METHOD declaration on a USTRUCT.
 	// This causes a NodeAdded modified event.
 	UFUNCTION(BlueprintCallable, Category = RigVMController)
-	URigVMStructNode* AddStructNodeFromStructPath(const FString& InScriptStructPath, const FName& InMethodName, const FVector2D& InPosition = FVector2D::ZeroVector, const FString& InNodeName = TEXT(""), bool bSetupUndoRedo = true);
+	URigVMUnitNode* AddUnitNodeFromStructPath(const FString& InScriptStructPath, const FName& InMethodName, const FVector2D& InPosition = FVector2D::ZeroVector, const FString& InNodeName = TEXT(""), bool bSetupUndoRedo = true);
 
 	// Adds a Variable Node to the edited Graph.
 	// Variables represent local work state for the function and
@@ -193,13 +193,13 @@ public:
 	URigVMPrototypeNode* AddPrototypeNode(const FName& InNotation, const FVector2D& InPosition = FVector2D::ZeroVector, const FString& InNodeName = TEXT(""), bool bSetupUndoRedo = true);
 
 	// Adds a Function / Struct Node to the edited Graph as an injected node
-	// StructNode represent a RIGVM_METHOD declaration on a USTRUCT.
+	// UnitNode represent a RIGVM_METHOD declaration on a USTRUCT.
 	// This causes a NodeAdded modified event.
 	UFUNCTION(BlueprintCallable, Category = RigVMController)
 	URigVMInjectionInfo* AddInjectedNode(const FString& InPinPath, bool bAsInput, UScriptStruct* InScriptStruct, const FName& InMethodName, const FName& InInputPinName, const FName& InOutputPinName, const FString& InNodeName = TEXT(""), bool bSetupUndoRedo = true);
 
 	// Adds a Function / Struct Node to the edited Graph as an injected node
-	// StructNode represent a RIGVM_METHOD declaration on a USTRUCT.
+	// UnitNode represent a RIGVM_METHOD declaration on a USTRUCT.
 	// This causes a NodeAdded modified event.
 	UFUNCTION(BlueprintCallable, Category = RigVMController)
 	URigVMInjectionInfo* AddInjectedNodeFromStructPath(const FString& InPinPath, bool bAsInput, const FString& InScriptStructPath, const FName& InMethodName, const FName& InInputPinName, const FName& InOutputPinName, const FString& InNodeName = TEXT(""), bool bSetupUndoRedo = true);
@@ -478,12 +478,12 @@ public:
 	void RepopulatePinsOnNode(URigVMNode* InNode, bool bFollowCoreRedirectors = true);
 #endif
 
-	FRigVMStructNodeCreatedContext& GetStructNodeCreatedContext() { return StructNodeCreatedContext; }
+	FRigVMUnitNodeCreatedContext& GetUnitNodeCreatedContext() { return UnitNodeCreatedContext; }
 
-	// Wires the struct node delegates to the default controller delegates.
+	// Wires the unit node delegates to the default controller delegates.
 	// this is used only within the Control Rig Editor currently.
-	void SetupDefaultStructNodeDelegates(TDelegate<FName(FRigVMExternalVariable)> InCreateExternalVariableDelegate);
-	void ResetStructNodeDelegates();
+	void SetupDefaultUnitNodeDelegates(TDelegate<FName(FRigVMExternalVariable)> InCreateExternalVariableDelegate);
+	void ResetUnitNodeDelegates();
 
 private:
 
@@ -628,7 +628,7 @@ private:
 	static TMap<FControlRigStructPinRedirectorKey, FString> PinPathCoreRedirectors;
 	FCriticalSection PinPathCoreRedirectorsLock;
 
-	FRigVMStructNodeCreatedContext StructNodeCreatedContext;
+	FRigVMUnitNodeCreatedContext UnitNodeCreatedContext;
 
 	friend class URigVMGraph;
 	friend class URigVMActionStack;
diff --git a/Engine/Source/Developer/RigVMDeveloper/Public/RigVMModel/RigVMControllerActions.h b/Engine/Source/Developer/RigVMDeveloper/Public/RigVMModel/RigVMControllerActions.h
--- a/Engine/Source/Developer/RigVMDeveloper/Public/RigVMModel/RigVMControllerActions.h
+++ b/Engine/Source/Developer/RigVMDeveloper/Public/RigVMModel/RigVMControllerActions.h
@@ -295,18 +295,18 @@ struct FRigVMInverseAction : public FRigVMBaseAction
 };
 
 /**
- * An action adding a struct node to the graph.
+ * An action adding a unit node to the graph.
  */
 USTRUCT()
-struct FRigVMAddStructNodeAction : public FRigVMBaseAction
+struct FRigVMAddUnitNodeAction : public FRigVMBaseAction
 {
 	GENERATED_BODY()
 
 public:
 
-	FRigVMAddStructNodeAction();
-	FRigVMAddStructNodeAction(URigVMStructNode* InNode);
-	virtual ~FRigVMAddStructNodeAction() {};
+	FRigVMAddUnitNodeAction();
+	FRigVMAddUnitNodeAction(URigVMUnitNode* InNode);
+	virtual ~FRigVMAddUnitNodeAction() {};
 	virtual bool Undo(URigVMController* InController) override;
 	virtual bool Redo(URigVMController* InController) override;
 
diff --git a/Engine/Source/Developer/RigVMDeveloper/Public/RigVMModel/RigVMPin.h b/Engine/Source/Developer/RigVMDeveloper/Public/RigVMModel/RigVMPin.h
--- a/Engine/Source/Developer/RigVMDeveloper/Public/RigVMModel/RigVMPin.h
+++ b/Engine/Source/Developer/RigVMDeveloper/Public/RigVMModel/RigVMPin.h
@@ -9,7 +9,7 @@
 
 class URigVMGraph;
 class URigVMNode;
-class URigVMStructNode;
+class URigVMUnitNode;
 class URigVMPin;
 class URigVMLink;
 
@@ -31,7 +31,7 @@ public:
 	}
 
 	UPROPERTY()
-	URigVMStructNode* StructNode;
+	URigVMUnitNode* UnitNode;
 
 	UPROPERTY()
 	bool bInjectedAsInput;
diff --git a/Engine/Source/Runtime/RigVM/Private/RigVMCore/RigVMStruct.cpp b/Engine/Source/Runtime/RigVM/Private/RigVMCore/RigVMStruct.cpp
--- a/Engine/Source/Runtime/RigVM/Private/RigVMCore/RigVMStruct.cpp
+++ b/Engine/Source/Runtime/RigVM/Private/RigVMCore/RigVMStruct.cpp
@@ -5,7 +5,7 @@
 
 ////////////////////////////////////////////////////////////////////////////////////////////////////////////////////////
 
-bool FRigVMStructNodeCreatedContext::IsValid() const
+bool FRigVMUnitNodeCreatedContext::IsValid() const
 {
 	return Reason != ERigVMNodeCreatedReason::Unknown &&
 		AllExternalVariablesDelegate.IsBound() &&
@@ -13,7 +13,7 @@ bool FRigVMStructNodeCreatedContext::IsValid() const
 		BindPinToExternalVariableDelegate.IsBound();
 }
 
-TArray<FRigVMExternalVariable> FRigVMStructNodeCreatedContext::GetExternalVariables() const
+TArray<FRigVMExternalVariable> FRigVMUnitNodeCreatedContext::GetExternalVariables() const
 {
 	TArray<FRigVMExternalVariable> ExternalVariables;
 
@@ -25,7 +25,7 @@ TArray<FRigVMExternalVariable> FRigVMStructNodeCreatedContext::GetExternalVariab
 	return ExternalVariables;
 }
 
-FName FRigVMStructNodeCreatedContext::AddExternalVariable(const FRigVMExternalVariable& InVariableToCreate)
+FName FRigVMUnitNodeCreatedContext::AddExternalVariable(const FRigVMExternalVariable& InVariableToCreate)
 {
 	if (CreateExternalVariableDelegate.IsBound())
 	{
@@ -34,7 +34,7 @@ FName FRigVMStructNodeCreatedContext::AddExternalVariable(const FRigVMExternalVa
 	return NAME_None;
 }
 
-bool FRigVMStructNodeCreatedContext::BindPinToExternalVariable(FString InPinPath, FString InVariablePath)
+bool FRigVMUnitNodeCreatedContext::BindPinToExternalVariable(FString InPinPath, FString InVariablePath)
 {
 	if (BindPinToExternalVariableDelegate.IsBound())
 	{
@@ -44,7 +44,7 @@ bool FRigVMStructNodeCreatedContext::BindPinToExternalVariable(FString InPinPath
 	return false;
 }
 
-FRigVMExternalVariable FRigVMStructNodeCreatedContext::FindVariable(FName InVariableName) const
+FRigVMExternalVariable FRigVMUnitNodeCreatedContext::FindVariable(FName InVariableName) const
 {
 	TArray<FRigVMExternalVariable> ExternalVariables = GetExternalVariables();
 	for (FRigVMExternalVariable ExternalVariable : ExternalVariables)
@@ -57,7 +57,7 @@ FRigVMExternalVariable FRigVMStructNodeCreatedContext::FindVariable(FName InVari
 	return FRigVMExternalVariable();
 }
 
-FName FRigVMStructNodeCreatedContext::FindFirstVariableOfType(FName InCPPTypeName) const
+FName FRigVMUnitNodeCreatedContext::FindFirstVariableOfType(FName InCPPTypeName) const
 {
 	TArray<FRigVMExternalVariable> ExternalVariables = GetExternalVariables();
 	for (FRigVMExternalVariable ExternalVariable : ExternalVariables)
@@ -70,7 +70,7 @@ FName FRigVMStructNodeCreatedContext::FindFirstVariableOfType(FName InCPPTypeNam
 	return NAME_None;
 }
 
-FName FRigVMStructNodeCreatedContext::FindFirstVariableOfType(UObject* InCPPTypeObject) const
+FName FRigVMUnitNodeCreatedContext::FindFirstVariableOfType(UObject* InCPPTypeObject) const
 {
 	TArray<FRigVMExternalVariable> ExternalVariables = GetExternalVariables();
 	for (FRigVMExternalVariable ExternalVariable : ExternalVariables)
diff --git a/Engine/Source/Runtime/RigVM/Public/RigVMCore/RigVMStruct.h b/Engine/Source/Runtime/RigVM/Public/RigVMCore/RigVMStruct.h
--- a/Engine/Source/Runtime/RigVM/Public/RigVMCore/RigVMStruct.h
+++ b/Engine/Source/Runtime/RigVM/Public/RigVMCore/RigVMStruct.h
@@ -24,16 +24,16 @@ enum class ERigVMNodeCreatedReason : uint8
 };
 
 /**
- * A context struct passed to FRigVMStruct::OnStructNodeCreated
+ * A context struct passed to FRigVMStruct::OnUnitNodeCreated
  */
-struct RIGVM_API FRigVMStructNodeCreatedContext
+struct RIGVM_API FRigVMUnitNodeCreatedContext
 {
 public:
 
 	struct FScope
 	{
 	public:
-		FScope(FRigVMStructNodeCreatedContext& InContext, ERigVMNodeCreatedReason InReason)
+		FScope(FRigVMUnitNodeCreatedContext& InContext, ERigVMNodeCreatedReason InReason)
 			: Context(InContext)
 			, PreviousReason(InContext.GetReason())
 		{
@@ -46,7 +46,7 @@ public:
 		}
 
 	private:
-		FRigVMStructNodeCreatedContext& Context;
+		FRigVMUnitNodeCreatedContext& Context;
 		ERigVMNodeCreatedReason PreviousReason;
 	};
 
@@ -153,7 +153,7 @@ public:
 	FORCEINLINE virtual int32 GetNumSlices() const { return 1; }
 
 	// node creation
-	FORCEINLINE virtual void OnStructNodeCreated(FRigVMStructNodeCreatedContext& InContext) const {}
+	FORCEINLINE virtual void OnUnitNodeCreated(FRigVMUnitNodeCreatedContext& InContext) const {}
 
 #if WITH_EDITOR
 
