Control Rig: Collapse + Expand working, milestone 2

diff --git a/Engine/Source/Developer/RigVMDeveloper/Private/RigVMCompiler/RigVMCompiler.cpp b/Engine/Source/Developer/RigVMDeveloper/Private/RigVMCompiler/RigVMCompiler.cpp
--- a/Engine/Source/Developer/RigVMDeveloper/Private/RigVMCompiler/RigVMCompiler.cpp
+++ b/Engine/Source/Developer/RigVMDeveloper/Private/RigVMCompiler/RigVMCompiler.cpp
@@ -1127,6 +1127,10 @@ FString URigVMCompiler::GetPinHash(URigVMPin* InPin, const FRigVMVarExprAST* InV
 		{
 			Suffix = TEXT("::IO");
 		}
+		else if (bIsLiteral)
+		{
+			Suffix = TEXT("::Const");
+		}
 	}
 
 	if (URigVMParameterNode* ParameterNode = Cast<URigVMParameterNode>(Node))
diff --git a/Engine/Source/Developer/RigVMDeveloper/Private/RigVMModel/Nodes/RigVMLibraryNode.cpp b/Engine/Source/Developer/RigVMDeveloper/Private/RigVMModel/Nodes/RigVMLibraryNode.cpp
--- a/Engine/Source/Developer/RigVMDeveloper/Private/RigVMModel/Nodes/RigVMLibraryNode.cpp
+++ b/Engine/Source/Developer/RigVMDeveloper/Private/RigVMModel/Nodes/RigVMLibraryNode.cpp
@@ -5,32 +5,46 @@
 #include "RigVMModel/Nodes/RigVMFunctionReturnNode.h"
 #include "RigVMModel/RigVMGraph.h"
 
+const TArray<URigVMNode*> URigVMLibraryNode::EmptyNodes;
+const TArray<URigVMLink*> URigVMLibraryNode::EmptyLinks;
+
 bool URigVMLibraryNode::IsDefinedAsConstant() const
 {
-	return !bDefinedAsVarying;
+	return !IsDefinedAsVarying();
 }
 
 bool URigVMLibraryNode::IsDefinedAsVarying() const
 {
-	return bDefinedAsVarying;
+	if (URigVMGraph* Graph = GetContainedGraph())
+	{
+		const TArray<URigVMNode*>& Nodes = Graph->GetNodes();
+		for(URigVMNode* Node : Nodes)
+		{
+			if (Node->IsDefinedAsVarying())
+			{
+				return true;
+			}
+		}
+	}
+	return false;
 }
 
-TArray<URigVMNode*> URigVMLibraryNode::GetContainedNodes() const
+const TArray<URigVMNode*>& URigVMLibraryNode::GetContainedNodes() const
 {
 	if(URigVMGraph* Graph = GetContainedGraph())
 	{
 		return Graph->GetNodes();
 	}
-	return TArray<URigVMNode*>();
+	return EmptyNodes;
 }
 
-TArray<URigVMLink*> URigVMLibraryNode::GetContainedLinks() const
+const TArray<URigVMLink*>& URigVMLibraryNode::GetContainedLinks() const
 {
 	if (URigVMGraph* Graph = GetContainedGraph())
 	{
 		return Graph->GetLinks();
 	}
-	return TArray<URigVMLink*>();
+	return EmptyLinks;
 
 }
 
diff --git a/Engine/Source/Developer/RigVMDeveloper/Private/RigVMModel/RigVMController.cpp b/Engine/Source/Developer/RigVMDeveloper/Private/RigVMModel/RigVMController.cpp
--- a/Engine/Source/Developer/RigVMDeveloper/Private/RigVMModel/RigVMController.cpp
+++ b/Engine/Source/Developer/RigVMDeveloper/Private/RigVMModel/RigVMController.cpp
@@ -1720,7 +1720,7 @@ URigVMCollapseNode* URigVMController::CollapseNodes(const TArray<URigVMNode*>& I
 		return nullptr;
 	}
 
-	FBox2D Bounds;
+	FBox2D Bounds = FBox2D(EForceInit::ForceInit);
 	TArray<FName> NodeNames;
 	for (URigVMNode* Node : Nodes)
 	{
@@ -1871,134 +1871,137 @@ URigVMCollapseNode* URigVMController::CollapseNodes(const TArray<URigVMNode*>& I
 	}
 
 	// create the new nodes within the collapse node
-	if (bSetupUndoRedo)
+	TArray<FName> ContainedNodeNames;
 	{
-		TArray<FName> ContainedNodeNames;
-		{
-			FString TextContent = ExportNodesToText(NodeNames);
+		FString TextContent = ExportNodesToText(NodeNames);
 
-			TGuardValue<URigVMGraph*> GraphGuard(Graph, CollapseNode->ContainedGraph);
-			ContainedNodeNames = ImportNodesFromText(TextContent, false);
+		TGuardValue<URigVMGraph*> GraphGuard(Graph, CollapseNode->ContainedGraph);
+		ContainedNodeNames = ImportNodesFromText(TextContent, false);
 
-			// move the nodes to the right place
-			for (const FName& ContainedNodeName : ContainedNodeNames)
+		// move the nodes to the right place
+		for (const FName& ContainedNodeName : ContainedNodeNames)
+		{
+			if (URigVMNode* ContainedNode = CollapseNode->ContainedGraph->FindNodeByName(ContainedNodeName))
 			{
-				if (URigVMNode* ContainedNode = CollapseNode->ContainedGraph->FindNodeByName(ContainedNodeName))
-				{
-					ContainedNode->Position -= Center;
-				}
+				ContainedNode->Position -= Center;
+				SetNodePosition(ContainedNode, ContainedNode->Position - Center, false, false);
 			}
+		}
 
-			for (URigVMLink* LinkToRewire : LinksToRewire)
-			{
-				URigVMPin* SourcePin = LinkToRewire->GetSourcePin();
-				URigVMPin* TargetPin = LinkToRewire->GetTargetPin();
-
-				if (Nodes.Contains(SourcePin->GetNode()))
-				{
-					// if the parent pin of this was collapsed
-					// it's possible that the child pin wasn't.
-					if (!CollapsedPins.Contains(SourcePin))
-					{
-						continue;
-					}
+		for (URigVMLink* LinkToRewire : LinksToRewire)
+		{
+			URigVMPin* SourcePin = LinkToRewire->GetSourcePin();
+			URigVMPin* TargetPin = LinkToRewire->GetTargetPin();
 
-					URigVMPin* CollapsedPin = CollapsedPins.FindChecked(SourcePin);
-					SourcePin = CollapseNode->ContainedGraph->FindPin(SourcePin->GetPinPath());
-					TargetPin = ReturnNode->FindPin(CollapsedPin->GetName());
-				}
-				else
+			if (Nodes.Contains(SourcePin->GetNode()))
+			{
+				// if the parent pin of this was collapsed
+				// it's possible that the child pin wasn't.
+				if (!CollapsedPins.Contains(SourcePin))
 				{
-					URigVMPin* CollapsedPin = CollapsedPins.FindChecked(TargetPin);
-					SourcePin = EntryNode->FindPin(CollapsedPin->GetName());
-					TargetPin = CollapseNode->ContainedGraph->FindPin(TargetPin->GetPinPath());
+					continue;
 				}
 
-				if (SourcePin && TargetPin)
+				URigVMPin* CollapsedPin = CollapsedPins.FindChecked(SourcePin);
+				SourcePin = CollapseNode->ContainedGraph->FindPin(SourcePin->GetPinPath());
+				TargetPin = ReturnNode->FindPin(CollapsedPin->GetName());
+			}
+			else
+			{
+				URigVMPin* CollapsedPin = CollapsedPins.FindChecked(TargetPin);
+				SourcePin = EntryNode->FindPin(CollapsedPin->GetName());
+				TargetPin = CollapseNode->ContainedGraph->FindPin(TargetPin->GetPinPath());
+			}
+
+			if (SourcePin && TargetPin)
+			{
+				if (!SourcePin->IsLinkedTo(TargetPin))
 				{
-					if (!SourcePin->IsLinkedTo(TargetPin))
-					{
-						AddLink(SourcePin, TargetPin, false);
-					}
+					AddLink(SourcePin, TargetPin, false);
 				}
 			}
 		}
+	}
 
-		TArray<URigVMLink*> RewiredLinks;
-		for (URigVMLink* LinkToRewire : LinksToRewire)
+	TArray<URigVMLink*> RewiredLinks;
+	for (URigVMLink* LinkToRewire : LinksToRewire)
+	{
+		if (RewiredLinks.Contains(LinkToRewire))
 		{
-			if (RewiredLinks.Contains(LinkToRewire))
-			{
-				continue;
-			}
+			continue;
+		}
 
-			URigVMPin* SourcePin = LinkToRewire->GetSourcePin();
-			URigVMPin* TargetPin = LinkToRewire->GetTargetPin();
+		URigVMPin* SourcePin = LinkToRewire->GetSourcePin();
+		URigVMPin* TargetPin = LinkToRewire->GetTargetPin();
 
-			if (Nodes.Contains(SourcePin->GetNode()))
-			{
-				FString SegmentPath;
-				URigVMPin* PinToCheck = SourcePin;
+		if (Nodes.Contains(SourcePin->GetNode()))
+		{
+			FString SegmentPath;
+			URigVMPin* PinToCheck = SourcePin;
 
-				URigVMPin** CollapsedPinPtr = CollapsedPins.Find(PinToCheck);
-				while (CollapsedPinPtr == nullptr)
+			URigVMPin** CollapsedPinPtr = CollapsedPins.Find(PinToCheck);
+			while (CollapsedPinPtr == nullptr)
+			{
+				if (SegmentPath.IsEmpty())
 				{
-					if (SegmentPath.IsEmpty())
-					{
-						SegmentPath = PinToCheck->GetName();
-					}
-					else
-					{
-						SegmentPath = URigVMPin::JoinPinPath(PinToCheck->GetName(), SegmentPath);
-					}
+					SegmentPath = PinToCheck->GetName();
+				}
+				else
+				{
+					SegmentPath = URigVMPin::JoinPinPath(PinToCheck->GetName(), SegmentPath);
+				}
 
-					PinToCheck = PinToCheck->GetParentPin();
-					check(PinToCheck);
+				PinToCheck = PinToCheck->GetParentPin();
+				check(PinToCheck);
 
-					CollapsedPinPtr = CollapsedPins.Find(PinToCheck);
-				}
+				CollapsedPinPtr = CollapsedPins.Find(PinToCheck);
+			}
 
-				URigVMPin* CollapsedPin = *CollapsedPinPtr;
-				check(CollapsedPin);
+			URigVMPin* CollapsedPin = *CollapsedPinPtr;
+			check(CollapsedPin);
 
-				if (!SegmentPath.IsEmpty())
-				{
-					CollapsedPin = CollapsedPin->FindSubPin(SegmentPath);
-					check(CollapsedPin);
-				}
+			if (!SegmentPath.IsEmpty())
+			{
+				CollapsedPin = CollapsedPin->FindSubPin(SegmentPath);
+				check(CollapsedPin);
+			}
 
-				TArray<URigVMLink*> TargetLinks = SourcePin->GetTargetLinks(false);
-				for (URigVMLink* TargetLink : TargetLinks)
+			TArray<URigVMLink*> TargetLinks = SourcePin->GetTargetLinks(false);
+			for (URigVMLink* TargetLink : TargetLinks)
+			{
+				TargetPin = TargetLink->GetTargetPin();
+				if (!CollapsedPin->IsLinkedTo(TargetPin))
 				{
-					TargetPin = TargetLink->GetTargetPin();
-					if (!CollapsedPin->IsLinkedTo(TargetPin))
-					{
-						AddLink(CollapsedPin, TargetPin, false);
-					}
+					AddLink(CollapsedPin, TargetPin, false);
 				}
-				RewiredLinks.Append(TargetLinks);
 			}
-			else
+			RewiredLinks.Append(TargetLinks);
+		}
+		else
+		{
+			URigVMPin* CollapsedPin = CollapsedPins.FindChecked(TargetPin);
+			if (!SourcePin->IsLinkedTo(CollapsedPin))
 			{
-				URigVMPin* CollapsedPin = CollapsedPins.FindChecked(TargetPin);
-				if (!SourcePin->IsLinkedTo(CollapsedPin))
-				{
-					AddLink(SourcePin, CollapsedPin, false);
-				}
+				AddLink(SourcePin, CollapsedPin, false);
 			}
-
-			RewiredLinks.Add(LinkToRewire);
 		}
 
+		RewiredLinks.Add(LinkToRewire);
+	}
 
-		for (const FName& NodeToRemove : NodeNames)
-		{
-			RemoveNodeByName(NodeToRemove, false, true);
-		}
+
+	for (const FName& NodeToRemove : NodeNames)
+	{
+		RemoveNodeByName(NodeToRemove, false, true);
 	}
 
 	if (bSetupUndoRedo)
 	{
+		CollapseAction.LibraryNodePath = CollapseNode->GetName();
+		for (URigVMNode* InNode : InNodes)
+		{
+			CollapseAction.CollapsedNodesPaths.Add(InNode->GetName());
+		}
 		ActionStack->EndAction(CollapseAction);
 	}
 
@@ -2028,7 +2031,7 @@ TArray<URigVMNode*> URigVMController::ExpandLibraryNode(URigVMLibraryNode* InNod
 	}
 
 	TArray<FName> NodeNames;
-	FBox2D Bounds;
+	FBox2D Bounds = FBox2D(EForceInit::ForceInit);
 	{
 		TArray<URigVMNode*> FilteredNodes;
 		for (URigVMNode* Node : ContainedNodes)
@@ -2075,6 +2078,7 @@ TArray<URigVMNode*> URigVMController::ExpandLibraryNode(URigVMLibraryNode* InNod
 	for (int32 NodeNameIndex = 0; NodeNameIndex < NodeNames.Num(); NodeNameIndex++)
 	{
 		NodeNameMap.Add(NodeNames[NodeNameIndex], ExpandedNodeNames[NodeNameIndex]);
+		SetNodePosition(ExpandedNodes[NodeNameIndex], InNode->Position + ContainedNodes[NodeNameIndex]->Position - Center, false, false);
 	}
 
 	// a) store all of the pin defaults off the library node
@@ -2083,18 +2087,29 @@ TArray<URigVMNode*> URigVMController::ExpandLibraryNode(URigVMLibraryNode* InNod
 	// b) create a map of new links to create by following the links to / from the library node
 	TMap<FString, TArray<FString>> ToLibraryNode;
 	TMap<FString, TArray<FString>> FromLibraryNode;
+	TArray<URigVMPin*> LibraryPinsToReroute;
 
 	TArray<URigVMLink*> LibraryLinks = InNode->GetLinks();
 	for (URigVMLink* Link : LibraryLinks)
 	{
 		if (Link->GetTargetPin()->GetNode() == InNode)
 		{
+			if (!Link->GetTargetPin()->IsRootPin())
+			{
+				LibraryPinsToReroute.AddUnique(Link->GetTargetPin()->GetRootPin());
+			}
+
 			FString NodeName, PinPath;
 			URigVMPin::SplitPinPathAtStart(Link->GetTargetPin()->GetPinPath(), NodeName, PinPath);
 			ToLibraryNode.FindOrAdd(PinPath).Add(Link->GetSourcePin()->GetPinPath());
 		}
 		else
 		{
+			if (!Link->GetSourcePin()->IsRootPin())
+			{
+				LibraryPinsToReroute.AddUnique(Link->GetSourcePin()->GetRootPin());
+			}
+
 			FString NodeName, PinPath;
 			URigVMPin::SplitPinPathAtStart(Link->GetSourcePin()->GetPinPath(), NodeName, PinPath);
 			FromLibraryNode.FindOrAdd(PinPath).Add(Link->GetTargetPin()->GetPinPath());
@@ -2113,6 +2128,11 @@ TArray<URigVMNode*> URigVMController::ExpandLibraryNode(URigVMLibraryNode* InNod
 				continue;
 			}
 
+			if (!Link->GetSourcePin()->IsRootPin())
+			{
+				LibraryPinsToReroute.AddUnique(InNode->FindPin(Link->GetSourcePin()->GetRootPin()->GetName()));
+			}
+
 			FString NodeName, PinPath;
 			URigVMPin::SplitPinPathAtStart(Link->GetSourcePin()->GetPinPath(), NodeName, PinPath);
 
@@ -2136,6 +2156,11 @@ TArray<URigVMNode*> URigVMController::ExpandLibraryNode(URigVMLibraryNode* InNod
 				continue;
 			}
 
+			if (!Link->GetTargetPin()->IsRootPin())
+			{
+				LibraryPinsToReroute.AddUnique(InNode->FindPin(Link->GetTargetPin()->GetRootPin()->GetName()));
+			}
+
 			FString NodeName, PinPath;
 			URigVMPin::SplitPinPathAtStart(Link->GetTargetPin()->GetPinPath(), NodeName, PinPath);
 
@@ -2172,7 +2197,61 @@ TArray<URigVMNode*> URigVMController::ExpandLibraryNode(URigVMLibraryNode* InNod
 		}
 	}
 
-	// f) remap all output / source pins and create a final list of links to create
+	// f) create reroutes for all pins which had wires on sub pins
+	TMap<FString, URigVMPin*> ReroutedInputPins;
+	TMap<FString, URigVMPin*> ReroutedOutputPins;
+	FVector2D RerouteInputPosition = InNode->Position + FVector2D(-Diagonal.X, -Diagonal.Y) * 0.5 + FVector2D(-200.f, 0.f);
+	FVector2D RerouteOutputPosition = InNode->Position + FVector2D(Diagonal.X, -Diagonal.Y) * 0.5 + FVector2D(250.f, 0.f);
+	for (URigVMPin* LibraryPinToReroute : LibraryPinsToReroute)
+	{
+		if (LibraryPinToReroute->GetDirection() == ERigVMPinDirection::Input ||
+			LibraryPinToReroute->GetDirection() == ERigVMPinDirection::IO)
+		{
+			URigVMRerouteNode* RerouteNode =
+				AddFreeRerouteNode(
+					true,
+					LibraryPinToReroute->GetCPPType(),
+					*LibraryPinToReroute->GetCPPTypeObject()->GetPathName(),
+					false,
+					NAME_None,
+					LibraryPinToReroute->GetDefaultValue(),
+					RerouteInputPosition,
+					FString::Printf(TEXT("Reroute_%s"), *LibraryPinToReroute->GetName()),
+					false);
+
+			RerouteInputPosition += FVector2D(0.f, 150.f);
+
+			URigVMPin* ReroutePin = RerouteNode->FindPin(URigVMRerouteNode::ValueName);
+			ApplyPinState(ReroutePin, GetPinState(LibraryPinToReroute));
+			ReroutedInputPins.Add(LibraryPinToReroute->GetName(), ReroutePin);
+			ExpandedNodes.Add(RerouteNode);
+		}
+
+		if (LibraryPinToReroute->GetDirection() == ERigVMPinDirection::Output ||
+			LibraryPinToReroute->GetDirection() == ERigVMPinDirection::IO)
+		{
+			URigVMRerouteNode* RerouteNode =
+				AddFreeRerouteNode(
+					true,
+					LibraryPinToReroute->GetCPPType(),
+					*LibraryPinToReroute->GetCPPTypeObject()->GetPathName(),
+					false,
+					NAME_None,
+					LibraryPinToReroute->GetDefaultValue(),
+					RerouteOutputPosition,
+					FString::Printf(TEXT("Reroute_%s"), *LibraryPinToReroute->GetName()),
+					false);
+
+			RerouteOutputPosition += FVector2D(0.f, 150.f);
+
+			URigVMPin* ReroutePin = RerouteNode->FindPin(URigVMRerouteNode::ValueName);
+			ApplyPinState(ReroutePin, GetPinState(LibraryPinToReroute));
+			ReroutedOutputPins.Add(LibraryPinToReroute->GetName(), ReroutePin);
+			ExpandedNodes.Add(RerouteNode);
+		}
+	}
+
+	// g) remap all output / source pins and create a final list of links to create
 	TMap<FString, FString> RemappedSourcePins;;
 	TArray<URigVMPin*> LibraryPins = InNode->GetAllPinsRecursively();
 	for (URigVMPin* LibraryPin : LibraryPins)
@@ -2181,6 +2260,9 @@ TArray<URigVMNode*> URigVMController::ExpandLibraryNode(URigVMLibraryNode* InNod
 		FString LibraryNodeName;
 		URigVMPin::SplitPinPathAtStart(LibraryPinPath, LibraryNodeName, LibraryPinPath);
 
+		FString SourcePinPath;
+		FString TargetPinPath;
+
 		if (LibraryPin->GetDirection() == ERigVMPinDirection::Input)
 		{
 			if (const TArray<FString>* LibraryPinLinksPtr = ToLibraryNode.Find(LibraryPinPath))
@@ -2188,7 +2270,8 @@ TArray<URigVMNode*> URigVMController::ExpandLibraryNode(URigVMLibraryNode* InNod
 				const TArray<FString>& LibraryPinLinks = *LibraryPinLinksPtr;
 				ensure(LibraryPinLinks.Num() == 1);
 
-				RemappedSourcePins.FindOrAdd(LibraryPinPath) = LibraryPinLinks[0];
+				SourcePinPath = LibraryPinPath;
+				TargetPinPath = LibraryPinLinks[0];
 			}
 		}
 		else if (LibraryPin->GetDirection() == ERigVMPinDirection::Output)
@@ -2198,17 +2281,64 @@ TArray<URigVMNode*> URigVMController::ExpandLibraryNode(URigVMLibraryNode* InNod
 				const TArray<FString>& LibraryPinLinks = *LibraryPinLinksPtr;
 				ensure(LibraryPinLinks.Num() == 1);
 
-				FString SourcePinNodeName, SourcePinPath;
-				URigVMPin::SplitPinPathAtStart(LibraryPinLinks[0], SourcePinNodeName, SourcePinPath);
+				FString TargetPinNodeName, TargetPinPathRemaining;
+				URigVMPin::SplitPinPathAtStart(LibraryPinLinks[0], TargetPinNodeName, TargetPinPathRemaining);
+
+				SourcePinPath = LibraryPinPath;
+				TargetPinNodeName = NodeNameMap.FindChecked(*TargetPinNodeName).ToString();
+				TargetPinPath = URigVMPin::JoinPinPath(TargetPinNodeName, TargetPinPathRemaining);
+			}
+		}
+
+		while (!SourcePinPath.IsEmpty() && !TargetPinPath.IsEmpty())
+		{
+			RemappedSourcePins.FindOrAdd(SourcePinPath) = TargetPinPath;
+
+			FString SourceLastSegment, TargetLastSegment;
+			if (!URigVMPin::SplitPinPathAtEnd(SourcePinPath, SourcePinPath, SourceLastSegment))
+			{
+				break;
+			}
+			if (!URigVMPin::SplitPinPathAtEnd(TargetPinPath, TargetPinPath, TargetLastSegment))
+			{
+				break;
+			}
+		}
+	}
+
+	// h) re-establish all of the links going to the left of the library node
+	//    in this pass we only care about pins which have reroutes
+	for (const TPair<FString, TArray<FString>>& ToLibraryNodePair : ToLibraryNode)
+	{
+		FString LibraryNodePinName, LibraryNodePinPathSuffix;
+		if (!URigVMPin::SplitPinPathAtStart(ToLibraryNodePair.Key, LibraryNodePinName, LibraryNodePinPathSuffix))
+		{
+			LibraryNodePinName = ToLibraryNodePair.Key;
+		}
+
+		if (!ReroutedInputPins.Contains(LibraryNodePinName))
+		{
+			continue;
+		}
 
-				SourcePinNodeName = NodeNameMap.FindChecked(*SourcePinNodeName).ToString();
+		URigVMPin* ReroutedPin = ReroutedInputPins.FindChecked(LibraryNodePinName);
+		URigVMPin* TargetPin = LibraryNodePinPathSuffix.IsEmpty() ? ReroutedPin : ReroutedPin->FindSubPin(LibraryNodePinPathSuffix);
+		check(TargetPin);
 
-				RemappedSourcePins.FindOrAdd(LibraryPinPath) = URigVMPin::JoinPinPath(SourcePinNodeName, SourcePinPath);
+		for (const FString& SourcePinPath : ToLibraryNodePair.Value)
+		{
+			URigVMPin* SourcePin = GetGraph()->FindPin(*SourcePinPath);
+			if (SourcePin && TargetPin)
+			{
+				if (!SourcePin->IsLinkedTo(TargetPin))
+				{
+					AddLink(SourcePin, TargetPin, false);
+				}
 			}
 		}
 	}
 
-	// g) re-establish all of the links going to the left of the library node
+	// i) re-establish all of the links going to the left of the library node (based on the entry node)
 	for (const TPair<FString, TArray<FString>>& FromEntryNodePair : FromEntryNode)
 	{
 		FString EntryPinPath = FromEntryNodePair.Key;
@@ -2246,6 +2376,21 @@ TArray<URigVMNode*> URigVMController::ExpandLibraryNode(URigVMLibraryNode* InNod
 			RemappedSourcePinPath = URigVMPin::JoinPinPath(RemappedSourcePinPath, EntryPinPathSuffix);
 		}
 
+		// remap the top level pin in case we need to insert a reroute
+		FString EntryPinName;
+		if (!URigVMPin::SplitPinPathAtStart(FromEntryNodePair.Key, EntryPinPath, EntryPinPathSuffix))
+		{
+			EntryPinName = FromEntryNodePair.Key;
+			EntryPinPathSuffix.Reset();
+		}
+		if (ReroutedInputPins.Contains(EntryPinName))
+		{
+			URigVMPin* ReroutedPin = ReroutedInputPins.FindChecked(EntryPinName);
+			URigVMPin* TargetPin = EntryPinPathSuffix.IsEmpty() ? ReroutedPin : ReroutedPin->FindSubPin(EntryPinPathSuffix);
+			check(TargetPin);
+			RemappedSourcePinPath = TargetPin->GetPinPath();
+		}
+
 		for (const FString& FromEntryNodeTargetPinPath : FromEntryNodePair.Value)
 		{
 			FString TargetPinNodeName, TargetPinPath;
@@ -2265,7 +2410,39 @@ TArray<URigVMNode*> URigVMController::ExpandLibraryNode(URigVMLibraryNode* InNod
 		}
 	}
 
-	// h) re-establish all of the links going from the right of the library node
+	// j) re-establish all of the links going from the right of the library node
+	//    in this pass we only check pins which have a reroute
+	for (const TPair<FString, TArray<FString>>& ToReturnNodePair : ToReturnNode)
+	{
+		FString LibraryNodePinName, LibraryNodePinPathSuffix;
+		if (!URigVMPin::SplitPinPathAtStart(ToReturnNodePair.Key, LibraryNodePinName, LibraryNodePinPathSuffix))
+		{
+			LibraryNodePinName = ToReturnNodePair.Key;
+		}
+
+		if (!ReroutedOutputPins.Contains(LibraryNodePinName))
+		{
+			continue;
+		}
+
+		URigVMPin* ReroutedPin = ReroutedOutputPins.FindChecked(LibraryNodePinName);
+		URigVMPin* TargetPin = LibraryNodePinPathSuffix.IsEmpty() ? ReroutedPin : ReroutedPin->FindSubPin(LibraryNodePinPathSuffix);
+		check(TargetPin);
+
+		for (const FString& SourcePinpath : ToReturnNodePair.Value)
+		{
+			URigVMPin* SourcePin = GetGraph()->FindPin(*SourcePinpath);
+			if (SourcePin && TargetPin)
+			{
+				if (!SourcePin->IsLinkedTo(TargetPin))
+				{
+					AddLink(SourcePin, TargetPin, false);
+				}
+			}
+		}
+	}
+
+	// k) re-establish all of the links going from the right of the library node
 	for (const TPair<FString, TArray<FString>>& FromLibraryNodePair : FromLibraryNode)
 	{
 		FString FromLibraryNodePinPath = FromLibraryNodePair.Key;
@@ -2303,6 +2480,21 @@ TArray<URigVMNode*> URigVMController::ExpandLibraryNode(URigVMLibraryNode* InNod
 			RemappedSourcePinPath = URigVMPin::JoinPinPath(RemappedSourcePinPath, FromLibraryNodePinPathSuffix);
 		}
 
+		// remap the top level pin in case we need to insert a reroute
+		FString ReturnPinName, ReturnPinPathSuffix;
+		if (!URigVMPin::SplitPinPathAtStart(FromLibraryNodePair.Key, ReturnPinName, ReturnPinPathSuffix))
+		{
+			ReturnPinName = FromLibraryNodePair.Key;
+			ReturnPinPathSuffix.Reset();
+		}
+		if (ReroutedOutputPins.Contains(ReturnPinName))
+		{
+			URigVMPin* ReroutedPin = ReroutedOutputPins.FindChecked(ReturnPinName);
+			URigVMPin* SourcePin = ReturnPinPathSuffix.IsEmpty() ? ReroutedPin : ReroutedPin->FindSubPin(ReturnPinPathSuffix);
+			check(SourcePin);
+			RemappedSourcePinPath = SourcePin->GetPinPath();
+		}
+
 		for (const FString& FromLibraryNodeTargetPinPath : FromLibraryNodePair.Value)
 		{
 			FString TargetPinNodeName, TargetPinPath;
@@ -2320,11 +2512,19 @@ TArray<URigVMNode*> URigVMController::ExpandLibraryNode(URigVMLibraryNode* InNod
 		}
 	}
 
-	// i) remove the library node from the graph
+	// l) remove the library node from the graph
+	if (bSetupUndoRedo)
+	{
+		ExpandAction.LibraryNodePath = InNode->GetName();
+	}
 	RemoveNode(InNode, false, true);
 
 	if (bSetupUndoRedo)
 	{
+		for (URigVMNode* ExpandedNode : ExpandedNodes)
+		{
+			ExpandAction.ExpandedNodePaths.Add(ExpandedNode->GetName());
+		}
 		ActionStack->EndAction(ExpandAction);
 	}
 
@@ -5313,6 +5513,18 @@ int32 URigVMController::DetachLinksFromPinObjects(const TArray<URigVMLink*>* InL
 		Link->TargetPin = nullptr;
 	}
 
+	if (InLinks == nullptr)
+	{
+		for (URigVMNode* Node : Graph->Nodes)
+		{
+			if (URigVMCollapseNode* CollapseNode = Cast<URigVMCollapseNode>(Node))
+			{
+				TGuardValue<URigVMGraph*> GraphGuard(Graph, CollapseNode->ContainedGraph);
+				DetachLinksFromPinObjects();
+			}
+		}
+	}
+
 	return Links.Num();
 }
 
@@ -5413,6 +5625,18 @@ int32 URigVMController::ReattachLinksToPinObjects(bool bFollowCoreRedirectors, c
 		}
 	}
 
+	if (InLinks == nullptr)
+	{
+		for (URigVMNode* Node : Graph->Nodes)
+		{
+			if (URigVMCollapseNode* CollapseNode = Cast<URigVMCollapseNode>(Node))
+			{
+				TGuardValue<URigVMGraph*> GraphGuard(Graph, CollapseNode->ContainedGraph);
+				ReattachLinksToPinObjects(bFollowCoreRedirectors, nullptr);
+			}
+		}
+	}
+
 	InputPinRedirectors.Reset();
 	OutputPinRedirectors.Reset();
 
@@ -5609,6 +5833,7 @@ void URigVMController::RepopulatePinsOnNode(URigVMNode* InNode, bool bFollowCore
 	URigVMRerouteNode* RerouteNode = Cast<URigVMRerouteNode>(InNode);
 	URigVMFunctionEntryNode* EntryNode = Cast<URigVMFunctionEntryNode>(InNode);
 	URigVMFunctionReturnNode* ReturnNode = Cast<URigVMFunctionReturnNode>(InNode);
+	URigVMCollapseNode* CollapseNode = Cast<URigVMCollapseNode>(InNode);
 
 	TGuardValue<bool> SuspendNotifs(bSuspendNotifications, true);
 	FScopeLock Lock(&PinPathCoreRedirectorsLock);
@@ -5767,6 +5992,14 @@ void URigVMController::RepopulatePinsOnNode(URigVMNode* InNode, bool bFollowCore
 			checkNoEntry();
 		}
 	}
+	else if (CollapseNode)
+	{
+		TGuardValue<URigVMGraph*> GraphGuard(Graph, CollapseNode->ContainedGraph);
+		for (URigVMNode* ContainedNode : CollapseNode->GetContainedNodes())
+		{
+			RepopulatePinsOnNode(ContainedNode, bFollowCoreRedirectors);
+		}
+	}
 	else
 	{
 		return;
diff --git a/Engine/Source/Developer/RigVMDeveloper/Private/RigVMModel/RigVMPin.cpp b/Engine/Source/Developer/RigVMDeveloper/Private/RigVMModel/RigVMPin.cpp
--- a/Engine/Source/Developer/RigVMDeveloper/Private/RigVMModel/RigVMPin.cpp
+++ b/Engine/Source/Developer/RigVMDeveloper/Private/RigVMModel/RigVMPin.cpp
@@ -84,7 +84,7 @@ URigVMPin::URigVMPin()
 	, DefaultValue(FString())
 	, BoundVariablePath()
 {
-#if WITH_EDITORONLY_DATA
+#if UE_BUILD_DEBUG
 	CachedPinPath = GetPinPath();
 #endif
 }
@@ -106,7 +106,7 @@ FString URigVMPin::GetPinPath() const
 			PinPath = FString::Printf(TEXT("%s.%s"), *Node->GetNodePath(), *GetName());
 		}
 	}
-#if WITH_EDITORONLY_DATA
+#if UE_BUILD_DEBUG
 	CachedPinPath = PinPath;
 #endif
 	return PinPath;
@@ -593,6 +593,11 @@ URigVMPin* URigVMPin::GetRootPin() const
 	return ParentPin->GetRootPin();
 }
 
+bool URigVMPin::IsRootPin() const
+{
+	return GetParentPin() == nullptr;
+}
+
 URigVMPin* URigVMPin::GetPinForLink() const
 {
 	URigVMPin* RootPin = GetRootPin();
diff --git a/Engine/Source/Developer/RigVMDeveloper/Public/RigVMModel/Nodes/RigVMLibraryNode.h b/Engine/Source/Developer/RigVMDeveloper/Public/RigVMModel/Nodes/RigVMLibraryNode.h
--- a/Engine/Source/Developer/RigVMDeveloper/Public/RigVMModel/Nodes/RigVMLibraryNode.h
+++ b/Engine/Source/Developer/RigVMDeveloper/Public/RigVMModel/Nodes/RigVMLibraryNode.h
@@ -28,15 +28,17 @@ public:
 
 	// Library node interface
 	virtual URigVMGraph* GetContainedGraph() const { return nullptr; }
-	virtual TArray<URigVMNode*> GetContainedNodes() const;
-	virtual TArray<URigVMLink*> GetContainedLinks() const;
+	virtual const TArray<URigVMNode*>& GetContainedNodes() const;
+	virtual const TArray<URigVMLink*>& GetContainedLinks() const;
 	virtual URigVMFunctionEntryNode* GetEntryNode() const;
 	virtual URigVMFunctionReturnNode* GetReturnNode() const;
 
-private:
+protected:
+
+	const static TArray<URigVMNode*> EmptyNodes;
+	const static TArray<URigVMLink*> EmptyLinks;
 
-	UPROPERTY()
-	bool bDefinedAsVarying = false;
+private:
 
 	friend class URigVMController;
 };
diff --git a/Engine/Source/Developer/RigVMDeveloper/Public/RigVMModel/RigVMPin.h b/Engine/Source/Developer/RigVMDeveloper/Public/RigVMModel/RigVMPin.h
--- a/Engine/Source/Developer/RigVMDeveloper/Public/RigVMModel/RigVMPin.h
+++ b/Engine/Source/Developer/RigVMDeveloper/Public/RigVMModel/RigVMPin.h
@@ -222,6 +222,10 @@ public:
 	UFUNCTION(BlueprintCallable, Category = RigVMPin)
 	URigVMPin* GetRootPin() const;
 
+	// Returns true if this pin is a root pin
+	UFUNCTION(BlueprintCallable, Category = RigVMPin)
+	bool IsRootPin() const;
+
 	// Returns the pin to be used for a link.
 	// This might differ from this actual pin, since
 	// the pin might contain injected nodes.
@@ -328,9 +332,14 @@ private:
 	UPROPERTY()
 	FName DisplayName;
 
-#if WITH_EDITORONLY_DATA
-	UPROPERTY(transient)
+#if UE_BUILD_DEBUG
+	
+	// A cache for the pin path for debugging purposes.
+	// When looking at a debug symbol of a pin it is difficult
+	// To grasp where the pin is stored etc.
+	// With this member you can see the full pin path during a debugging session
 	mutable FString CachedPinPath;
+
 #endif
 
 	UPROPERTY()
