Control Rig: Only allow event nodes in top level graph

diff --git a/Engine/Source/Developer/RigVMDeveloper/Private/RigVMModel/RigVMController.cpp b/Engine/Source/Developer/RigVMDeveloper/Private/RigVMModel/RigVMController.cpp
--- a/Engine/Source/Developer/RigVMDeveloper/Private/RigVMModel/RigVMController.cpp
+++ b/Engine/Source/Developer/RigVMDeveloper/Private/RigVMModel/RigVMController.cpp
@@ -193,6 +193,23 @@ URigVMUnitNode* URigVMController::AddUnitNode(UScriptStruct* InScriptStruct, con
 		return nullptr;
 	}
 
+	// don't allow event nodes in anything but top level graphs
+	if (bSetupUndoRedo)
+	{
+		if (!Graph->IsTopLevelGraph())
+		{
+			FStructOnScope StructOnScope(InScriptStruct);
+			FRigVMStruct* StructMemory = (FRigVMStruct*)StructOnScope.GetStructMemory();
+			InScriptStruct->InitializeDefaultValue((uint8*)StructMemory);
+
+			if (!StructMemory->GetEventName().IsNone())
+			{
+				ReportError(TEXT("Event nodes can only be added to top level graphs."));
+				return nullptr;
+			}
+		}
+	}
+
 	FString Name = GetValidNodeName(InNodeName.IsEmpty() ? InScriptStruct->GetName() : InNodeName);
 	URigVMUnitNode* Node = NewObject<URigVMUnitNode>(Graph, *Name);
 	Node->ScriptStruct = InScriptStruct;
diff --git a/Engine/Source/Developer/RigVMDeveloper/Private/RigVMModel/RigVMGraph.cpp b/Engine/Source/Developer/RigVMDeveloper/Private/RigVMModel/RigVMGraph.cpp
--- a/Engine/Source/Developer/RigVMDeveloper/Private/RigVMModel/RigVMGraph.cpp
+++ b/Engine/Source/Developer/RigVMDeveloper/Private/RigVMModel/RigVMGraph.cpp
@@ -109,6 +109,16 @@ const TArray<FName>& URigVMGraph::GetSelectNodes() const
 	return SelectedNodes;
 }
 
+bool URigVMGraph::IsTopLevelGraph() const
+{
+	if (GetOuter()->IsA<URigVMLibraryNode>())
+	{
+		return false;
+	}
+	// todo: implement this for function libraries
+	return true;
+}
+
 FRigVMGraphModifiedEvent& URigVMGraph::OnModified()
 {
 	return ModifiedEvent;
diff --git a/Engine/Source/Developer/RigVMDeveloper/Public/RigVMModel/RigVMGraph.h b/Engine/Source/Developer/RigVMDeveloper/Public/RigVMModel/RigVMGraph.h
--- a/Engine/Source/Developer/RigVMDeveloper/Public/RigVMModel/RigVMGraph.h
+++ b/Engine/Source/Developer/RigVMDeveloper/Public/RigVMModel/RigVMGraph.h
@@ -72,6 +72,10 @@ public:
 	UFUNCTION(BlueprintCallable, Category = RigVMGraph)
 	const TArray<FName>& GetSelectNodes() const;
 
+	// Returns true if this graph is the top level graph
+	UFUNCTION(BlueprintCallable, Category = RigVMGraph)
+	bool IsTopLevelGraph() const;
+
 	// Returns the modified event, which can be used to 
 	// subscribe to changes happening within the Graph.
 	FRigVMGraphModifiedEvent& OnModified();
