Control Rig Encapsulation: Filtering for multi branch for collapse action

diff --git a/Engine/Source/Developer/RigVMDeveloper/Private/RigVMModel/RigVMController.cpp b/Engine/Source/Developer/RigVMDeveloper/Private/RigVMModel/RigVMController.cpp
--- a/Engine/Source/Developer/RigVMDeveloper/Private/RigVMModel/RigVMController.cpp
+++ b/Engine/Source/Developer/RigVMDeveloper/Private/RigVMModel/RigVMController.cpp
@@ -17,6 +17,8 @@
 #include "UnrealExporter.h"
 #include "Factories.h"
 #include "UObject/CoreRedirects.h"
+#include "Framework/Notifications/NotificationManager.h"
+#include "Widgets/Notifications/SNotificationList.h"
 #endif
 
 TMap<URigVMController::FControlRigStructPinRedirectorKey, FString> URigVMController::PinPathCoreRedirectors;
@@ -1747,20 +1749,6 @@ URigVMCollapseNode* URigVMController::CollapseNodes(const TArray<URigVMNode*>& I
 
   	FVector2D Diagonal = Bounds.Max - Bounds.Min;
 	FVector2D Center = (Bounds.Min + Bounds.Max) * 0.5f;
-	
-	FRigVMCollapseNodesAction CollapseAction;
-	CollapseAction.Title = TEXT("Collapse Nodes");
-
-	if (bSetupUndoRedo)
-	{
-		ActionStack->BeginAction(CollapseAction);
-	}
-
-	FString CollapseNodeName = GetValidNodeName(InCollapseNodeName.IsEmpty() ? FString(TEXT("CollapseNode")) : InCollapseNodeName);
-	URigVMCollapseNode* CollapseNode = NewObject<URigVMCollapseNode>(Graph, *CollapseNodeName);
-	CollapseNode->ContainedGraph = NewObject<URigVMGraph>(CollapseNode, TEXT("ContainedGraph"));
-	CollapseNode->Position = Center;
-	Graph->Nodes.Add(CollapseNode);
 
 	bool bContainsOutputs = false;
 
@@ -1774,18 +1762,143 @@ URigVMCollapseNode* URigVMController::CollapseNodes(const TArray<URigVMNode*>& I
 	// collapsed.
 	for (URigVMLink* Link : AllLinks)
 	{
-		bool bSourceToBeCollapsed = Nodes.Contains(Link->GetSourcePin()->GetNode());
-		bool bTargetToBeCollapsed = Nodes.Contains(Link->GetTargetPin()->GetNode());
+		URigVMPin* SourcePin = Link->GetSourcePin();
+		URigVMPin* TargetPin = Link->GetTargetPin();
+		bool bSourceToBeCollapsed = Nodes.Contains(SourcePin->GetNode());
+		bool bTargetToBeCollapsed = Nodes.Contains(TargetPin->GetNode());
 		if (bSourceToBeCollapsed == bTargetToBeCollapsed)
 		{
 			continue;
 		}
 
-		URigVMPin* PinToCollapse = Link->GetSourcePin();
+		URigVMPin* PinToCollapse = SourcePin;
 		PinsToCollapse.AddUnique(PinToCollapse);
 		LinksToRewire.Add(Link);
 	}
 
+	// make sure that for execute pins we are on one branch only
+	TArray<URigVMPin*> InputExecutePins;
+	TArray<URigVMPin*> IntermediateExecutePins;
+	TArray<URigVMPin*> OutputExecutePins;
+
+	// first collect the output execute pins
+	for (URigVMLink* Link : LinksToRewire)
+	{
+		URigVMPin* ExecutePin = Link->GetSourcePin();
+		if (!ExecutePin->IsExecuteContext())
+		{
+			continue;
+		}
+		if (!Nodes.Contains(ExecutePin->GetNode()))
+		{
+			continue;
+		}
+		if (!OutputExecutePins.IsEmpty())
+		{
+			if (bSetupUndoRedo)
+			{
+				ReportAndNotifyErrorf(
+					TEXT("Only one set of execute branches can be collapsed, pin %s and %s are on separate branches"),
+					*OutputExecutePins[0]->GetPinPath(),
+					*ExecutePin->GetPinPath()
+				);
+			}
+			return nullptr;
+		}
+		OutputExecutePins.Add(ExecutePin);
+
+		while (ExecutePin)
+		{
+			if (IntermediateExecutePins.Contains(ExecutePin))
+			{
+				if (bSetupUndoRedo)
+				{
+					ReportAndNotifyErrorf(TEXT("Only one set of execute branches can be collapsed."));
+				}
+				return nullptr;
+			}
+			IntermediateExecutePins.Add(ExecutePin);
+
+			// walk backwards and find all "known execute pins"
+			URigVMNode* ExecuteNode = ExecutePin->GetNode();
+			for (URigVMPin* Pin : ExecuteNode->GetPins())
+			{
+				if (Pin->GetDirection() != ERigVMPinDirection::Input &&
+					Pin->GetDirection() != ERigVMPinDirection::IO)
+				{
+					continue;
+				}
+				if (!Pin->IsExecuteContext())
+				{
+					continue;
+				}
+				TArray<URigVMLink*> SourceLinks = Pin->GetSourceLinks();
+				ExecutePin = nullptr;
+				if (SourceLinks.Num() > 0)
+				{
+					URigVMPin* PreviousExecutePin = SourceLinks[0]->GetSourcePin();
+					if (Nodes.Contains(PreviousExecutePin->GetNode()))
+					{
+						if (Pin != IntermediateExecutePins.Last())
+						{
+							IntermediateExecutePins.Add(Pin);
+						}
+						ExecutePin = PreviousExecutePin;
+						break;
+					}
+				}
+			}
+		}
+	}
+	for (URigVMLink* Link : LinksToRewire)
+	{
+		URigVMPin* ExecutePin = Link->GetTargetPin();
+		if (!ExecutePin->IsExecuteContext())
+		{
+			continue;
+		}
+		if (!Nodes.Contains(ExecutePin->GetNode()))
+		{
+			continue;
+		}
+		if (!IntermediateExecutePins.Contains(ExecutePin) && !IntermediateExecutePins.IsEmpty())
+		{
+			if (bSetupUndoRedo)
+			{
+				ReportAndNotifyErrorf(TEXT("Only one set of execute branches can be collapsed"));
+			}
+			return nullptr;
+		}
+
+		if (!InputExecutePins.IsEmpty())
+		{
+			if (bSetupUndoRedo)
+			{
+				ReportAndNotifyErrorf(
+					TEXT("Only one set of execute branches can be collapsed, pin %s and %s are on separate branches"),
+					*InputExecutePins[0]->GetPinPath(),
+					*ExecutePin->GetPinPath()
+				);
+			}
+			return nullptr;
+		}
+		InputExecutePins.Add(ExecutePin);
+	}
+
+	FRigVMCollapseNodesAction CollapseAction;
+	CollapseAction.Title = TEXT("Collapse Nodes");
+
+	if (bSetupUndoRedo)
+	{
+		ActionStack->BeginAction(CollapseAction);
+	}
+
+	FString CollapseNodeName = GetValidNodeName(InCollapseNodeName.IsEmpty() ? FString(TEXT("CollapseNode")) : InCollapseNodeName);
+	URigVMCollapseNode* CollapseNode = NewObject<URigVMCollapseNode>(Graph, *CollapseNodeName);
+	CollapseNode->ContainedGraph = NewObject<URigVMGraph>(CollapseNode, TEXT("ContainedGraph"));
+	CollapseNode->Position = Center;
+	Graph->Nodes.Add(CollapseNode);
+
 	// now looper over the links to be rewired
 	for (URigVMLink* Link : LinksToRewire)
 	{
@@ -6255,6 +6368,32 @@ void URigVMController::ReportError(const FString& InMessage)
 	FScriptExceptionHandler::Get().HandleException(ELogVerbosity::Error, *Message, *FString());
 }
 
+void URigVMController::ReportAndNotifyError(const FString& InMessage)
+{
+	if (!bReportWarningsAndErrors)
+	{
+		return;
+	}
+
+	ReportError(InMessage);
+
+#if WITH_EDITOR
+	FNotificationInfo Info(FText::FromString(InMessage));
+	Info.bUseSuccessFailIcons = true;
+	Info.Image = FEditorStyle::GetBrush(TEXT("MessageLog.Warning"));
+	Info.bFireAndForget = true;
+	Info.bUseThrobber = true;
+	Info.FadeOutDuration = 5.0f;
+	Info.ExpireDuration = Info.FadeOutDuration;
+	TSharedPtr<SNotificationItem> NotificationPtr = FSlateNotificationManager::Get().AddNotification(Info);
+	if (NotificationPtr)
+	{
+		NotificationPtr->SetCompletionState(SNotificationItem::CS_Fail);
+	}
+#endif
+}
+
+
 void URigVMController::CreateDefaultValueForStructIfRequired(UScriptStruct* InStruct, FString& OutDefaultValue)
 {
 	if ((OutDefaultValue.IsEmpty() || OutDefaultValue == TEXT("()")) && InStruct != nullptr)
diff --git a/Engine/Source/Developer/RigVMDeveloper/Public/RigVMModel/RigVMController.h b/Engine/Source/Developer/RigVMDeveloper/Public/RigVMModel/RigVMController.h
--- a/Engine/Source/Developer/RigVMDeveloper/Public/RigVMModel/RigVMController.h
+++ b/Engine/Source/Developer/RigVMDeveloper/Public/RigVMModel/RigVMController.h
@@ -545,6 +545,7 @@ private:
 
 	void ReportWarning(const FString& InMessage);
 	void ReportError(const FString& InMessage);
+	void ReportAndNotifyError(const FString& InMessage);
 
 	template <typename FmtType, typename... Types>
 	void ReportWarningf(const FmtType& Fmt, Types... Args)
@@ -558,6 +559,12 @@ private:
 		ReportError(FString::Printf(Fmt, Args...));
 	}
 
+	template <typename FmtType, typename... Types>
+	void ReportAndNotifyErrorf(const FmtType& Fmt, Types... Args)
+	{
+		ReportAndNotifyError(FString::Printf(Fmt, Args...));
+	}
+
 	static FLinearColor GetColorFromMetadata(const FString& InMetadata);
 	static void CreateDefaultValueForStructIfRequired(UScriptStruct* InStruct, FString& OutDefaultValue);
 	static void PostProcessDefaultValue(URigVMPin* Pin, FString& OutDefaultValue);
diff --git a/Engine/Source/Developer/RigVMDeveloper/RigVMDeveloper.Build.cs b/Engine/Source/Developer/RigVMDeveloper/RigVMDeveloper.Build.cs
--- a/Engine/Source/Developer/RigVMDeveloper/RigVMDeveloper.Build.cs
+++ b/Engine/Source/Developer/RigVMDeveloper/RigVMDeveloper.Build.cs
@@ -22,8 +22,11 @@ public class RigVMDeveloper : ModuleRules
                 {
 					"EditorFramework",
                     "UnrealEd",
-                }
-            );
+					"Slate",
+					"SlateCore",
+					"EditorStyle",
+				}
+			);
         }
     }
 }
