Control Rig: Encapsulation - support for IO pins for expand operation

diff --git a/Engine/Source/Developer/RigVMDeveloper/Private/RigVMModel/RigVMController.cpp b/Engine/Source/Developer/RigVMDeveloper/Private/RigVMModel/RigVMController.cpp
--- a/Engine/Source/Developer/RigVMDeveloper/Private/RigVMModel/RigVMController.cpp
+++ b/Engine/Source/Developer/RigVMDeveloper/Private/RigVMModel/RigVMController.cpp
@@ -2393,7 +2393,8 @@ TArray<URigVMNode*> URigVMController::ExpandLibraryNode(URigVMLibraryNode* InNod
 		FString SourcePinPath;
 		FString TargetPinPath;
 
-		if (LibraryPin->GetDirection() == ERigVMPinDirection::Input)
+		if (LibraryPin->GetDirection() == ERigVMPinDirection::Input ||
+			LibraryPin->GetDirection() == ERigVMPinDirection::IO)
 		{
 			if (const TArray<FString>* LibraryPinLinksPtr = ToLibraryNode.Find(LibraryPinPath))
 			{
