Control Rig: Compiler to support node instruction markup for Collapse nodes

diff --git a/Engine/Source/Developer/RigVMDeveloper/Private/RigVMCompiler/RigVMCompiler.cpp b/Engine/Source/Developer/RigVMDeveloper/Private/RigVMCompiler/RigVMCompiler.cpp
--- a/Engine/Source/Developer/RigVMDeveloper/Private/RigVMCompiler/RigVMCompiler.cpp
+++ b/Engine/Source/Developer/RigVMDeveloper/Private/RigVMCompiler/RigVMCompiler.cpp
@@ -276,6 +276,22 @@ bool URigVMCompiler::Compile(URigVMGraph* InGraph, URigVMController* InControlle
 						}
 					}
 				}
+				else if (URigVMCollapseNode* CollapseNode = Cast<URigVMCollapseNode>(InNode))
+				{
+					int32 FirstInstructionIndex = INDEX_NONE;
+					for (URigVMNode* ContainedNode : CollapseNode->GetContainedNodes())
+					{
+						int32 ContainedInstructionIndex = GetInstructionIndex(ContainedNode, InOutKnownInstructionIndex);
+						if (ContainedInstructionIndex != INDEX_NONE)
+						{
+							if (FirstInstructionIndex == INDEX_NONE || ContainedInstructionIndex < FirstInstructionIndex)
+							{
+								FirstInstructionIndex = ContainedInstructionIndex;
+							}
+						}
+					}
+					return FirstInstructionIndex;
+				}
 
 				return INDEX_NONE;
 			}
