RigVM: Minor fixes for bugs found during nativization research
1. Select node's "IntEquals" operand should not be a literal since it stores the result of a runtime comparison
2. Register.CppType now renamed to Register.BaseCppType since it does not include TArray prefix for array type.
3. Added a few missed accessor/convertor functions for rig vm data types

diff --git a/Engine/Source/Developer/RigVMDeveloper/Private/RigVMCompiler/RigVMCompiler.cpp b/Engine/Source/Developer/RigVMDeveloper/Private/RigVMCompiler/RigVMCompiler.cpp
--- a/Engine/Source/Developer/RigVMDeveloper/Private/RigVMCompiler/RigVMCompiler.cpp
+++ b/Engine/Source/Developer/RigVMDeveloper/Private/RigVMCompiler/RigVMCompiler.cpp
@@ -1012,15 +1012,27 @@ void URigVMCompiler::TraverseSelect(const FRigVMSelectExprAST* InExpr, FRigVMCom
 			{
 				FName LiteralName = *FString::FromInt(CaseIndex);
 				int32 Register = WorkData.VM->GetLiteralMemory().Add<int32>(LiteralName, CaseIndex);
+ 
+#if WITH_EDITORONLY_DATA
+				WorkData.VM->GetLiteralMemory().Registers[Register].BaseCPPType = TEXT("int32");
+				WorkData.VM->GetLiteralMemory().Registers[Register].BaseCPPTypeObject = nullptr;
+#endif
+
 				FRigVMOperand Operand = WorkData.VM->GetLiteralMemory().GetOperand(Register);
 				WorkData.IntegerLiterals.Add(CaseIndex, Operand);
 			}
 		}
 
-		if (!WorkData.ComparisonLiteral.IsValid())
+		if (!WorkData.ComparisonOperand.IsValid())
 		{
-			int32 Register = WorkData.VM->GetLiteralMemory().Add<bool>(FName(TEXT("IntEquals")), false);
-			WorkData.ComparisonLiteral = WorkData.VM->GetLiteralMemory().GetOperand(Register);
+			int32 Register = WorkData.VM->GetWorkMemory().Add<bool>(FName(TEXT("IntEquals")), false);
+
+#if WITH_EDITORONLY_DATA
+			WorkData.VM->GetWorkMemory().Registers[Register].BaseCPPType = TEXT("bool");
+			WorkData.VM->GetWorkMemory().Registers[Register].BaseCPPTypeObject = nullptr;
+#endif
+
+			WorkData.ComparisonOperand = WorkData.VM->GetWorkMemory().GetOperand(Register);
 		}
 		return;
 	}
@@ -1060,13 +1072,13 @@ void URigVMCompiler::TraverseSelect(const FRigVMSelectExprAST* InExpr, FRigVMCom
 	for (int32 CaseIndex = 1; CaseIndex < NumCases; CaseIndex++)
 	{
 		// compare and jump eventually
-		WorkData.VM->GetByteCode().AddEqualsOp(IndexOperand, WorkData.IntegerLiterals.FindChecked(CaseIndex), WorkData.ComparisonLiteral);
+		WorkData.VM->GetByteCode().AddEqualsOp(IndexOperand, WorkData.IntegerLiterals.FindChecked(CaseIndex), WorkData.ComparisonOperand);
 		if (CaseIndex == 1 && Settings.SetupNodeInstructionIndex)
 		{
 			SelectNode->InstructionIndex = WorkData.VM->GetByteCode().GetNumInstructions() - 1;
 		}
 
-		uint64 JumpByte = WorkData.VM->GetByteCode().AddJumpIfOp(ERigVMOpCode::JumpForwardIf, 1, WorkData.ComparisonLiteral, true);
+		uint64 JumpByte = WorkData.VM->GetByteCode().AddJumpIfOp(ERigVMOpCode::JumpForwardIf, 1, WorkData.ComparisonOperand, true);
 		int32 JumpInstruction = WorkData.VM->GetByteCode().GetNumInstructions() - 1;
 		JumpToCaseBytes.Add(JumpByte);
 		JumpToCaseInstructions.Add(JumpInstruction);
@@ -1540,8 +1552,8 @@ FRigVMOperand URigVMCompiler::FindOrAddRegister(const FRigVMVarExprAST* InVarExp
 
 				int32 Register = Memory.AddRegisterArray(!Pin->IsDynamicArray() && !bIsDebugValue, RegisterName, ScriptStruct->GetStructureSize(), DefaultValues.Num(), Pin->IsArray(), Data.GetData(), NumSlices, ERigVMRegisterType::Struct, ScriptStruct);
 #if WITH_EDITORONLY_DATA
-				Memory.Registers[Register].CPPType = *ScriptStruct->GetStructCPPName();
-				Memory.Registers[Register].CPPTypeObject = ScriptStruct;
+				Memory.Registers[Register].BaseCPPType = *ScriptStruct->GetStructCPPName();
+				Memory.Registers[Register].BaseCPPTypeObject = ScriptStruct;
 #endif
 				ScriptStruct->DestroyStruct(Data.GetData(), DefaultValues.Num());
 
@@ -1564,8 +1576,8 @@ FRigVMOperand URigVMCompiler::FindOrAddRegister(const FRigVMVarExprAST* InVarExp
 
 				int32 Register = Memory.AddRegisterArray<uint8>(!Pin->IsDynamicArray() && !bIsDebugValue, RegisterName, Values.Num(), Pin->IsArray(), Values.GetData(), NumSlices, ERigVMRegisterType::Plain, nullptr);
 #if WITH_EDITORONLY_DATA
-				Memory.Registers[Register].CPPType = *Enum->CppType;
-				Memory.Registers[Register].CPPTypeObject = Enum;
+				Memory.Registers[Register].BaseCPPType = *Enum->CppType;
+				Memory.Registers[Register].BaseCPPTypeObject = Enum;
 #endif
 				Operand = Memory.GetOperand(Register);
 			}
@@ -1578,8 +1590,8 @@ FRigVMOperand URigVMCompiler::FindOrAddRegister(const FRigVMVarExprAST* InVarExp
 				}
 				int32 Register = Memory.AddRegisterArray<bool>(!Pin->IsDynamicArray() && !bIsDebugValue, RegisterName, Values.Num(), Pin->IsArray(), (uint8*)Values.GetData(), NumSlices, ERigVMRegisterType::Plain, nullptr);
 #if WITH_EDITORONLY_DATA
-				Memory.Registers[Register].CPPType = *BaseCPPType;
-				Memory.Registers[Register].CPPTypeObject = nullptr;
+				Memory.Registers[Register].BaseCPPType = *BaseCPPType;
+				Memory.Registers[Register].BaseCPPTypeObject = nullptr;
 #endif
 				Operand = Memory.GetOperand(Register);
 			}
@@ -1599,8 +1611,8 @@ FRigVMOperand URigVMCompiler::FindOrAddRegister(const FRigVMVarExprAST* InVarExp
 				}
 				int32 Register = Memory.AddRegisterArray<int32>(!Pin->IsDynamicArray() && !bIsDebugValue, RegisterName, Values.Num(), Pin->IsArray(), (uint8*)Values.GetData(), NumSlices, ERigVMRegisterType::Plain, nullptr);
 #if WITH_EDITORONLY_DATA
-				Memory.Registers[Register].CPPType = *BaseCPPType;
-				Memory.Registers[Register].CPPTypeObject = nullptr;
+				Memory.Registers[Register].BaseCPPType = *BaseCPPType;
+				Memory.Registers[Register].BaseCPPTypeObject = nullptr;
 #endif
 				Operand = Memory.GetOperand(Register);
 			}
@@ -1620,8 +1632,8 @@ FRigVMOperand URigVMCompiler::FindOrAddRegister(const FRigVMVarExprAST* InVarExp
 				}
 				int32 Register = Memory.AddRegisterArray<float>(!Pin->IsDynamicArray() && !bIsDebugValue, RegisterName, Values.Num(), Pin->IsArray(), (uint8*)Values.GetData(), NumSlices, ERigVMRegisterType::Plain, nullptr);
 #if WITH_EDITORONLY_DATA
-				Memory.Registers[Register].CPPType = *BaseCPPType;
-				Memory.Registers[Register].CPPTypeObject = nullptr;
+				Memory.Registers[Register].BaseCPPType = *BaseCPPType;
+				Memory.Registers[Register].BaseCPPTypeObject = nullptr;
 #endif
 				Operand = Memory.GetOperand(Register);
 			}
@@ -1634,8 +1646,8 @@ FRigVMOperand URigVMCompiler::FindOrAddRegister(const FRigVMVarExprAST* InVarExp
 				}
 				int32 Register = Memory.AddRegisterArray<FName>(!Pin->IsDynamicArray() && !bIsDebugValue, RegisterName, Values.Num(), Pin->IsArray(), (uint8*)Values.GetData(), NumSlices, ERigVMRegisterType::Name, nullptr);
 #if WITH_EDITORONLY_DATA
-				Memory.Registers[Register].CPPType = *BaseCPPType;
-				Memory.Registers[Register].CPPTypeObject = nullptr;
+				Memory.Registers[Register].BaseCPPType = *BaseCPPType;
+				Memory.Registers[Register].BaseCPPTypeObject = nullptr;
 #endif
 				Operand = Memory.GetOperand(Register);
 			}
@@ -1643,8 +1655,8 @@ FRigVMOperand URigVMCompiler::FindOrAddRegister(const FRigVMVarExprAST* InVarExp
 			{
 				int32 Register = Memory.AddRegisterArray<FString>(!Pin->IsDynamicArray() && !bIsDebugValue, RegisterName, DefaultValues.Num(), Pin->IsArray(), (uint8*)DefaultValues.GetData(), NumSlices, ERigVMRegisterType::String, nullptr);
 #if WITH_EDITORONLY_DATA
-				Memory.Registers[Register].CPPType = *BaseCPPType;
-				Memory.Registers[Register].CPPTypeObject = nullptr;
+				Memory.Registers[Register].BaseCPPType = *BaseCPPType;
+				Memory.Registers[Register].BaseCPPTypeObject = nullptr;
 #endif
 				Operand = Memory.GetOperand(Register);
 			}
diff --git a/Engine/Source/Developer/RigVMDeveloper/Public/RigVMCompiler/RigVMCompiler.h b/Engine/Source/Developer/RigVMDeveloper/Public/RigVMCompiler/RigVMCompiler.h
--- a/Engine/Source/Developer/RigVMDeveloper/Public/RigVMCompiler/RigVMCompiler.h
+++ b/Engine/Source/Developer/RigVMDeveloper/Public/RigVMCompiler/RigVMCompiler.h
@@ -80,7 +80,7 @@ public:
 	TMap<int32, int32> RegisterRefCount;
 	TArray<TPair<FName, int32>> RefCountSteps;
 	TMap<int32, FRigVMOperand> IntegerLiterals;
-	FRigVMOperand ComparisonLiteral;
+	FRigVMOperand ComparisonOperand;
 
 	int32 IncRefRegister(int32 InRegister, int32 InIncrement = 1);
 	int32 DecRefRegister(int32 InRegister, int32 InDecrement = 1);
diff --git a/Engine/Source/Runtime/RigVM/Public/RigVMCore/RigVMArray.h b/Engine/Source/Runtime/RigVM/Public/RigVMCore/RigVMArray.h
--- a/Engine/Source/Runtime/RigVM/Public/RigVMCore/RigVMArray.h
+++ b/Engine/Source/Runtime/RigVM/Public/RigVMCore/RigVMArray.h
@@ -322,6 +322,14 @@ public:
 	{
 	}
 
+	// constructor from a typed TArray
+	template<uint32 NumInlineElements>
+	FORCEINLINE_DEBUGGABLE FRigVMFixedArray(const TArray<T, TFixedAllocator<NumInlineElements>>& InStorage)
+		: Data((T*)InStorage.GetData())
+		, Size(NumInlineElements)
+	{
+	}
+
 	// constructor from direct memory
 	FORCEINLINE_DEBUGGABLE FRigVMFixedArray(T* InData, int32 InSize)
 	: Data(InData)
diff --git a/Engine/Source/Runtime/RigVM/Public/RigVMCore/RigVMExternalVariable.h b/Engine/Source/Runtime/RigVM/Public/RigVMCore/RigVMExternalVariable.h
--- a/Engine/Source/Runtime/RigVM/Public/RigVMCore/RigVMExternalVariable.h
+++ b/Engine/Source/Runtime/RigVM/Public/RigVMCore/RigVMExternalVariable.h
@@ -378,6 +378,13 @@ struct RIGVM_API FRigVMExternalVariable
 		return *(T*)Memory;
 	}
 
+	template<typename T>
+	FORCEINLINE T& GetRef()
+	{
+		ensure(IsValid() && !bIsArray);
+		return *(T*)Memory;
+	}
+
 	template<typename T>
 	FORCEINLINE void SetValue(const T& InValue)
 	{
diff --git a/Engine/Source/Runtime/RigVM/Public/RigVMCore/RigVMMemory.h b/Engine/Source/Runtime/RigVM/Public/RigVMCore/RigVMMemory.h
--- a/Engine/Source/Runtime/RigVM/Public/RigVMCore/RigVMMemory.h
+++ b/Engine/Source/Runtime/RigVM/Public/RigVMCore/RigVMMemory.h
@@ -89,7 +89,7 @@ public:
 	// returns the index of the register of this argument
 	FORCEINLINE_DEBUGGABLE int32 GetRegisterIndex() const { return RegisterIndex == UINT16_MAX ? INDEX_NONE : (int32)RegisterIndex; }
 
-	// returns the index of the register of this argument
+	// returns the register offset of this argument
 	FORCEINLINE_DEBUGGABLE int32 GetRegisterOffset() const { return RegisterOffset == UINT16_MAX ? INDEX_NONE : (int32)RegisterOffset; }
 
 	void Serialize(FArchive& Ar);
@@ -158,8 +158,8 @@ struct RIGVM_API FRigVMRegister
 		, bIsArray(false)
 		, bIsDynamic(false)
 #if WITH_EDITORONLY_DATA
-		, CPPType(NAME_None)
-		, CPPTypeObject(nullptr)
+		, BaseCPPType(NAME_None)
+		, BaseCPPTypeObject(nullptr)
 #endif
 	{
 	}
@@ -214,10 +214,10 @@ struct RIGVM_API FRigVMRegister
 
 	// Defines the CPP type used for the register
 	UPROPERTY(transient)
-	FName CPPType;
+	FName BaseCPPType;
 
 	UPROPERTY(transient)
-	UObject* CPPTypeObject;
+	UObject* BaseCPPTypeObject;
 
 #endif
 
