Control Rig: Use the node path on the RigVMNode for the uniqueness of register paths

diff --git a/Engine/Source/Developer/RigVMDeveloper/Private/RigVMCompiler/RigVMCompiler.cpp b/Engine/Source/Developer/RigVMDeveloper/Private/RigVMCompiler/RigVMCompiler.cpp
--- a/Engine/Source/Developer/RigVMDeveloper/Private/RigVMCompiler/RigVMCompiler.cpp
+++ b/Engine/Source/Developer/RigVMDeveloper/Private/RigVMCompiler/RigVMCompiler.cpp
@@ -1175,7 +1175,7 @@ FString URigVMCompiler::GetPinHash(URigVMPin* InPin, const FRigVMVarExprAST* InV
 			return FString::Printf(TEXT("%sVariable::%s%s"), *Prefix, *VariableNode->GetVariableName().ToString(), *Suffix);
 		}
 	}
-	return FString::Printf(TEXT("%s%s%s"), *Prefix, *InPin->GetPinPath(), *Suffix);
+	return FString::Printf(TEXT("%s%s%s"), *Prefix, *InPin->GetPinPath(true /*  full node path */), *Suffix);
 }
 
 UScriptStruct* URigVMCompiler::GetScriptStructForCPPType(const FString& InCPPType)
diff --git a/Engine/Source/Developer/RigVMDeveloper/Private/RigVMModel/RigVMGraph.cpp b/Engine/Source/Developer/RigVMDeveloper/Private/RigVMModel/RigVMGraph.cpp
--- a/Engine/Source/Developer/RigVMDeveloper/Private/RigVMModel/RigVMGraph.cpp
+++ b/Engine/Source/Developer/RigVMDeveloper/Private/RigVMModel/RigVMGraph.cpp
@@ -64,6 +64,15 @@ TArray<FRigVMGraphParameterDescription> URigVMGraph::GetParameterDescriptions()
 	return Parameters;
 }
 
+FString URigVMGraph::GetNodePath() const
+{
+	if (URigVMCollapseNode* CollapseNode = Cast<URigVMCollapseNode>(GetOuter()))
+	{
+		return CollapseNode->GetNodePath();
+	}
+	return FString();
+}
+
 URigVMNode* URigVMGraph::FindNodeByName(const FName& InNodeName) const
 {
 	for (URigVMNode* Node : Nodes)
diff --git a/Engine/Source/Developer/RigVMDeveloper/Private/RigVMModel/RigVMNode.cpp b/Engine/Source/Developer/RigVMDeveloper/Private/RigVMModel/RigVMNode.cpp
--- a/Engine/Source/Developer/RigVMDeveloper/Private/RigVMModel/RigVMNode.cpp
+++ b/Engine/Source/Developer/RigVMDeveloper/Private/RigVMModel/RigVMNode.cpp
@@ -23,8 +23,16 @@ URigVMNode::~URigVMNode()
 {
 }
 
-FString URigVMNode::GetNodePath() const
+FString URigVMNode::GetNodePath(bool bRecursive) const
 {
+	if (bRecursive)
+	{
+		FString ParentNodePath = GetGraph()->GetNodePath();
+		if (!ParentNodePath.IsEmpty())
+		{
+			return FString::Printf(TEXT("%s|%s"), *ParentNodePath, *GetName());
+		}
+	}
 	return GetName();
 }
 
diff --git a/Engine/Source/Developer/RigVMDeveloper/Private/RigVMModel/RigVMPin.cpp b/Engine/Source/Developer/RigVMDeveloper/Private/RigVMModel/RigVMPin.cpp
--- a/Engine/Source/Developer/RigVMDeveloper/Private/RigVMModel/RigVMPin.cpp
+++ b/Engine/Source/Developer/RigVMDeveloper/Private/RigVMModel/RigVMPin.cpp
@@ -89,7 +89,7 @@ URigVMPin::URigVMPin()
 #endif
 }
 
-FString URigVMPin::GetPinPath() const
+FString URigVMPin::GetPinPath(bool bUseNodePath) const
 {
 	FString PinPath;
 
@@ -103,7 +103,7 @@ FString URigVMPin::GetPinPath() const
 		URigVMNode* Node = GetNode();
 		if (Node != nullptr)
 		{
-			PinPath = FString::Printf(TEXT("%s.%s"), *Node->GetNodePath(), *GetName());
+			PinPath = FString::Printf(TEXT("%s.%s"), *Node->GetNodePath(bUseNodePath), *GetName());
 		}
 	}
 #if UE_BUILD_DEBUG
diff --git a/Engine/Source/Developer/RigVMDeveloper/Public/RigVMModel/RigVMGraph.h b/Engine/Source/Developer/RigVMDeveloper/Public/RigVMModel/RigVMGraph.h
--- a/Engine/Source/Developer/RigVMDeveloper/Public/RigVMModel/RigVMGraph.h
+++ b/Engine/Source/Developer/RigVMDeveloper/Public/RigVMModel/RigVMGraph.h
@@ -50,6 +50,9 @@ public:
 	UFUNCTION(BlueprintCallable, Category = RigVMGraph)
 	TArray<FRigVMGraphParameterDescription> GetParameterDescriptions() const;
 
+	// Returns the path of this graph as defined by its invoking nodes
+	FString GetNodePath() const;
+
 	// Returns a Node given its name (or nullptr).
 	UFUNCTION(BlueprintCallable, Category = RigVMGraph)
 	URigVMNode* FindNodeByName(const FName& InNodeName) const;
diff --git a/Engine/Source/Developer/RigVMDeveloper/Public/RigVMModel/RigVMNode.h b/Engine/Source/Developer/RigVMDeveloper/Public/RigVMModel/RigVMNode.h
--- a/Engine/Source/Developer/RigVMDeveloper/Public/RigVMModel/RigVMNode.h
+++ b/Engine/Source/Developer/RigVMDeveloper/Public/RigVMModel/RigVMNode.h
@@ -33,7 +33,7 @@ public:
 	// names used to reach this Node within the Graph.
 	// (for now this is the same as the Node's name)
 	UFUNCTION(BlueprintCallable, Category = RigVMNode)
-	FString GetNodePath() const;
+	FString GetNodePath(bool bRecursive = false) const;
 
 	// Returns the current index of the Node within the Graph.
 	UFUNCTION(BlueprintCallable, Category = RigVMNode)
diff --git a/Engine/Source/Developer/RigVMDeveloper/Public/RigVMModel/RigVMPin.h b/Engine/Source/Developer/RigVMDeveloper/Public/RigVMModel/RigVMPin.h
--- a/Engine/Source/Developer/RigVMDeveloper/Public/RigVMModel/RigVMPin.h
+++ b/Engine/Source/Developer/RigVMDeveloper/Public/RigVMModel/RigVMPin.h
@@ -101,7 +101,7 @@ public:
 	// Returns a . separated path containing all names of the pin and its owners,
 	// this includes the node name, for example "Node.Color.R"
 	UFUNCTION(BlueprintCallable, Category = RigVMPin)
-	FString GetPinPath() const;
+	FString GetPinPath(bool bUseNodePath = false) const;
 
 	// Returns a . separated path containing all names of the pin within its main
 	// memory owner / storage. This is typically used to create an offset pointer
