Control Rig: Fix missing include for non-unity build

diff --git a/Engine/Source/Developer/RigVMDeveloper/Private/RigVMModel/RigVMController.cpp b/Engine/Source/Developer/RigVMDeveloper/Private/RigVMModel/RigVMController.cpp
--- a/Engine/Source/Developer/RigVMDeveloper/Private/RigVMModel/RigVMController.cpp
+++ b/Engine/Source/Developer/RigVMDeveloper/Private/RigVMModel/RigVMController.cpp
@@ -19,6 +19,7 @@
 #include "UObject/CoreRedirects.h"
 #include "Framework/Notifications/NotificationManager.h"
 #include "Widgets/Notifications/SNotificationList.h"
+#include "EditorStyle.h"
 #endif
 
 TMap<URigVMController::FControlRigStructPinRedirectorKey, FString> URigVMController::PinPathCoreRedirectors;
