Control Rig Encapsulation - UI changes

diff --git a/Engine/Plugins/Experimental/ControlRig/Source/ControlRigDeveloper/Private/ControlRigBlueprint.cpp b/Engine/Plugins/Experimental/ControlRig/Source/ControlRigDeveloper/Private/ControlRigBlueprint.cpp
--- a/Engine/Plugins/Experimental/ControlRig/Source/ControlRigDeveloper/Private/ControlRigBlueprint.cpp
+++ b/Engine/Plugins/Experimental/ControlRig/Source/ControlRigDeveloper/Private/ControlRigBlueprint.cpp
@@ -10,6 +10,7 @@
 #include "ControlRig.h"
 #include "Graph/ControlRigGraph.h"
 #include "Graph/ControlRigGraphNode.h"
+#include "Graph/ControlRigGraphSchema.h"
 #include "UObject/UObjectGlobals.h"
 #include "ControlRigObjectVersion.h"
 #include "ControlRigDeveloper.h"
@@ -376,15 +377,54 @@ URigVMGraph* UControlRigBlueprint::GetModel(const UEdGraph* InEdGraph) const
 		return Model;
 	}
 
-	ensure(false);
-	return nullptr;
+	FString ModelNodePath = RigGraph->ModelNodePath;
+	ensure(!ModelNodePath.IsEmpty());
+
+	URigVMGraph* SubModel = Model;
+	while (!ModelNodePath.IsEmpty())
+	{
+		FString NodeName = ModelNodePath;
+		if (NodeName.Contains(TEXT("|")))
+		{
+			NodeName = NodeName.Left(NodeName.Find(TEXT("|")));
+			ModelNodePath = ModelNodePath.Right(ModelNodePath.Len() - NodeName.Len() - 1);
+		}
+		else
+		{
+			ModelNodePath.Reset();
+		}
+
+		URigVMCollapseNode* CollapseNode = Cast<URigVMCollapseNode>(SubModel->FindNodeByName(*NodeName));
+		if (CollapseNode == nullptr)
+		{
+			return nullptr;
+		}
+
+		SubModel = CollapseNode->GetContainedGraph();
+	}
+
+	return SubModel;
 }
 
+URigVMGraph* UControlRigBlueprint::GetModel(const FString& InNodePath) const
+{
+	if (!InNodePath.IsEmpty())
+	{
+		if (URigVMLibraryNode* LibraryNode = Cast<URigVMLibraryNode>(Model->FindNode(InNodePath)))
+		{
+			return LibraryNode->GetContainedGraph();
+		}
+		return nullptr;
+	}
+	return Model;
+}
+
+
 TArray<URigVMGraph*> UControlRigBlueprint::GetAllModels() const
 {
 	TArray<URigVMGraph*> Models;
 	Models.Add(GetModel());
-	Models.Append(GetModel()->GetContainedGraphs());
+	Models.Append(GetModel()->GetContainedGraphs(true /* recursive */));
 	return Models;
 }
 
@@ -488,6 +528,64 @@ URigVMController* UControlRigBlueprint::GetOrCreateController(const UEdGraph* In
 	return GetOrCreateController(GetModel(InEdGraph));
 }
 
+URigVMGraph* UControlRigBlueprint::GetTemplateModel()
+{
+#if WITH_EDITORONLY_DATA
+	if (TemplateModel == nullptr)
+	{
+		TemplateModel = NewObject<URigVMGraph>(this, TEXT("TemplateModel"));
+		TemplateModel->SetFlags(RF_Transient);
+	}
+	return TemplateModel;
+#else
+	return nullptr;
+#endif
+}
+
+URigVMController* UControlRigBlueprint::GetTemplateController()
+{
+#if WITH_EDITORONLY_DATA
+	if (TemplateController == nullptr)
+	{
+		TemplateController = NewObject<URigVMController>(this, TEXT("TemplateController"));
+		TemplateController->SetExecuteContextStruct(FControlRigExecuteContext::StaticStruct());
+		TemplateController->SetGraph(GetTemplateModel());
+		TemplateController->EnableReporting(false);
+		TemplateController->SetFlags(RF_Transient);
+	}
+	return TemplateController;
+#else
+	return nullptr;
+#endif
+}
+
+UEdGraph* UControlRigBlueprint::GetEdGraph(URigVMGraph* InModel) const
+{
+	TArray<UEdGraph*> EdGraphs;
+	GetAllGraphs(EdGraphs);
+
+	for (UEdGraph* EdGraph : EdGraphs)
+	{
+		if (UControlRigGraph* RigGraph = Cast<UControlRigGraph>(EdGraph))
+		{
+			if (RigGraph->ModelNodePath == InModel->GetNodePath())
+			{
+				return RigGraph;
+			}
+		}
+	}
+	return nullptr;
+}
+
+UEdGraph* UControlRigBlueprint::GetEdGraph(const FString& InNodePath) const
+{
+	if (URigVMGraph* ModelForNodePath = GetModel(InNodePath))
+	{
+		return GetEdGraph(ModelForNodePath);
+	}
+	return nullptr;
+}
+
 void UControlRigBlueprint::GetTypeActions(FBlueprintActionDatabaseRegistrar& ActionRegistrar) const
 {
 	DECLARE_SCOPE_HIERARCHICAL_COUNTER_FUNC()
@@ -565,6 +663,22 @@ void UControlRigBlueprint::PostTransacted(const FTransactionObjectEvent& Transac
 	}
 }
 
+void UControlRigBlueprint::ReplaceDeprecatedNodes()
+{
+	TArray<UEdGraph*> EdGraphs;
+	GetAllGraphs(EdGraphs);
+
+	for (UEdGraph* EdGraph : EdGraphs)
+	{
+		if (UControlRigGraph* RigGraph = Cast<UControlRigGraph>(EdGraph))
+		{
+			RigGraph->Schema = UControlRigGraphSchema::StaticClass();
+		}
+	}
+
+	Super::ReplaceDeprecatedNodes();
+}
+
 FRigVMGraphModifiedEvent& UControlRigBlueprint::OnModified()
 {
 	return ModifiedEvent;
@@ -1171,7 +1285,10 @@ void UControlRigBlueprint::RebuildGraphFromModel()
 	TGuardValue<bool> SelfGuard(bSuspendModelNotificationsForSelf, true);
 	check(GetOrCreateController());
 
-	for (UEdGraph* Graph : UbergraphPages)
+	TArray<UEdGraph*> EdGraphs;
+	GetAllGraphs(EdGraphs);
+
+	for (UEdGraph* Graph : EdGraphs)
 	{
 		TArray<UEdGraphNode*> Nodes = Graph->Nodes;
 		for (UEdGraphNode* Node : Nodes)
@@ -1180,7 +1297,24 @@ void UControlRigBlueprint::RebuildGraphFromModel()
 		}
 	}
 
-	GetController()->ResendAllNotifications();
+	TArray<URigVMGraph*> RigGraphs;
+	RigGraphs.Add(GetModel());
+
+	GetOrCreateController(RigGraphs[0])->ResendAllNotifications();
+
+	for (int32 RigGraphIndex = 0; RigGraphIndex < RigGraphs.Num(); RigGraphIndex++)
+	{
+		URigVMGraph* RigGraph = RigGraphs[RigGraphIndex];
+
+		for (URigVMNode* RigNode : RigGraph->GetNodes())
+		{
+			if (URigVMCollapseNode* CollapseNode = Cast<URigVMCollapseNode>(RigNode))
+			{
+				CreateEdGraphForCollapseNodeIfNeeded(CollapseNode, true);
+				RigGraphs.Add(CollapseNode->GetContainedGraph());
+			}
+		}
+	}
 }
 
 void UControlRigBlueprint::Notify(ERigVMGraphNotifType InNotifType, UObject* InSubject)
@@ -1199,6 +1333,11 @@ void UControlRigBlueprint::HandleModifiedEvent(ERigVMGraphNotifType InNotifType,
 		return;
 	}
 
+	// since it's possible that a notification will be already sent / forwarded to the
+	// listening objects within the switch statement below - we keep a flag to mark
+	// the notify for still pending (or already sent)
+	bool bNotifForOthersPending = true;
+
 	if (!bSuspendModelNotificationsForSelf)
 	{
 		switch (InNotifType)
@@ -1328,6 +1467,27 @@ void UControlRigBlueprint::HandleModifiedEvent(ERigVMGraphNotifType InNotifType,
 			}
 			case ERigVMGraphNotifType::NodeAdded:
 			case ERigVMGraphNotifType::NodeRemoved:
+			{
+				if (URigVMCollapseNode* CollapseNode = Cast<URigVMCollapseNode>(InSubject))
+				{
+					if (InNotifType == ERigVMGraphNotifType::NodeAdded)
+					{
+						CreateEdGraphForCollapseNodeIfNeeded(CollapseNode);
+					}
+					else
+					{
+						bNotifForOthersPending = !RemoveEdGraphForCollapseNode(CollapseNode, true);
+					}
+
+					ClearTransientControls();
+					RequestAutoVMRecompilation();
+					MarkPackageDirty();
+					FBlueprintEditorUtils::MarkBlueprintAsStructurallyModified(this);
+					break;
+				}
+
+				// fall through to the next case
+			}
 			case ERigVMGraphNotifType::LinkAdded:
 			case ERigVMGraphNotifType::LinkRemoved:
 			case ERigVMGraphNotifType::PinArraySizeChanged:
@@ -1336,9 +1496,6 @@ void UControlRigBlueprint::HandleModifiedEvent(ERigVMGraphNotifType InNotifType,
 				ClearTransientControls();
 				RequestAutoVMRecompilation();
 				MarkPackageDirty();
-
-				// This is not necessarily required but due to workflow
-				// expectations we still mark the blueprint as dirty.
 				FBlueprintEditorUtils::MarkBlueprintAsModified(this);
 				break;
 			}
@@ -1381,6 +1538,44 @@ void UControlRigBlueprint::HandleModifiedEvent(ERigVMGraphNotifType InNotifType,
 				MarkPackageDirty();
 				break;
 			}
+			case ERigVMGraphNotifType::NodeRenamed:
+			{
+				if (URigVMLibraryNode* LibraryNode = Cast<URigVMLibraryNode>(InSubject))
+				{
+					FString NewNodePath = LibraryNode->GetNodePath(true /* recursive */);
+					FString Left, Right = NewNodePath;
+					URigVMNode::SplitNodePathAtEnd(NewNodePath, Left, Right);
+					FString OldNodePath = LibraryNode->GetPreviousFName().ToString();
+					if (!Left.IsEmpty())
+					{
+						OldNodePath = URigVMNode::JoinNodePath(Left, OldNodePath);
+					}
+
+					FString NewNodePathPrefix = NewNodePath + TEXT("|");
+					FString OldNodePathPrefix = OldNodePath + TEXT("|");
+
+					TArray<UEdGraph*> EdGraphs;
+					GetAllGraphs(EdGraphs);
+
+					for (UEdGraph* EdGraph : EdGraphs)
+					{
+						if (UControlRigGraph* RigGraph = Cast<UControlRigGraph>(EdGraph))
+						{
+							if (RigGraph->ModelNodePath == OldNodePath)
+							{
+								RigGraph->ModelNodePath = NewNodePath;
+							}
+							else if (RigGraph->ModelNodePath.StartsWith(OldNodePathPrefix))
+							{
+								RigGraph->ModelNodePath = NewNodePathPrefix + RigGraph->ModelNodePath.LeftChop(OldNodePathPrefix.Len());
+							}
+						}
+					}
+
+					FBlueprintEditorUtils::MarkBlueprintAsStructurallyModified(this);
+				}
+				break;
+			}
 			default:
 			{
 				break;
@@ -1388,7 +1583,8 @@ void UControlRigBlueprint::HandleModifiedEvent(ERigVMGraphNotifType InNotifType,
 		}
 	}
 
-	if (!bSuspendModelNotificationsForOthers)
+	// if the notification still has to be sent...
+	if (bNotifForOthersPending && !bSuspendModelNotificationsForOthers)
 	{
 		if (ModifiedEvent.IsBound())
 		{
@@ -2112,34 +2308,46 @@ void UControlRigBlueprint::OnVariableAdded(const FName& InVarName)
 
 void UControlRigBlueprint::OnVariableRemoved(const FName& InVarName)
 {
-	if (URigVMController* Controller = GetController())
+	TArray<URigVMGraph*> AllGraphs = GetAllModels();
+	for (URigVMGraph* Graph : AllGraphs)
 	{
-		Controller->OnExternalVariableRemoved(InVarName, true);
+		if (URigVMController* Controller = GetOrCreateController(Graph))
+		{
+			Controller->OnExternalVariableRemoved(InVarName, true);
+		}
 	}
 	BroadcastExternalVariablesChangedEvent();
 }
 
 void UControlRigBlueprint::OnVariableRenamed(const FName& InOldVarName, const FName& InNewVarName)
 {
-	if (URigVMController* Controller = GetController())
+	TArray<URigVMGraph*> AllGraphs = GetAllModels();
+	for (URigVMGraph* Graph : AllGraphs)
 	{
-		Controller->OnExternalVariableRenamed(InOldVarName, InNewVarName, true);
+		if (URigVMController* Controller = GetOrCreateController(Graph))
+		{
+			Controller->OnExternalVariableRenamed(InOldVarName, InNewVarName, true);
+		}
 	}
 	BroadcastExternalVariablesChangedEvent();
 }
 
 void UControlRigBlueprint::OnVariableTypeChanged(const FName& InVarName, FEdGraphPinType InOldPinType, FEdGraphPinType InNewPinType)
 {
-	if (URigVMController* Controller = GetController())
+	TArray<URigVMGraph*> AllGraphs = GetAllModels();
+	for (URigVMGraph* Graph : AllGraphs)
 	{
-		FRigVMExternalVariable NewVariable = UControlRig::GetExternalVariableFromPinType(InVarName, InNewPinType);
-		if (NewVariable.IsValid(true)) // allow nullptr
+		if (URigVMController* Controller = GetOrCreateController(Graph))
 		{
-			Controller->OnExternalVariableTypeChanged(InVarName, NewVariable.TypeName.ToString(), NewVariable.TypeObject, true);
-		}
-		else
-		{
-			Controller->OnExternalVariableRemoved(InVarName, true);
+			FRigVMExternalVariable NewVariable = UControlRig::GetExternalVariableFromPinType(InVarName, InNewPinType);
+			if (NewVariable.IsValid(true)) // allow nullptr
+			{
+				Controller->OnExternalVariableTypeChanged(InVarName, NewVariable.TypeName.ToString(), NewVariable.TypeObject, true);
+			}
+			else
+			{
+				Controller->OnExternalVariableRemoved(InVarName, true);
+			}
 		}
 	}
 	BroadcastExternalVariablesChangedEvent();
@@ -2156,7 +2364,94 @@ void UControlRigBlueprint::BroadcastExternalVariablesChangedEvent()
 	}
 }
 
+void UControlRigBlueprint::BroadcastNodeDoubleClicked(URigVMNode* InNode)
+{
+	NodeDoubleClickedEvent.Broadcast(this, InNode);
+}
+
 #endif
 
+void UControlRigBlueprint::CreateEdGraphForCollapseNodeIfNeeded(URigVMCollapseNode* InNode, bool bForce)
+{
+	check(InNode);
+
+	if (bForce)
+	{
+		RemoveEdGraphForCollapseNode(InNode, false);
+	}
+
+	if (UControlRigGraph* RigGraph = Cast<UControlRigGraph>(GetEdGraph(InNode->GetGraph())))
+	{
+		if (URigVMGraph* ContainedGraph = InNode->GetContainedGraph())
+		{
+			bool bSubGraphExists = false;
+			for (UEdGraph* SubGraph : RigGraph->SubGraphs)
+			{
+				if (UControlRigGraph* SubRigGraph = Cast<UControlRigGraph>(SubGraph))
+				{
+					if (SubRigGraph->ModelNodePath == ContainedGraph->GetNodePath())
+					{
+						bSubGraphExists = true;
+						break;
+					}
+				}
+			}
+
+			if (!bSubGraphExists)
+			{
+				// create a sub graph
+				UControlRigGraph* SubRigGraph = NewObject<UControlRigGraph>(RigGraph, *InNode->GetEditorSubGraphName(), RF_Transactional);
+				SubRigGraph->Schema = UControlRigGraphSchema::StaticClass();
+				SubRigGraph->bAllowRenaming = 1;
+				SubRigGraph->bEditable = 1;
+				SubRigGraph->bAllowDeletion = 1;
+				SubRigGraph->ModelNodePath = InNode->GetNodePath(true);
+
+				RigGraph->SubGraphs.Add(SubRigGraph);
+
+				SubRigGraph->Initialize(this);
+
+				GetOrCreateController(ContainedGraph)->ResendAllNotifications();
+			}
+		}
+	}
+
+}
+
+bool UControlRigBlueprint::RemoveEdGraphForCollapseNode(URigVMCollapseNode* InNode, bool bNotify)
+{
+	check(InNode);
+
+	if (UControlRigGraph* RigGraph = Cast<UControlRigGraph>(GetEdGraph(InNode->GetGraph())))
+	{
+		if (URigVMGraph* ContainedGraph = InNode->GetContainedGraph())
+		{
+			for (UEdGraph* SubGraph : RigGraph->SubGraphs)
+			{
+				if (UControlRigGraph* SubRigGraph = Cast<UControlRigGraph>(SubGraph))
+				{
+					if (SubRigGraph->ModelNodePath == ContainedGraph->GetNodePath())
+					{
+						if (URigVMController* SubController = GetController(ContainedGraph))
+						{
+							SubController->OnModified().RemoveAll(SubRigGraph);
+						}
+
+						if (ModifiedEvent.IsBound() && bNotify)
+						{
+							ModifiedEvent.Broadcast(ERigVMGraphNotifType::NodeRemoved, InNode->GetGraph(), InNode);
+						}
+
+						RigGraph->SubGraphs.Remove(SubRigGraph);
+						return bNotify;
+					}
+				}
+			}
+		}
+	}
+
+	return false;
+}
+
 #undef LOCTEXT_NAMESPACE
 
diff --git a/Engine/Plugins/Experimental/ControlRig/Source/ControlRigDeveloper/Private/Graph/ControlRigGraph.cpp b/Engine/Plugins/Experimental/ControlRig/Source/ControlRigDeveloper/Private/Graph/ControlRigGraph.cpp
--- a/Engine/Plugins/Experimental/ControlRig/Source/ControlRigDeveloper/Private/Graph/ControlRigGraph.cpp
+++ b/Engine/Plugins/Experimental/ControlRig/Source/ControlRigDeveloper/Private/Graph/ControlRigGraph.cpp
@@ -47,6 +47,11 @@ void UControlRigGraph::Serialize(FArchive& Ar)
 	Super::Serialize(Ar);
 
 	Ar.UsingCustomVersion(FControlRigObjectVersion::GUID);
+
+	if (Ar.IsLoading())
+	{
+		Schema = UControlRigGraphSchema::StaticClass();
+	}
 }
 #endif
 
@@ -147,6 +152,13 @@ void UControlRigGraph::HandleModifiedEvent(ERigVMGraphNotifType InNotifType, URi
 		return;
 	}
 
+	// only make sure to receive notifs for this graph - unless
+	// we are on a template graph (used by node spawners)
+	if (GetModel() != InGraph && TemplateController == nullptr)
+	{
+		return;
+	}
+
 	switch (InNotifType)
 	{
 		case ERigVMGraphNotifType::GraphChanged:
@@ -261,14 +273,6 @@ void UControlRigGraph::HandleModifiedEvent(ERigVMGraphNotifType InNotifType, URi
 					NewNode->SetColorFromModel(ModelNode->GetNodeColor());
 					NewNode->SetFlags(RF_Transactional);
 					NewNode->AllocateDefaultPins();
-
-					/*
-					// if we need to setup a sub graph.. we should
-					if (URigVMCollapseNode* CollapseNode = Cast<URigVMCollapseNode>(ModelNode))
-					{
-
-					}
-					*/
 				}
 			}
 			break;
@@ -389,8 +393,14 @@ void UControlRigGraph::HandleModifiedEvent(ERigVMGraphNotifType InNotifType, URi
 				URigVMPin* SourcePin = Link->GetSourcePin();
 				URigVMPin* TargetPin = Link->GetTargetPin();
 
-				SourcePin = SourcePin->GetOriginalPinFromInjectedNode();
-				TargetPin = TargetPin->GetOriginalPinFromInjectedNode();
+				if (SourcePin)
+				{
+					SourcePin = SourcePin->GetOriginalPinFromInjectedNode();
+				}
+				if (TargetPin)
+				{
+					TargetPin = TargetPin->GetOriginalPinFromInjectedNode();
+				}
 
 				if (SourcePin && TargetPin && SourcePin != TargetPin)
 				{
@@ -507,6 +517,9 @@ void UControlRigGraph::HandleModifiedEvent(ERigVMGraphNotifType InNotifType, URi
 		case ERigVMGraphNotifType::PinDirectionChanged:
 		case ERigVMGraphNotifType::PinTypeChanged:
 		case ERigVMGraphNotifType::PinBoundVariableChanged:
+		case ERigVMGraphNotifType::PinAdded:
+		case ERigVMGraphNotifType::PinRemoved:
+		case ERigVMGraphNotifType::PinRenamed:
 		{
 			if (URigVMPin* ModelPin = Cast<URigVMPin>(InSubject))
 			{
@@ -517,6 +530,20 @@ void UControlRigGraph::HandleModifiedEvent(ERigVMGraphNotifType InNotifType, URi
 			}
 			break;
 		}
+		case ERigVMGraphNotifType::NodeRenamed:
+		{
+			if (URigVMNode* ModelNode = Cast<URigVMNode>(InSubject))
+			{
+				if (UControlRigGraphNode* RigNode = Cast<UControlRigGraphNode>(FindNodeForModelNodeName(ModelNode->GetPreviousFName())))
+				{
+					RigNode->Rename(*ModelNode->GetName());
+					RigNode->ModelNodePath = ModelNode->GetNodePath();
+					RigNode->InvalidateNodeTitle();
+					RigNode->ReconstructNode_Internal(true);
+				}
+			}
+			break;
+		}
 		case ERigVMGraphNotifType::VariableRenamed:
 		{
 			if (URigVMNode* ModelNode = Cast<URigVMNode>(InSubject))
@@ -615,20 +642,11 @@ URigVMController* UControlRigGraph::GetController() const
 
 URigVMController* UControlRigGraph::GetTemplateController()
 {
-	if (TemplateModel == nullptr)
-	{
-		TemplateModel = NewObject<URigVMGraph>(this, TEXT("TemplateModel"));
-	}
 	if (TemplateController == nullptr)
 	{
-		TemplateController = NewObject<URigVMController>(this, TEXT("TemplateController"));
-		TemplateController->SetExecuteContextStruct(FControlRigExecuteContext::StaticStruct());
-		TemplateController->SetGraph(TemplateModel);
-		TemplateController->EnableReporting(false);
+		TemplateController = GetBlueprint()->GetTemplateController();
+		TemplateController->OnModified().RemoveAll(this);
 		TemplateController->OnModified().AddUObject(this, &UControlRigGraph::HandleModifiedEvent);
-
-		TemplateController->SetFlags(RF_Transient);
-		TemplateModel->SetFlags(RF_Transient);
 	}
 	return TemplateController;
 }
diff --git a/Engine/Plugins/Experimental/ControlRig/Source/ControlRigDeveloper/Private/Graph/ControlRigGraphNode.cpp b/Engine/Plugins/Experimental/ControlRig/Source/ControlRigDeveloper/Private/Graph/ControlRigGraphNode.cpp
--- a/Engine/Plugins/Experimental/ControlRig/Source/ControlRigDeveloper/Private/Graph/ControlRigGraphNode.cpp
+++ b/Engine/Plugins/Experimental/ControlRig/Source/ControlRigDeveloper/Private/Graph/ControlRigGraphNode.cpp
@@ -654,9 +654,9 @@ URigVMNode* UControlRigGraphNode::GetModelNode() const
 	{
 #if WITH_EDITOR
 
-		if (Graph->TemplateModel != nullptr)
+		if (Graph->TemplateController != nullptr)
 		{
-			return MutableThis->CachedModelNode = Graph->TemplateModel->FindNode(ModelNodePath);
+			return MutableThis->CachedModelNode = Graph->TemplateController->GetGraph()->FindNode(ModelNodePath);
 		}
 
 #endif
@@ -681,9 +681,13 @@ FName UControlRigGraphNode::GetModelNodeName() const
 
 URigVMPin* UControlRigGraphNode::GetModelPinFromPinPath(const FString& InPinPath) const
 {
-	if (URigVMPin*const* CachedModelPin = CachedModelPins.Find(InPinPath))
+	if (URigVMPin*const* CachedModelPinPtr = CachedModelPins.Find(InPinPath))
 	{
-		return (URigVMPin*)*CachedModelPin;
+		URigVMPin* CachedModelPin = *CachedModelPinPtr;
+		if (!CachedModelPin->HasAnyFlags(RF_Transient) && CachedModelPin->GetNode())
+		{
+			return CachedModelPin;
+		}
 	}
 
 	if (URigVMNode* ModelNode = GetModelNode())
@@ -691,10 +695,10 @@ URigVMPin* UControlRigGraphNode::GetModelPinFromPinPath(const FString& InPinPath
 		FString PinPath = InPinPath.RightChop(ModelNode->GetNodePath().Len() + 1);
 		URigVMPin* ModelPin = ModelNode->FindPin(PinPath);
 		if (ModelPin)
-	{
+		{
 			UControlRigGraphNode* MutableThis = (UControlRigGraphNode*)this;
-			MutableThis->CachedModelPins.Add(InPinPath, ModelPin);
-	}
+			MutableThis->CachedModelPins.FindOrAdd(InPinPath) = ModelPin;
+		}
 		return ModelPin;
 	}
 	
diff --git a/Engine/Plugins/Experimental/ControlRig/Source/ControlRigDeveloper/Private/Graph/ControlRigGraphSchema.cpp b/Engine/Plugins/Experimental/ControlRig/Source/ControlRigDeveloper/Private/Graph/ControlRigGraphSchema.cpp
--- a/Engine/Plugins/Experimental/ControlRig/Source/ControlRigDeveloper/Private/Graph/ControlRigGraphSchema.cpp
+++ b/Engine/Plugins/Experimental/ControlRig/Source/ControlRigDeveloper/Private/Graph/ControlRigGraphSchema.cpp
@@ -116,7 +116,7 @@ const FPinConnectionResponse UControlRigGraphSchema::CanCreateConnection(const U
 		UControlRigGraphNode* RigNodeA = Cast<UControlRigGraphNode>(A->GetOwningNode());
 		UControlRigGraphNode* RigNodeB = Cast<UControlRigGraphNode>(B->GetOwningNode());
 
-		if (RigNodeA && RigNodeB)
+		if (RigNodeA && RigNodeB && RigNodeA != RigNodeB)
 		{
 			URigVMPin* PinA = RigNodeA->GetModelPinFromPinPath(A->GetName());
 			if (PinA)
@@ -268,6 +268,64 @@ void UControlRigGraphSchema::SetNodePosition(UEdGraphNode* Node, const FVector2D
 	}
 }
 
+void UControlRigGraphSchema::GetGraphDisplayInformation(const UEdGraph& Graph, /*out*/ FGraphDisplayInfo& DisplayInfo) const
+{
+	Super::GetGraphDisplayInformation(Graph, DisplayInfo);
+
+	if (UControlRigGraph* RigGraph = Cast<UControlRigGraph>((UEdGraph*)&Graph))
+	{
+		TArray<FString> NodePathParts;
+		if(URigVMNode::SplitNodePath(RigGraph->ModelNodePath, NodePathParts))
+		{
+			DisplayInfo.DisplayName = FText::FromString(NodePathParts.Last());
+			DisplayInfo.PlainName = DisplayInfo.DisplayName;
+		}
+	}
+}
+
+bool UControlRigGraphSchema::TryDeleteGraph(UEdGraph* GraphToDelete) const
+{
+	if (UControlRigGraph* RigGraph = Cast<UControlRigGraph>(GraphToDelete))
+	{
+		if (UControlRigBlueprint* RigBlueprint = Cast<UControlRigBlueprint>(FBlueprintEditorUtils::FindBlueprintForGraph(RigGraph)))
+		{
+			if (URigVMGraph* Model = RigBlueprint->GetModel())
+			{
+				if (URigVMLibraryNode* LibraryNode = Cast<URigVMLibraryNode>(Model->FindNode(RigGraph->ModelNodePath)))
+				{
+					if (URigVMController* Controller = RigBlueprint->GetController(LibraryNode->GetGraph()))
+					{
+						return Controller->RemoveNode(LibraryNode);
+					}
+				}
+			}
+		}
+	}
+	return false;
+}
+
+bool UControlRigGraphSchema::TryRenameGraph(UEdGraph* GraphToRename, const FName& InNewName) const
+{
+	if (UControlRigGraph* RigGraph = Cast<UControlRigGraph>(GraphToRename))
+	{
+		if (UControlRigBlueprint* RigBlueprint = Cast<UControlRigBlueprint>(FBlueprintEditorUtils::FindBlueprintForGraph(RigGraph)))
+		{
+			if (URigVMGraph* Model = RigBlueprint->GetModel())
+			{
+				if (URigVMLibraryNode* LibraryNode = Cast<URigVMLibraryNode>(Model->FindNode(RigGraph->ModelNodePath)))
+				{
+					if (URigVMController* Controller = RigBlueprint->GetController(LibraryNode->GetGraph()))
+					{
+						Controller->RenameNode(LibraryNode, InNewName);
+						return true;
+					}
+				}
+			}
+		}
+	}
+	return false;
+}
+
 UControlRigGraphNode* UControlRigGraphSchema::CreateGraphNode(UControlRigGraph* InGraph, const FName& InPropertyName) const
 {
 	const bool bSelectNewNode = true;
@@ -320,6 +378,10 @@ bool UControlRigGraphSchema::ArePinsCompatible(const UEdGraphPin* PinA, const UE
 		return false;
 	}
 
+	if (UControlRigGraphNode* GraphNode = Cast<UControlRigGraphNode>(PinB->GetOwningNode()))
+	{
+	}
+
 	// for reroute nodes - always allow it
 	if (PinA->PinType.PinCategory == TEXT("POLYMORPH"))
 	{
diff --git a/Engine/Plugins/Experimental/ControlRig/Source/ControlRigDeveloper/Public/ControlRigBlueprint.h b/Engine/Plugins/Experimental/ControlRig/Source/ControlRigDeveloper/Public/ControlRigBlueprint.h
--- a/Engine/Plugins/Experimental/ControlRig/Source/ControlRigDeveloper/Public/ControlRigBlueprint.h
+++ b/Engine/Plugins/Experimental/ControlRig/Source/ControlRigDeveloper/Public/ControlRigBlueprint.h
@@ -34,6 +34,7 @@ DECLARE_EVENT_TwoParams(UControlRigBlueprint, FOnVMCompiledEvent, UBlueprint*, U
 DECLARE_EVENT_OneParam(UControlRigBlueprint, FOnRefreshEditorEvent, UControlRigBlueprint*);
 DECLARE_EVENT_FourParams(UControlRigBlueprint, FOnVariableDroppedEvent, UObject*, FProperty*, const FVector2D&, const FVector2D&);
 DECLARE_EVENT_OneParam(UControlRigBlueprint, FOnExternalVariablesChanged, const TArray<FRigVMExternalVariable>&);
+DECLARE_EVENT_TwoParams(UControlRigBlueprint, FOnNodeDoubleClicked, UControlRigBlueprint*, URigVMNode*);
 
 UCLASS(BlueprintType, meta=(IgnoreClassThumbnail))
 class CONTROLRIGDEVELOPER_API UControlRigBlueprint : public UBlueprint, public IInterface_PreviewMeshProvider
@@ -64,6 +65,7 @@ public:
 	virtual void PreSave(const class ITargetPlatform* TargetPlatform) override;
 	virtual void PostLoad() override;
 	virtual void PostTransacted(const FTransactionObjectEvent& TransactionEvent) override;
+	virtual void ReplaceDeprecatedNodes() override;
 
 	virtual bool SupportsGlobalVariables() const override { return true; }
 	virtual bool SupportsLocalVariables() const override { return false; }
@@ -102,6 +104,7 @@ public:
 
 	UFUNCTION(BlueprintCallable, Category = "Control Rig Blueprint")
 	URigVMGraph* GetModel(const UEdGraph* InEdGraph = nullptr) const;
+	URigVMGraph* GetModel(const FString& InNodePath) const;
 
 	UFUNCTION(BlueprintCallable, Category = "Control Rig Blueprint")
 	TArray<URigVMGraph*> GetAllModels() const;
@@ -115,6 +118,14 @@ public:
 	URigVMController* GetController(const UEdGraph* InEdGraph) const;
 	URigVMController* GetOrCreateController(const UEdGraph* InGraph);
 
+	URigVMGraph* GetTemplateModel();
+	URigVMController* GetTemplateController();
+
+#if WITH_EDITOR
+	UEdGraph* GetEdGraph(URigVMGraph* InModel = nullptr) const;
+	UEdGraph* GetEdGraph(const FString& InNodePath) const;
+#endif
+
 	UPROPERTY(EditAnywhere, BlueprintReadWrite, Category = "VM")
 	FRigVMCompileSettings VMCompileSettings;
 
@@ -126,6 +137,16 @@ protected:
 	UPROPERTY(BlueprintReadOnly, transient, Category = "VM")
 	TMap<URigVMGraph*, URigVMController*> Controllers;
 
+#if WITH_EDITORONLY_DATA
+
+	UPROPERTY(transient)
+	URigVMGraph* TemplateModel;
+
+	UPROPERTY(transient)
+	URigVMController* TemplateController;
+
+#endif
+
 public:
 
 	UPROPERTY(transient)
@@ -326,7 +347,7 @@ private:
 public:
 
 	FOnExternalVariablesChanged& OnExternalVariablesChanged() { return ExternalVariablesChangedEvent; }
-	
+
 	virtual void OnPreVariableChange(UObject* InObject);
 	virtual void OnPostVariableChange(UBlueprint* InBlueprint);
 	virtual void OnVariableAdded(const FName& InVarName);
@@ -334,13 +355,21 @@ public:
 	virtual void OnVariableRenamed(const FName& InOldVarName, const FName& InNewVarName);
 	virtual void OnVariableTypeChanged(const FName& InVarName, FEdGraphPinType InOldPinType, FEdGraphPinType InNewPinType);
 
+	FOnNodeDoubleClicked& OnNodeDoubleClicked() { return NodeDoubleClickedEvent; }
+	void BroadcastNodeDoubleClicked(URigVMNode* InNode);
+
 private:
 
 	FOnExternalVariablesChanged ExternalVariablesChangedEvent;
 	void BroadcastExternalVariablesChangedEvent();
 
+	FOnNodeDoubleClicked NodeDoubleClickedEvent;
+
 #endif
 
+	void CreateEdGraphForCollapseNodeIfNeeded(URigVMCollapseNode* InNode, bool bForce = false);
+	bool RemoveEdGraphForCollapseNode(URigVMCollapseNode* InNode, bool bNotify = false);
+
 	bool bDirtyDuringLoad;
 
 	friend class FControlRigBlueprintCompilerContext;
diff --git a/Engine/Plugins/Experimental/ControlRig/Source/ControlRigDeveloper/Public/Graph/ControlRigGraph.h b/Engine/Plugins/Experimental/ControlRig/Source/ControlRigDeveloper/Public/Graph/ControlRigGraph.h
--- a/Engine/Plugins/Experimental/ControlRig/Source/ControlRigDeveloper/Public/Graph/ControlRigGraph.h
+++ b/Engine/Plugins/Experimental/ControlRig/Source/ControlRigDeveloper/Public/Graph/ControlRigGraph.h
@@ -57,10 +57,13 @@ public:
 	URigVMGraph* GetModel() const;
 	URigVMController* GetController() const;
 
-private:
-
 	void HandleModifiedEvent(ERigVMGraphNotifType InNotifType, URigVMGraph* InGraph, UObject* InSubject);
 
+	UPROPERTY()
+	FString ModelNodePath;
+
+private:
+
 	template<class T>
 	void CacheNameList(const T& ElementList, TArray<TSharedPtr<FString>>& OutNameList)
 	{
@@ -95,9 +98,6 @@ private:
 #endif
 #if WITH_EDITORONLY_DATA
 
-	UPROPERTY(transient)
-	URigVMGraph* TemplateModel;
-
 	UPROPERTY(transient)
 	URigVMController* TemplateController;
 
diff --git a/Engine/Plugins/Experimental/ControlRig/Source/ControlRigDeveloper/Public/Graph/ControlRigGraphNode.h b/Engine/Plugins/Experimental/ControlRig/Source/ControlRigDeveloper/Public/Graph/ControlRigGraphNode.h
--- a/Engine/Plugins/Experimental/ControlRig/Source/ControlRigDeveloper/Public/Graph/ControlRigGraphNode.h
+++ b/Engine/Plugins/Experimental/ControlRig/Source/ControlRigDeveloper/Public/Graph/ControlRigGraphNode.h
@@ -200,4 +200,5 @@ private:
 	FNodeTitleDirtied NodeTitleDirtied;
 
 	friend class SControlRigGraphNode;
+	friend class FControlRigArgumentLayout;
 };
diff --git a/Engine/Plugins/Experimental/ControlRig/Source/ControlRigDeveloper/Public/Graph/ControlRigGraphSchema.h b/Engine/Plugins/Experimental/ControlRig/Source/ControlRigDeveloper/Public/Graph/ControlRigGraphSchema.h
--- a/Engine/Plugins/Experimental/ControlRig/Source/ControlRigDeveloper/Public/Graph/ControlRigGraphSchema.h
+++ b/Engine/Plugins/Experimental/ControlRig/Source/ControlRigDeveloper/Public/Graph/ControlRigGraphSchema.h
@@ -77,6 +77,9 @@ public:
 	virtual bool RequestVariableDropOnPin(UEdGraph* InGraph, FProperty* InVariableToDrop, UEdGraphPin* InPin, const FVector2D& InDropPosition, const FVector2D& InScreenPosition) override;
 	virtual bool IsStructEditable(UStruct* InStruct) const;
 	virtual void SetNodePosition(UEdGraphNode* Node, const FVector2D& Position) const override;
+	virtual void GetGraphDisplayInformation(const UEdGraph& Graph, /*out*/ FGraphDisplayInfo& DisplayInfo) const override;
+	virtual bool TryDeleteGraph(UEdGraph* GraphToDelete) const override;
+	virtual bool TryRenameGraph(UEdGraph* GraphToRename, const FName& InNewName) const override;
 
 	/** Create a graph node for a rig */
 	UControlRigGraphNode* CreateGraphNode(UControlRigGraph* InGraph, const FName& InPropertyName) const;
diff --git a/Engine/Plugins/Experimental/ControlRig/Source/ControlRigEditor/ControlRigEditor.Build.cs b/Engine/Plugins/Experimental/ControlRig/Source/ControlRigEditor/ControlRigEditor.Build.cs
--- a/Engine/Plugins/Experimental/ControlRig/Source/ControlRigEditor/ControlRigEditor.Build.cs
+++ b/Engine/Plugins/Experimental/ControlRig/Source/ControlRigEditor/ControlRigEditor.Build.cs
@@ -75,6 +75,7 @@ namespace UnrealBuildTool.Rules
 					"AnimationEditor",
 					"MessageLog",
 					"PropertyAccessEditor",
+					"KismetWidgets",
 				}
             );
 
diff --git a/Engine/Plugins/Experimental/ControlRig/Source/ControlRigEditor/Private/ControlRigEditorModule.cpp b/Engine/Plugins/Experimental/ControlRig/Source/ControlRigEditor/Private/ControlRigEditorModule.cpp
--- a/Engine/Plugins/Experimental/ControlRig/Source/ControlRigEditor/Private/ControlRigEditorModule.cpp
+++ b/Engine/Plugins/Experimental/ControlRig/Source/ControlRigEditor/Private/ControlRigEditorModule.cpp
@@ -68,6 +68,7 @@
 #include "ControlRigElementDetails.h"
 #include "ControlRigCompilerDetails.h"
 #include "ControlRigDrawingDetails.h"
+#include "ControlRigGraphDetails.h"
 #include "ControlRigInfluenceMapDetails.h"
 #include "Animation/AnimSequence.h"
 #include "Editor/SControlRigProfilingView.h"
@@ -137,6 +138,7 @@ void FControlRigEditorModule::StartupModule()
 	// Register Blueprint editor variable customization
 	FBlueprintEditorModule& BlueprintEditorModule = FModuleManager::LoadModuleChecked<FBlueprintEditorModule>("Kismet");
 	BlueprintEditorModule.RegisterVariableCustomization(FProperty::StaticClass(), FOnGetVariableCustomizationInstance::CreateStatic(&FControlRigVariableDetailsCustomization::MakeInstance));
+	BlueprintEditorModule.RegisterGraphCustomization(GetDefault<UControlRigGraphSchema>(), FOnGetGraphCustomizationInstance::CreateStatic(&FControlRigGraphDetails::MakeInstance));
 
 	// Register to fixup newly created BPs
 	FKismetEditorUtilities::RegisterOnBlueprintCreatedCallback(this, UControlRig::StaticClass(), FKismetEditorUtilities::FOnBlueprintCreated::CreateRaw(this, &FControlRigEditorModule::HandleNewBlueprintCreated));
@@ -246,7 +248,6 @@ void FControlRigEditorModule::StartupModule()
 	//UThumbnailManager::Get().RegisterCustomRenderer(UControlRigPoseAsset::StaticClass(), UControlRigPoseThumbnailRenderer::StaticClass());
 
 	bFilterAssetBySkeleton = true;
-
 }
 
 void FControlRigEditorModule::ShutdownModule()
@@ -292,6 +293,7 @@ void FControlRigEditorModule::ShutdownModule()
 	if (BlueprintEditorModule)
 	{
 		BlueprintEditorModule->UnregisterVariableCustomization(FProperty::StaticClass());
+		BlueprintEditorModule->UnregisterGraphCustomization(GetDefault<UControlRigGraphSchema>());
 	}
 
 	FPropertyEditorModule* PropertyEditorModule = FModuleManager::GetModulePtr<FPropertyEditorModule>("PropertyEditor");
diff --git a/Engine/Plugins/Experimental/ControlRig/Source/ControlRigEditor/Private/ControlRigGraphDetails.cpp b/Engine/Plugins/Experimental/ControlRig/Source/ControlRigEditor/Private/ControlRigGraphDetails.cpp
new file mode 100644
--- /dev/null
+++ b/Engine/Plugins/Experimental/ControlRig/Source/ControlRigEditor/Private/ControlRigGraphDetails.cpp
@@ -0,0 +1,612 @@
+// Copyright Epic Games, Inc. All Rights Reserved.
+
+#include "ControlRigGraphDetails.h"
+#include "Widgets/SWidget.h"
+#include "DetailLayoutBuilder.h"
+#include "DetailCategoryBuilder.h"
+#include "DetailWidgetRow.h"
+#include "IDetailChildrenBuilder.h"
+#include "EditorStyleSet.h"
+#include "SPinTypeSelector.h"
+#include "Widgets/Input/SButton.h"
+#include "Widgets/Text/STextBlock.h"
+#include "PropertyCustomizationHelpers.h"
+#include "NodeFactory.h"
+#include "Graph/ControlRigGraphNode.h"
+#include "ControlRig.h"
+#include "RigVMCore/RigVMExternalVariable.h"
+#include "Kismet2/BlueprintEditorUtils.h"
+
+#define LOCTEXT_NAMESPACE "ControlRigGraphDetails"
+
+FControlRigArgumentGroupLayout::FControlRigArgumentGroupLayout(URigVMGraph* InGraph, UControlRigBlueprint* InBlueprint, bool bInputs)
+	: GraphPtr(InGraph)
+	, ControlRigBlueprintPtr(InBlueprint)
+	, bIsInputGroup(bInputs)
+{
+	if (ControlRigBlueprintPtr.IsValid())
+	{
+		ControlRigBlueprintPtr.Get()->OnModified().AddRaw(this, &FControlRigArgumentGroupLayout::HandleModifiedEvent);
+	}
+}
+
+FControlRigArgumentGroupLayout::~FControlRigArgumentGroupLayout()
+{
+	if (ControlRigBlueprintPtr.IsValid())
+	{
+		ControlRigBlueprintPtr.Get()->OnModified().RemoveAll(this);
+	}
+}
+
+void FControlRigArgumentGroupLayout::GenerateChildContent(IDetailChildrenBuilder& ChildrenBuilder)
+{
+	bool WasContentAdded = false;
+	if (GraphPtr.IsValid())
+	{
+		URigVMGraph* Graph = GraphPtr.Get();
+		if (URigVMLibraryNode* LibraryNode = Cast<URigVMLibraryNode>(Graph->GetOuter()))
+		{
+			for (URigVMPin* Pin : LibraryNode->GetPins())
+			{
+				if ((bIsInputGroup && (Pin->GetDirection() == ERigVMPinDirection::Input || Pin->GetDirection() == ERigVMPinDirection::IO)) ||
+					(!bIsInputGroup && (Pin->GetDirection() == ERigVMPinDirection::Output || Pin->GetDirection() == ERigVMPinDirection::IO)))
+				{
+					TSharedRef<class FControlRigArgumentLayout> ControlRigArgumentLayout = MakeShareable(new FControlRigArgumentLayout(
+						Pin,
+						Graph,
+						ControlRigBlueprintPtr.Get()));
+					ChildrenBuilder.AddCustomBuilder(ControlRigArgumentLayout);
+					WasContentAdded = true;
+				}
+			}
+		}
+	}
+	if (!WasContentAdded)
+	{
+		// Add a text widget to let the user know to hit the + icon to add parameters.
+		ChildrenBuilder.AddCustomRow(FText::GetEmpty()).WholeRowContent()
+		.MaxDesiredWidth(980.f)
+		[
+			SNew(SHorizontalBox)
+			+ SHorizontalBox::Slot()
+			.VAlign(VAlign_Center)
+			.Padding(0.0f, 0.0f, 4.0f, 0.0f)
+			.AutoWidth()
+			[
+				SNew(STextBlock)
+				.Text(LOCTEXT("NoArgumentsAddedForControlRig", "Please press the + icon above to add parameters"))
+				.Font(IDetailLayoutBuilder::GetDetailFont())
+			]
+		];
+	}
+}
+
+void FControlRigArgumentGroupLayout::HandleModifiedEvent(ERigVMGraphNotifType InNotifType, URigVMGraph* InGraph, UObject* InSubject)
+{
+	if (!GraphPtr.IsValid())
+	{
+		return;
+	}
+	
+	URigVMGraph* Graph = GraphPtr.Get();
+	URigVMLibraryNode* LibraryNode = Cast<URigVMLibraryNode>(Graph->GetOuter());
+	if (LibraryNode == nullptr)
+	{
+		return;
+	}
+
+	switch (InNotifType)
+	{
+		case ERigVMGraphNotifType::PinAdded:
+		case ERigVMGraphNotifType::PinRemoved:
+		{
+			URigVMPin* Pin = CastChecked<URigVMPin>(InSubject);
+			if (Pin->GetNode() == LibraryNode)
+			{
+				OnRebuildChildren.ExecuteIfBound();
+			}
+			break;
+		}
+		default:
+		{
+			break;
+		}
+	}
+}
+
+void FControlRigArgumentLayout::GenerateHeaderRowContent(FDetailWidgetRow& NodeRow)
+{
+	const UEdGraphSchema_K2* K2Schema = GetDefault<UEdGraphSchema_K2>();
+
+	ETypeTreeFilter TypeTreeFilter = ETypeTreeFilter::None;
+	TypeTreeFilter |= ETypeTreeFilter::AllowExec;
+
+	NodeRow
+	.NameContent()
+	[
+		SNew(SHorizontalBox)
+		+ SHorizontalBox::Slot()
+		.FillWidth(1)
+		.VAlign(VAlign_Center)
+		[
+			SAssignNew(ArgumentNameWidget, SEditableTextBox)
+			.Text(this, &FControlRigArgumentLayout::OnGetArgNameText)
+			.OnTextChanged(this, &FControlRigArgumentLayout::OnArgNameChange)
+			.OnTextCommitted(this, &FControlRigArgumentLayout::OnArgNameTextCommitted)
+			.ToolTipText(this, &FControlRigArgumentLayout::OnGetArgToolTipText)
+			.Font(IDetailLayoutBuilder::GetDetailFont())
+			.IsEnabled(!ShouldPinBeReadOnly())
+		]
+	]
+	.ValueContent()
+	.MaxDesiredWidth(980.f)
+	[
+		SNew(SHorizontalBox)
+		+ SHorizontalBox::Slot()
+		.VAlign(VAlign_Center)
+		.Padding(0.0f, 0.0f, 4.0f, 0.0f)
+		.AutoWidth()
+		[
+			SNew(SPinTypeSelector, FGetPinTypeTree::CreateUObject(K2Schema, &UEdGraphSchema_K2::GetVariableTypeTree))
+			.TargetPinType(this, &FControlRigArgumentLayout::OnGetPinInfo)
+			.OnPinTypePreChanged(this, &FControlRigArgumentLayout::OnPrePinInfoChange)
+			.OnPinTypeChanged(this, &FControlRigArgumentLayout::PinInfoChanged)
+			.Schema(K2Schema)
+			.TypeTreeFilter(TypeTreeFilter)
+			.bAllowArrays(!ShouldPinBeReadOnly())
+			.IsEnabled(!ShouldPinBeReadOnly(true))
+			.Font(IDetailLayoutBuilder::GetDetailFont())
+		]
+		+ SHorizontalBox::Slot()
+		.AutoWidth()
+		[
+			SNew(SButton)
+			.ButtonStyle(FAppStyle::Get(), TEXT("SimpleButton"))
+			.ContentPadding(0)
+			.IsEnabled(!IsPinEditingReadOnly())
+			.OnClicked(this, &FControlRigArgumentLayout::OnArgMoveUp)
+			.ToolTipText(LOCTEXT("FunctionArgDetailsArgMoveUpTooltip", "Move this parameter up in the list."))
+			[
+				SNew(SImage)
+				.Image(FEditorStyle::GetBrush("Icons.ChevronUp"))
+			.ColorAndOpacity(FSlateColor::UseForeground())
+			]
+		]
+		+ SHorizontalBox::Slot()
+		.AutoWidth()
+		.Padding(2, 0)
+		[
+			SNew(SButton)
+			.ButtonStyle(FAppStyle::Get(), TEXT("SimpleButton"))
+			.ContentPadding(0)
+			.IsEnabled(!IsPinEditingReadOnly())
+			.OnClicked(this, &FControlRigArgumentLayout::OnArgMoveDown)
+			.ToolTipText(LOCTEXT("FunctionArgDetailsArgMoveDownTooltip", "Move this parameter down in the list."))
+			[
+				SNew(SImage)
+				.Image(FEditorStyle::GetBrush("Icons.ChevronDown"))
+			.ColorAndOpacity(FSlateColor::UseForeground())
+			]
+		]
+		+ SHorizontalBox::Slot()
+		.HAlign(HAlign_Right)
+		.VAlign(VAlign_Center)
+		.Padding(10, 0, 0, 0)
+		.AutoWidth()
+		[
+			PropertyCustomizationHelpers::MakeClearButton(FSimpleDelegate::CreateSP(this, &FControlRigArgumentLayout::OnRemoveClicked), LOCTEXT("FunctionArgDetailsClearTooltip", "Remove this parameter."), !IsPinEditingReadOnly())
+		]
+	];
+}
+
+void FControlRigArgumentLayout::GenerateChildContent(IDetailChildrenBuilder& ChildrenBuilder)
+{
+	if (UEdGraphPin* FoundPin = GetPin())
+	{
+		// Certain types are outlawed at the compiler level, or to keep consistency with variable rules for actors
+		const UClass* ClassObject = Cast<UClass>(FoundPin->PinType.PinSubCategoryObject.Get());
+		const bool bTypeWithNoDefaults = (FoundPin->PinType.PinCategory == UEdGraphSchema_K2::PC_Object) || (FoundPin->PinType.PinCategory == UEdGraphSchema_K2::PC_Class) || (FoundPin->PinType.PinCategory == UEdGraphSchema_K2::PC_Interface)
+			|| (FoundPin->PinType.PinCategory == UEdGraphSchema_K2::PC_SoftObject && ClassObject && ClassObject->IsChildOf(AActor::StaticClass()))
+			|| UEdGraphSchema_K2::IsExecPin(*FoundPin)
+			|| PinPtr.Get()->IsExecuteContext()
+			|| FoundPin->PinType.bIsReference;
+
+		if (!bTypeWithNoDefaults)
+		{
+			DefaultValuePinWidget = FNodeFactory::CreatePinWidget(FoundPin);
+			DefaultValuePinWidget->SetOnlyShowDefaultValue(true);
+			TSharedRef<SWidget> DefaultValueWidget = DefaultValuePinWidget->GetDefaultValueWidget();
+
+			if (DefaultValueWidget != SNullWidget::NullWidget)
+			{
+				ChildrenBuilder.AddCustomRow(LOCTEXT("FunctionArgDetailsDefaultValue", "Default Value"))
+				.NameContent()
+				[
+					SNew(STextBlock)
+					.Text(LOCTEXT("FunctionArgDetailsDefaultValue", "Default Value"))
+					.ToolTipText(LOCTEXT("FunctionArgDetailsDefaultValueParamTooltip", "The default value of the parameter."))
+					.Font(IDetailLayoutBuilder::GetDetailFont())
+				]
+				.ValueContent()
+				.MaxDesiredWidth(512)
+				[
+					DefaultValueWidget
+				];
+			}
+			else
+			{
+				DefaultValuePinWidget.Reset();
+			}
+		}
+	}
+}
+
+void FControlRigArgumentLayout::OnRemoveClicked()
+{
+	if (PinPtr.IsValid() && ControlRigBlueprintPtr.IsValid())
+	{
+		URigVMPin* Pin = PinPtr.Get();
+		UControlRigBlueprint* Blueprint = ControlRigBlueprintPtr.Get();
+		if (URigVMLibraryNode* LibraryNode = Cast<URigVMLibraryNode>(Pin->GetNode()))
+		{
+			if (URigVMController* Controller = Blueprint->GetController(LibraryNode->GetContainedGraph()))
+			{
+				Controller->RemoveExposedPin(Pin->GetFName(), true);
+			}
+		}
+	}
+}
+
+FReply FControlRigArgumentLayout::OnArgMoveUp()
+{
+	/* todo
+
+	const int32 ThisParamIndex = TargetNode->UserDefinedPins.Find(ParamItemPtr.Pin());
+	const int32 NewParamIndex = ThisParamIndex - 1;
+	if (ThisParamIndex != INDEX_NONE && NewParamIndex >= 0)
+	{
+		const FScopedTransaction Transaction(LOCTEXT("K2_MovePinUp", "Move Pin Up"));
+		TArray<UK2Node_EditablePinBase*> TargetNodes = GatherAllResultNodes(TargetNode);
+		for (UK2Node_EditablePinBase* Node : TargetNodes)
+		{
+			Node->Modify();
+			Node->UserDefinedPins.Swap(ThisParamIndex, NewParamIndex);
+
+			TSharedPtr<FBaseBlueprintGraphActionDetails> GraphActionDetails = GraphActionDetailsPtr.Pin();
+			if (GraphActionDetails.IsValid())
+			{
+				GraphActionDetails->OnParamsChanged(Node, true);
+			}
+		}
+	}
+
+	*/
+	return FReply::Handled();
+}
+
+FReply FControlRigArgumentLayout::OnArgMoveDown()
+{
+	/* todo
+	const int32 ThisParamIndex = TargetNode->UserDefinedPins.Find(ParamItemPtr.Pin());
+	const int32 NewParamIndex = ThisParamIndex + 1;
+	if (ThisParamIndex != INDEX_NONE && NewParamIndex < TargetNode->UserDefinedPins.Num())
+	{
+		const FScopedTransaction Transaction(LOCTEXT("K2_MovePinDown", "Move Pin Down"));
+		TArray<UK2Node_EditablePinBase*> TargetNodes = GatherAllResultNodes(TargetNode);
+		for (UK2Node_EditablePinBase* Node : TargetNodes)
+		{
+			Node->Modify();
+			Node->UserDefinedPins.Swap(ThisParamIndex, NewParamIndex);
+
+			TSharedPtr<FBaseBlueprintGraphActionDetails> GraphActionDetails = GraphActionDetailsPtr.Pin();
+			if (GraphActionDetails.IsValid())
+			{
+				GraphActionDetails->OnParamsChanged(Node, true);
+			}
+		}
+	}
+	*/
+	return FReply::Handled();
+}
+
+bool FControlRigArgumentLayout::ShouldPinBeReadOnly(bool bIsEditingPinType/* = false*/) const
+{
+	if (PinPtr.IsValid())
+	{
+		return PinPtr.Get()->IsExecuteContext();
+	}
+	return false;
+}
+
+bool FControlRigArgumentLayout::IsPinEditingReadOnly(bool bIsEditingPinType/* = false*/) const
+{
+	if (PinPtr.IsValid())
+	{
+		return PinPtr.Get()->IsExecuteContext();
+	}
+	return false;
+}
+
+FText FControlRigArgumentLayout::OnGetArgNameText() const
+{
+	if (PinPtr.IsValid())
+	{
+		return FText::FromName(PinPtr.Get()->GetFName());
+	}
+	return FText();
+}
+
+FText FControlRigArgumentLayout::OnGetArgToolTipText() const
+{
+	return OnGetArgNameText(); // for now since we don't have tooltips
+}
+
+void FControlRigArgumentLayout::OnArgNameChange(const FText& InNewText)
+{
+	// do we need validation?
+}
+
+void FControlRigArgumentLayout::OnArgNameTextCommitted(const FText& NewText, ETextCommit::Type InTextCommit)
+{
+	if (!NewText.IsEmpty() && PinPtr.IsValid() && ControlRigBlueprintPtr.IsValid() && !ShouldPinBeReadOnly())
+	{
+		URigVMPin* Pin = PinPtr.Get();
+		UControlRigBlueprint* Blueprint = ControlRigBlueprintPtr.Get();
+		if (URigVMLibraryNode* LibraryNode = Cast<URigVMLibraryNode>(Pin->GetNode()))
+		{
+			if (URigVMController* Controller = Blueprint->GetController(LibraryNode->GetContainedGraph()))
+			{
+				const FString& NewName = NewText.ToString();
+				Controller->RenameExposedPin(Pin->GetFName(), *NewName, true);
+			}
+		}
+	}
+}
+
+FEdGraphPinType FControlRigArgumentLayout::OnGetPinInfo() const
+{
+	if (PinPtr.IsValid())
+	{
+		return UControlRigGraphNode::GetPinTypeForModelPin(PinPtr.Get());
+	}
+	return FEdGraphPinType();
+}
+
+UEdGraphPin* FControlRigArgumentLayout::GetPin() const
+{
+	if (PinPtr.IsValid() && ControlRigBlueprintPtr.IsValid())
+	{
+		UControlRigBlueprint* Blueprint = ControlRigBlueprintPtr.Get();
+		if(UControlRigGraph* RigGraph = Cast<UControlRigGraph>(Blueprint->GetEdGraph(PinPtr.Get()->GetGraph())))
+		{
+			if (UControlRigGraphNode* RigNode = Cast<UControlRigGraphNode>(RigGraph->FindNodeForModelNodeName(PinPtr.Get()->GetNode()->GetFName())))
+			{
+				return RigNode->FindPin(PinPtr.Get()->GetPinPath());
+			}
+		}
+	}
+	return nullptr;
+}
+
+ECheckBoxState FControlRigArgumentLayout::IsRefChecked() const
+{
+	FEdGraphPinType PinType = OnGetPinInfo();
+	return PinType.bIsReference ? ECheckBoxState::Checked : ECheckBoxState::Unchecked;
+}
+
+void FControlRigArgumentLayout::OnRefCheckStateChanged(ECheckBoxState InState)
+{
+	const FScopedTransaction Transaction(LOCTEXT("ChangeByRef", "Change Pass By Reference"));
+
+	FEdGraphPinType PinType = OnGetPinInfo();
+	PinType.bIsReference = (InState == ECheckBoxState::Checked) ? true : false;
+
+	PinInfoChanged(PinType);
+}
+
+void FControlRigArgumentLayout::PinInfoChanged(const FEdGraphPinType& PinType)
+{
+	if (PinPtr.IsValid() && ControlRigBlueprintPtr.IsValid() && FBlueprintEditorUtils::IsPinTypeValid(PinType))
+	{
+		URigVMPin* Pin = PinPtr.Get();
+		UControlRigBlueprint* Blueprint = ControlRigBlueprintPtr.Get();
+		if (URigVMLibraryNode* LibraryNode = Cast<URigVMLibraryNode>(Pin->GetNode()))
+		{
+			if (URigVMController* Controller = Blueprint->GetController(LibraryNode->GetContainedGraph()))
+			{
+				FRigVMExternalVariable ExternalVariable = UControlRig::GetExternalVariableFromPinType(Pin->GetFName(), PinType, true, false);
+				if (!ExternalVariable.IsValid(true /* allow nullptr memory */))
+				{
+					return;
+				}
+
+				FString CPPType = ExternalVariable.TypeName.ToString();
+				FName CPPTypeObjectName = NAME_None;
+				if (ExternalVariable.TypeObject)
+				{
+					CPPTypeObjectName = *ExternalVariable.TypeObject->GetPathName();
+
+					if (UScriptStruct* ScriptStruct = Cast<UScriptStruct>(ExternalVariable.TypeObject))
+					{
+						CPPType = ScriptStruct->GetStructCPPName();
+					}
+				}
+
+				if (ExternalVariable.bIsArray)
+				{
+					CPPType = FString::Printf(TEXT("TArray<%s>"), *CPPType);
+				}
+
+				Controller->ChangeExposedPinType(Pin->GetFName(), CPPType, CPPTypeObjectName, true);
+			}
+		}
+	}
+}
+
+void FControlRigArgumentLayout::OnPrePinInfoChange(const FEdGraphPinType& PinType)
+{
+	// not needed for Control Rig
+}
+
+TSharedPtr<IDetailCustomization> FControlRigGraphDetails::MakeInstance(TSharedPtr<IBlueprintEditor> InBlueprintEditor)
+{
+	const TArray<UObject*>* Objects = (InBlueprintEditor.IsValid() ? InBlueprintEditor->GetObjectsCurrentlyBeingEdited() : nullptr);
+	if (Objects && Objects->Num() == 1)
+	{
+		if (UControlRigBlueprint* ControlRigBlueprint = Cast<UControlRigBlueprint>((*Objects)[0]))
+		{
+			return MakeShareable(new FControlRigGraphDetails(StaticCastSharedPtr<IControlRigEditor>(InBlueprintEditor), ControlRigBlueprint));
+		}
+	}
+
+	return nullptr;
+}
+
+void FControlRigGraphDetails::CustomizeDetails(IDetailLayoutBuilder& DetailLayout)
+{
+	TArray<TWeakObjectPtr<UObject>> Objects;
+	DetailLayout.GetObjectsBeingCustomized(Objects);
+
+	GraphPtr = CastChecked<UControlRigGraph>(Objects[0].Get());
+	UControlRigGraph* Graph = GraphPtr.Get();
+
+	UControlRigBlueprint* Blueprint = ControlRigBlueprintPtr.Get();
+	URigVMGraph* Model = nullptr;
+	URigVMController* Controller = nullptr;
+
+	if (Blueprint)
+	{
+		Model = Blueprint->GetModel(Graph);
+		Controller = Blueprint->GetController(Model);
+	}
+
+	if (Blueprint == nullptr || Model == nullptr || Controller == nullptr)
+	{
+		IDetailCategoryBuilder& Category = DetailLayout.EditCategory("Graph", LOCTEXT("FunctionDetailsGraph", "Graph"));
+		Category.AddCustomRow(FText::GetEmpty())
+		[
+			SNew(STextBlock)
+			.Text(LOCTEXT("GraphPresentButNotEditable", "Graph is not editable."))
+		];
+		return;
+	}
+
+	if (Model->IsTopLevelGraph())
+	{
+		IDetailCategoryBuilder& Category = DetailLayout.EditCategory("Graph", LOCTEXT("FunctionDetailsGraph", "Graph"));
+		Category.AddCustomRow(FText::GetEmpty())
+			[
+				SNew(STextBlock)
+				.Text(LOCTEXT("GraphIsTopLevelGraph", "Top-level Graphs are not editable."))
+			];
+		return;
+	}
+
+	IDetailCategoryBuilder& InputsCategory = DetailLayout.EditCategory("Inputs", LOCTEXT("FunctionDetailsInputs", "Inputs"));
+	TSharedRef<FControlRigArgumentGroupLayout> InputArgumentGroup = MakeShareable(new FControlRigArgumentGroupLayout(Model, Blueprint, true));
+	InputsCategory.AddCustomBuilder(InputArgumentGroup);
+
+	TSharedRef<SHorizontalBox> InputsHeaderContentWidget = SNew(SHorizontalBox);
+	TWeakPtr<SWidget> WeakInputsHeaderWidget = InputsHeaderContentWidget;
+
+	InputsHeaderContentWidget->AddSlot()
+	.HAlign(HAlign_Right)
+	[
+		SNew(SButton)
+		.ButtonStyle(FEditorStyle::Get(), "SimpleButton")
+		.ContentPadding(FMargin(1, 0))
+		.OnClicked(this, &FControlRigGraphDetails::OnAddNewInputClicked)
+		.Visibility(this, &FControlRigGraphDetails::GetAddNewInputOutputVisibility)
+		.HAlign(HAlign_Right)
+		.ToolTipText(LOCTEXT("FunctionNewInputArgTooltip", "Create a new input argument"))
+		.VAlign(VAlign_Center)
+		.AddMetaData<FTagMetaData>(FTagMetaData(TEXT("FunctionNewInputArg")))
+		.IsEnabled(this, &FControlRigGraphDetails::IsAddNewInputOutputEnabled)
+		[
+			SNew(SImage)
+			.Image(FAppStyle::Get().GetBrush("Icons.PlusCircle"))
+			.ColorAndOpacity(FSlateColor::UseForeground())
+		]
+	];
+	InputsCategory.HeaderContent(InputsHeaderContentWidget);
+
+	IDetailCategoryBuilder& OutputsCategory = DetailLayout.EditCategory("Outputs", LOCTEXT("FunctionDetailsOutputs", "Outputs"));
+	TSharedRef<FControlRigArgumentGroupLayout> OutputArgumentGroup = MakeShareable(new FControlRigArgumentGroupLayout(Model, Blueprint, false));
+	OutputsCategory.AddCustomBuilder(OutputArgumentGroup);
+
+	TSharedRef<SHorizontalBox> OutputsHeaderContentWidget = SNew(SHorizontalBox);
+	TWeakPtr<SWidget> WeakOutputsHeaderWidget = OutputsHeaderContentWidget;
+
+	OutputsHeaderContentWidget->AddSlot()
+	.HAlign(HAlign_Right)
+	[
+		SNew(SButton)
+		.ButtonStyle(FEditorStyle::Get(), "SimpleButton")
+		.ContentPadding(FMargin(1, 0))
+		.OnClicked(this, &FControlRigGraphDetails::OnAddNewOutputClicked)
+		.Visibility(this, &FControlRigGraphDetails::GetAddNewInputOutputVisibility)
+		.HAlign(HAlign_Right)
+		.ToolTipText(LOCTEXT("FunctionNewOutputArgTooltip", "Create a new output argument"))
+		.VAlign(VAlign_Center)
+		.AddMetaData<FTagMetaData>(FTagMetaData(TEXT("FunctionNewOutputArg")))
+		.IsEnabled(this, &FControlRigGraphDetails::IsAddNewInputOutputEnabled)
+		[
+			SNew(SImage)
+			.Image(FAppStyle::Get().GetBrush("Icons.PlusCircle"))
+			.ColorAndOpacity(FSlateColor::UseForeground())
+		]
+	];
+	OutputsCategory.HeaderContent(OutputsHeaderContentWidget);
+}
+
+bool FControlRigGraphDetails::IsAddNewInputOutputEnabled() const
+{
+	return true;
+}
+
+EVisibility FControlRigGraphDetails::GetAddNewInputOutputVisibility() const
+{
+	return EVisibility::Visible;
+}
+
+FReply FControlRigGraphDetails::OnAddNewInputClicked()
+{
+	if (GraphPtr.IsValid() && ControlRigBlueprintPtr.IsValid())
+	{
+		UControlRigBlueprint* Blueprint = ControlRigBlueprintPtr.Get();
+		URigVMGraph* Model = Blueprint->GetModel(GraphPtr.Get());
+		if (URigVMController* Controller = Blueprint->GetController(Model))
+		{
+			FName ArgumentName = TEXT("Argument");
+			FString CPPType = TEXT("bool");
+			FName CPPTypeObjectPath = NAME_None;
+			FString DefaultValue = TEXT("False");
+			// todo: base decisions on types on last argument
+
+			Controller->AddExposedPin(ArgumentName, ERigVMPinDirection::Input, CPPType, CPPTypeObjectPath, DefaultValue, true);
+		}
+	}
+	return FReply::Unhandled();
+}
+
+FReply FControlRigGraphDetails::OnAddNewOutputClicked()
+{
+	if (GraphPtr.IsValid() && ControlRigBlueprintPtr.IsValid())
+	{
+		UControlRigBlueprint* Blueprint = ControlRigBlueprintPtr.Get();
+		URigVMGraph* Model = Blueprint->GetModel(GraphPtr.Get());
+		if (URigVMController* Controller = Blueprint->GetController(Model))
+		{
+			FName ArgumentName = TEXT("Argument");
+			FString CPPType = TEXT("bool");
+			FName CPPTypeObjectPath = NAME_None;
+			FString DefaultValue = TEXT("False");
+			// todo: base decisions on types on last argument
+
+			Controller->AddExposedPin(ArgumentName, ERigVMPinDirection::Output, CPPType, CPPTypeObjectPath, DefaultValue, true);
+		}
+	}
+	return FReply::Unhandled();
+}
+
+#undef LOCTEXT_NAMESPACE
diff --git a/Engine/Plugins/Experimental/ControlRig/Source/ControlRigEditor/Private/ControlRigGraphDetails.h b/Engine/Plugins/Experimental/ControlRig/Source/ControlRigEditor/Private/ControlRigGraphDetails.h
new file mode 100644
--- /dev/null
+++ b/Engine/Plugins/Experimental/ControlRig/Source/ControlRigEditor/Private/ControlRigGraphDetails.h
@@ -0,0 +1,143 @@
+// Copyright Epic Games, Inc. All Rights Reserved.
+
+#pragma once
+
+#include "CoreMinimal.h"
+#include "IDetailCustomization.h"
+#include "IDetailCustomNodeBuilder.h"
+#include "UObject/WeakObjectPtr.h"
+#include "Graph/ControlRigGraph.h"
+#include "ControlRigBlueprint.h"
+#include "ControlRigEditor.h"
+#include "SGraphPin.h"
+
+class FMobilityCustomization;
+class IDetailLayoutBuilder;
+
+class FControlRigArgumentGroupLayout : public IDetailCustomNodeBuilder, public TSharedFromThis<FControlRigArgumentGroupLayout>
+{
+public:
+	FControlRigArgumentGroupLayout(URigVMGraph* InGraph, UControlRigBlueprint* InBlueprint, bool bInputs);
+	virtual ~FControlRigArgumentGroupLayout();
+
+private:
+	/** IDetailCustomNodeBuilder Interface*/
+	virtual void SetOnRebuildChildren(FSimpleDelegate InOnRegenerateChildren) override { OnRebuildChildren = InOnRegenerateChildren; }
+	virtual void GenerateHeaderRowContent(FDetailWidgetRow& NodeRow) override {}
+	virtual void GenerateChildContent(IDetailChildrenBuilder& ChildrenBuilder) override;
+	virtual void Tick(float DeltaTime) override {}
+	virtual bool RequiresTick() const override { return false; }
+	virtual FName GetName() const override { return NAME_None; }
+	virtual bool InitiallyCollapsed() const override { return false; }
+
+private:
+
+	void HandleModifiedEvent(ERigVMGraphNotifType InNotifType, URigVMGraph* InGraph, UObject* InSubject);
+
+	TWeakObjectPtr<URigVMGraph> GraphPtr;
+	TWeakObjectPtr<UControlRigBlueprint> ControlRigBlueprintPtr;
+	bool bIsInputGroup;
+	FSimpleDelegate OnRebuildChildren;
+
+};
+
+class FControlRigArgumentLayout : public IDetailCustomNodeBuilder, public TSharedFromThis<FControlRigArgumentLayout>
+{
+public:
+
+	FControlRigArgumentLayout(URigVMPin* InPin, URigVMGraph* InGraph, UControlRigBlueprint* InBlueprint)
+		: PinPtr(InPin)
+		, GraphPtr(InGraph)
+		, ControlRigBlueprintPtr(InBlueprint)
+	{}
+
+private:
+
+	/** IDetailCustomNodeBuilder Interface*/
+	virtual void GenerateHeaderRowContent(FDetailWidgetRow& NodeRow) override;
+	virtual void GenerateChildContent(IDetailChildrenBuilder& ChildrenBuilder) override;
+	virtual void Tick(float DeltaTime) override {}
+	virtual bool RequiresTick() const override { return false; }
+	virtual FName GetName() const override { return PinPtr.Get()->GetFName(); }
+	virtual bool InitiallyCollapsed() const override { return true; }
+
+private:
+
+	/** Determines if this pin should not be editable */
+	bool ShouldPinBeReadOnly(bool bIsEditingPinType = false) const;
+
+	/** Determines if editing the pins on the node should be read only */
+	bool IsPinEditingReadOnly(bool bIsEditingPinType = false) const;
+
+	/** Callbacks for all the functionality for modifying arguments */
+	void OnRemoveClicked();
+	FReply OnArgMoveUp();
+	FReply OnArgMoveDown();
+
+	FText OnGetArgNameText() const;
+	FText OnGetArgToolTipText() const;
+	void OnArgNameChange(const FText& InNewText);
+	void OnArgNameTextCommitted(const FText& NewText, ETextCommit::Type InTextCommit);
+
+	FEdGraphPinType OnGetPinInfo() const;
+	void PinInfoChanged(const FEdGraphPinType& PinType);
+	void OnPrePinInfoChange(const FEdGraphPinType& PinType);
+
+	/** Returns the graph pin representing this variable */
+	UEdGraphPin* GetPin() const;
+
+	/** Returns whether the "Pass-by-Reference" checkbox is checked or not */
+	ECheckBoxState IsRefChecked() const;
+
+	/** Handles toggling the "Pass-by-Reference" checkbox */
+	void OnRefCheckStateChanged(ECheckBoxState InState);
+
+private:
+
+	/** The argument pin that this layout reflects */
+	TWeakObjectPtr<URigVMPin> PinPtr;
+	
+	/** The target graph that this argument is on */
+	TWeakObjectPtr<URigVMGraph> GraphPtr;
+
+	/** The blueprint we are editing */
+	TWeakObjectPtr<UControlRigBlueprint> ControlRigBlueprintPtr;
+
+	/** Holds a weak pointer to the argument name widget, used for error notifications */
+	TWeakPtr<SEditableTextBox> ArgumentNameWidget;
+
+	/** The SGraphPin widget created to show/edit default value */
+	TSharedPtr<SGraphPin> DefaultValuePinWidget;
+};
+
+/** Customization for editing Control Rig graphs */
+class FControlRigGraphDetails : public IDetailCustomization
+{
+public:
+
+	/** Makes a new instance of this detail layout class for a specific detail view requesting it */
+	static TSharedPtr<IDetailCustomization> MakeInstance(TSharedPtr<IBlueprintEditor> InBlueprintEditor);
+
+	FControlRigGraphDetails(TSharedPtr<IControlRigEditor> InControlRigEditor, UControlRigBlueprint* ControlRigBlueprint)
+		: ControlRigEditorPtr(InControlRigEditor)
+		, ControlRigBlueprintPtr(ControlRigBlueprint)
+	{}
+
+	// IDetailCustomization interface
+	virtual void CustomizeDetails(IDetailLayoutBuilder& DetailLayout) override;
+
+	bool IsAddNewInputOutputEnabled() const;
+	EVisibility GetAddNewInputOutputVisibility() const;
+	FReply OnAddNewInputClicked();
+	FReply OnAddNewOutputClicked();
+
+private:
+	/** The Blueprint editor we are embedded in */
+	TWeakPtr<IControlRigEditor> ControlRigEditorPtr;
+
+	/** The blueprint we are editing */
+	TWeakObjectPtr<UControlRigBlueprint> ControlRigBlueprintPtr;
+
+	/** The graph we are editing */
+	TWeakObjectPtr<UControlRigGraph> GraphPtr;
+};
\ No newline at end of file
diff --git a/Engine/Plugins/Experimental/ControlRig/Source/ControlRigEditor/Private/Editor/ControlRigEditor.cpp b/Engine/Plugins/Experimental/ControlRig/Source/ControlRigEditor/Private/Editor/ControlRigEditor.cpp
--- a/Engine/Plugins/Experimental/ControlRig/Source/ControlRigEditor/Private/Editor/ControlRigEditor.cpp
+++ b/Engine/Plugins/Experimental/ControlRig/Source/ControlRigEditor/Private/Editor/ControlRigEditor.cpp
@@ -76,6 +76,8 @@
 #include "AssetEditorModeManager.h"
 #include "IPersonaEditorModeManager.h"
 #include "BlueprintEditorTabs.h"
+#include "RigVMModel/Nodes/RigVMFunctionEntryNode.h"
+#include "RigVMModel/Nodes/RigVMFunctionReturnNode.h"
 
 #define LOCTEXT_NAMESPACE "ControlRigEditor"
 
@@ -125,6 +127,7 @@ FControlRigEditor::~FControlRigEditor()
 		}
 		RigBlueprint->OnRefreshEditor().RemoveAll(this);
 		RigBlueprint->OnVariableDropped().RemoveAll(this);
+		RigBlueprint->OnNodeDoubleClicked().RemoveAll(this);
 	}
 
 	if (NodeDetailBuffer.Num() > 0 && NodeDetailStruct != nullptr)
@@ -265,6 +268,10 @@ void FControlRigEditor::InitControlRigEditor(const EToolkitMode::Type Mode, cons
 	// Post-layout initialization
 	PostLayoutBlueprintEditorInitialization();
 
+	// tabs opened before reload
+	FString ActiveTabNodePath;
+	TArray<FString> OpenedTabNodePaths;
+
 	if (ControlRigBlueprints.Num() > 0)
 	{
 		bool bBroughtGraphToFront = false;
@@ -282,6 +289,7 @@ void FControlRigEditor::InitControlRigEditor(const EToolkitMode::Type Mode, cons
 			if (UControlRigGraph* RigGraph = Cast<UControlRigGraph>(Graph))
 			{
 				RigGraph->OnGraphNodeClicked.AddSP(this, &FControlRigEditor::OnGraphNodeClicked);
+				ActiveTabNodePath = RigGraph->ModelNodePath;
 			}
 		}
 	}
@@ -301,6 +309,24 @@ void FControlRigEditor::InitControlRigEditor(const EToolkitMode::Type Mode, cons
 			}
 			else
 			{
+				// remember all ed graphs which were visible as tabs
+				TArray<UEdGraph*> EdGraphs;
+				InControlRigBlueprint->GetAllGraphs(EdGraphs);
+
+				for (UEdGraph* EdGraph : EdGraphs)
+				{
+					if (UControlRigGraph* RigGraph = Cast<UControlRigGraph>(EdGraph))
+					{
+						TArray<TSharedPtr<SDockTab>> TabsForEdGraph;
+						FindOpenTabsContainingDocument(EdGraph, TabsForEdGraph);
+
+						if (TabsForEdGraph.Num() > 0)
+						{
+							OpenedTabNodePaths.Add(RigGraph->ModelNodePath);
+						}
+					}
+				}
+
 				InControlRigBlueprint->RebuildGraphFromModel();
 
 				// selection state does not need to be persistent, even though it is saved in the RigVM.
@@ -332,9 +358,27 @@ void FControlRigEditor::InitControlRigEditor(const EToolkitMode::Type Mode, cons
 			InControlRigBlueprint->HierarchyContainer.OnElementChanged.AddSP(EditMode, &FControlRigEditMode::OnRigElementChanged);
 			InControlRigBlueprint->HierarchyContainer.ControlHierarchy.OnControlUISettingsChanged.AddSP(EditMode, &FControlRigEditMode::OnControlUISettingChanged);
 		}
+
+		InControlRigBlueprint->OnNodeDoubleClicked().AddSP(this, &FControlRigEditor::OnNodeDoubleClicked);
 	}
 
 	UpdateStaleWatchedPins();
+
+	for (const FString& OpenedTabNodePath : OpenedTabNodePaths)
+	{
+		if (UEdGraph* EdGraph = InControlRigBlueprint->GetEdGraph(OpenedTabNodePath))
+		{
+			OpenDocument(EdGraph, FDocumentTracker::RestorePreviousDocument);
+		}
+	}
+
+	if (UEdGraph* ActiveGraph = InControlRigBlueprint->GetEdGraph(ActiveTabNodePath))
+	{
+		OpenGraphAndBringToFront(ActiveGraph, true);
+	}
+
+	FControlRigBlueprintUtils::HandleRefreshAllNodes(InControlRigBlueprint);
+
 	bControlRigEditorInitialized = true;
 }
 
@@ -1069,7 +1113,29 @@ void FControlRigEditor::SetDetailObject(UObject* Obj)
 		}
 		return;
 	}
-	
+	else if (URigVMCollapseNode* CollapseNode = Cast<URigVMCollapseNode>(Obj))
+	{
+		UEdGraph* EdGraph = GetControlRigBlueprint()->GetEdGraph(CollapseNode->GetContainedGraph());
+		if (EdGraph)
+		{
+			TArray<UObject*> Objects;
+			Objects.Add(EdGraph);
+			SetDetailObjects(Objects);
+			return;
+		}
+	}
+	else if (Cast<URigVMFunctionEntryNode>(Obj) || Cast<URigVMFunctionReturnNode>(Obj))
+	{
+		UEdGraph* EdGraph = GetControlRigBlueprint()->GetEdGraph(CastChecked<URigVMNode>(Obj)->GetGraph());
+		if (EdGraph)
+		{
+			TArray<UObject*> Objects;
+			Objects.Add(EdGraph);
+			SetDetailObjects(Objects);
+			return;
+		}
+	}
+
 	TArray<UObject*> Objects;
 	if (Obj)
 	{
@@ -1817,6 +1883,17 @@ void FControlRigEditor::HandleModifiedEvent(ERigVMGraphNotifType InNotifType, UR
 
 			break;
 		}
+		case ERigVMGraphNotifType::NodeRemoved:
+		{
+			if (URigVMCollapseNode* CollapseNode = Cast<URigVMCollapseNode>(InSubject))
+			{
+				if (UEdGraph* EdGraph = ControlRigBlueprint->GetEdGraph(CollapseNode->GetContainedGraph()))
+				{
+					CloseDocumentTab(EdGraph);
+				}
+			}
+			break;
+		}
 		case ERigVMGraphNotifType::NodeSelectionChanged:
 		default:
 		{
@@ -4045,6 +4122,19 @@ void FControlRigEditor::OnGraphNodeClicked(UControlRigGraphNode* InNode)
 	}
 }
 
+void FControlRigEditor::OnNodeDoubleClicked(UControlRigBlueprint* InBlueprint, URigVMNode* InNode)
+{
+	ensure(GetControlRigBlueprint() == InBlueprint);
+
+	if (URigVMLibraryNode* LibraryNode = Cast<URigVMLibraryNode>(InNode))
+	{
+		if (UEdGraph* EdGraph = InBlueprint->GetEdGraph(LibraryNode->GetContainedGraph()))
+		{
+			OpenGraphAndBringToFront(EdGraph, true);
+		}
+	}
+}
+
 void FControlRigEditor::UpdateGraphCompilerErrors()
 {
 	DECLARE_SCOPE_HIERARCHICAL_COUNTER_FUNC()
diff --git a/Engine/Plugins/Experimental/ControlRig/Source/ControlRigEditor/Private/Editor/ControlRigEditor.h b/Engine/Plugins/Experimental/ControlRig/Source/ControlRigEditor/Private/Editor/ControlRigEditor.h
--- a/Engine/Plugins/Experimental/ControlRig/Source/ControlRigEditor/Private/Editor/ControlRigEditor.h
+++ b/Engine/Plugins/Experimental/ControlRig/Source/ControlRigEditor/Private/Editor/ControlRigEditor.h
@@ -305,6 +305,8 @@ private:
 
 	 void OnGraphNodeClicked(UControlRigGraphNode* InNode);
 
+	 void OnNodeDoubleClicked(UControlRigBlueprint* InBlueprint, URigVMNode* InNode);
+
 protected:
 
 	/** Toolbox hosting widget */
diff --git a/Engine/Plugins/Experimental/ControlRig/Source/ControlRigEditor/Private/Graph/NodeSpawners/ControlRigBranchNodeSpawner.cpp b/Engine/Plugins/Experimental/ControlRig/Source/ControlRigEditor/Private/Graph/NodeSpawners/ControlRigBranchNodeSpawner.cpp
--- a/Engine/Plugins/Experimental/ControlRig/Source/ControlRigEditor/Private/Graph/NodeSpawners/ControlRigBranchNodeSpawner.cpp
+++ b/Engine/Plugins/Experimental/ControlRig/Source/ControlRigEditor/Private/Graph/NodeSpawners/ControlRigBranchNodeSpawner.cpp
@@ -12,6 +12,7 @@
 
 #if WITH_EDITOR
 #include "Editor.h"
+#include "Kismet2/BlueprintEditorUtils.h"
 #endif
 
 #define LOCTEXT_NAMESPACE "ControlRigBranchNodeSpawner"
@@ -56,7 +57,7 @@ UEdGraphNode* UControlRigBranchNodeSpawner::Invoke(UEdGraph* ParentGraph, FBindi
 	// First create a backing member for our node
 	UControlRigGraph* RigGraph = Cast<UControlRigGraph>(ParentGraph);
 	check(RigGraph);
-	UControlRigBlueprint* RigBlueprint = Cast<UControlRigBlueprint>(ParentGraph->GetOuter());
+	UControlRigBlueprint* RigBlueprint = Cast<UControlRigBlueprint>(FBlueprintEditorUtils::FindBlueprintForGraph(ParentGraph));
 	check(RigBlueprint);
 
 	FName MemberName = NAME_None;
diff --git a/Engine/Plugins/Experimental/ControlRig/Source/ControlRigEditor/Private/Graph/NodeSpawners/ControlRigEnumNodeSpawner.cpp b/Engine/Plugins/Experimental/ControlRig/Source/ControlRigEditor/Private/Graph/NodeSpawners/ControlRigEnumNodeSpawner.cpp
--- a/Engine/Plugins/Experimental/ControlRig/Source/ControlRigEditor/Private/Graph/NodeSpawners/ControlRigEnumNodeSpawner.cpp
+++ b/Engine/Plugins/Experimental/ControlRig/Source/ControlRigEditor/Private/Graph/NodeSpawners/ControlRigEnumNodeSpawner.cpp
@@ -14,6 +14,7 @@
 
 #if WITH_EDITOR
 #include "Editor.h"
+#include "Kismet2/BlueprintEditorUtils.h"
 #endif
 
 #define LOCTEXT_NAMESPACE "ControlRigEnumNodeSpawner"
@@ -78,7 +79,7 @@ UEdGraphNode* UControlRigEnumNodeSpawner::Invoke(UEdGraph* ParentGraph, FBinding
 	// First create a backing member for our node
 	UControlRigGraph* RigGraph = Cast<UControlRigGraph>(ParentGraph);
 	check(RigGraph);
-	UControlRigBlueprint* RigBlueprint = Cast<UControlRigBlueprint>(ParentGraph->GetOuter());
+	UControlRigBlueprint* RigBlueprint = Cast<UControlRigBlueprint>(FBlueprintEditorUtils::FindBlueprintForGraph(ParentGraph));
 	check(RigBlueprint);
 
 #if WITH_EDITOR
diff --git a/Engine/Plugins/Experimental/ControlRig/Source/ControlRigEditor/Private/Graph/NodeSpawners/ControlRigIfNodeSpawner.cpp b/Engine/Plugins/Experimental/ControlRig/Source/ControlRigEditor/Private/Graph/NodeSpawners/ControlRigIfNodeSpawner.cpp
--- a/Engine/Plugins/Experimental/ControlRig/Source/ControlRigEditor/Private/Graph/NodeSpawners/ControlRigIfNodeSpawner.cpp
+++ b/Engine/Plugins/Experimental/ControlRig/Source/ControlRigEditor/Private/Graph/NodeSpawners/ControlRigIfNodeSpawner.cpp
@@ -13,6 +13,7 @@
 
 #if WITH_EDITOR
 #include "Editor.h"
+#include "Kismet2/BlueprintEditorUtils.h"
 #endif
 
 #define LOCTEXT_NAMESPACE "ControlRigIfNodeSpawner"
@@ -81,7 +82,7 @@ UEdGraphNode* UControlRigIfNodeSpawner::Invoke(UEdGraph* ParentGraph, FBindingSe
 	// First create a backing member for our node
 	UControlRigGraph* RigGraph = Cast<UControlRigGraph>(ParentGraph);
 	check(RigGraph);
-	UControlRigBlueprint* RigBlueprint = Cast<UControlRigBlueprint>(ParentGraph->GetOuter());
+	UControlRigBlueprint* RigBlueprint = Cast<UControlRigBlueprint>(FBlueprintEditorUtils::FindBlueprintForGraph(ParentGraph));
 	check(RigBlueprint);
 
 	FString CPPType = "int32";
diff --git a/Engine/Plugins/Experimental/ControlRig/Source/ControlRigEditor/Private/Graph/NodeSpawners/ControlRigPrototypeNodeSpawner.cpp b/Engine/Plugins/Experimental/ControlRig/Source/ControlRigEditor/Private/Graph/NodeSpawners/ControlRigPrototypeNodeSpawner.cpp
--- a/Engine/Plugins/Experimental/ControlRig/Source/ControlRigEditor/Private/Graph/NodeSpawners/ControlRigPrototypeNodeSpawner.cpp
+++ b/Engine/Plugins/Experimental/ControlRig/Source/ControlRigEditor/Private/Graph/NodeSpawners/ControlRigPrototypeNodeSpawner.cpp
@@ -92,7 +92,7 @@ UEdGraphNode* UControlRigPrototypeNodeSpawner::Invoke(UEdGraph* ParentGraph, FBi
 		}
 #endif
 
-		UBlueprint* Blueprint = CastChecked<UBlueprint>(ParentGraph->GetOuter());
+		UBlueprint* Blueprint = FBlueprintEditorUtils::FindBlueprintForGraph(ParentGraph);
 		NewNode = SpawnNode(ParentGraph, Blueprint, PrototypeNotation, Location);
 	}
 
diff --git a/Engine/Plugins/Experimental/ControlRig/Source/ControlRigEditor/Private/Graph/NodeSpawners/ControlRigRerouteNodeSpawner.cpp b/Engine/Plugins/Experimental/ControlRig/Source/ControlRigEditor/Private/Graph/NodeSpawners/ControlRigRerouteNodeSpawner.cpp
--- a/Engine/Plugins/Experimental/ControlRig/Source/ControlRigEditor/Private/Graph/NodeSpawners/ControlRigRerouteNodeSpawner.cpp
+++ b/Engine/Plugins/Experimental/ControlRig/Source/ControlRigEditor/Private/Graph/NodeSpawners/ControlRigRerouteNodeSpawner.cpp
@@ -67,7 +67,7 @@ UEdGraphNode* UControlRigRerouteNodeSpawner::Invoke(UEdGraph* ParentGraph, FBind
 	// First create a backing member for our node
 	UControlRigGraph* RigGraph = Cast<UControlRigGraph>(ParentGraph);
 	check(RigGraph);
-	UControlRigBlueprint* RigBlueprint = Cast<UControlRigBlueprint>(ParentGraph->GetOuter());
+	UControlRigBlueprint* RigBlueprint = Cast<UControlRigBlueprint>(FBlueprintEditorUtils::FindBlueprintForGraph(ParentGraph));
 	check(RigBlueprint);
 
 	FName MemberName = NAME_None;
diff --git a/Engine/Plugins/Experimental/ControlRig/Source/ControlRigEditor/Private/Graph/NodeSpawners/ControlRigSelectNodeSpawner.cpp b/Engine/Plugins/Experimental/ControlRig/Source/ControlRigEditor/Private/Graph/NodeSpawners/ControlRigSelectNodeSpawner.cpp
--- a/Engine/Plugins/Experimental/ControlRig/Source/ControlRigEditor/Private/Graph/NodeSpawners/ControlRigSelectNodeSpawner.cpp
+++ b/Engine/Plugins/Experimental/ControlRig/Source/ControlRigEditor/Private/Graph/NodeSpawners/ControlRigSelectNodeSpawner.cpp
@@ -14,6 +14,7 @@
 
 #if WITH_EDITOR
 #include "Editor.h"
+#include "Kismet2/BlueprintEditorUtils.h"
 #endif
 
 #define LOCTEXT_NAMESPACE "ControlRigSelectNodeSpawner"
@@ -82,7 +83,7 @@ UEdGraphNode* UControlRigSelectNodeSpawner::Invoke(UEdGraph* ParentGraph, FBindi
 	// First create a backing member for our node
 	UControlRigGraph* RigGraph = Cast<UControlRigGraph>(ParentGraph);
 	check(RigGraph);
-	UControlRigBlueprint* RigBlueprint = Cast<UControlRigBlueprint>(ParentGraph->GetOuter());
+	UControlRigBlueprint* RigBlueprint = Cast<UControlRigBlueprint>(FBlueprintEditorUtils::FindBlueprintForGraph(ParentGraph));
 	check(RigBlueprint);
 
 	FString CPPType = "int32";
diff --git a/Engine/Plugins/Experimental/ControlRig/Source/ControlRigEditor/Private/Graph/NodeSpawners/ControlRigUnitNodeSpawner.cpp b/Engine/Plugins/Experimental/ControlRig/Source/ControlRigEditor/Private/Graph/NodeSpawners/ControlRigUnitNodeSpawner.cpp
--- a/Engine/Plugins/Experimental/ControlRig/Source/ControlRigEditor/Private/Graph/NodeSpawners/ControlRigUnitNodeSpawner.cpp
+++ b/Engine/Plugins/Experimental/ControlRig/Source/ControlRigEditor/Private/Graph/NodeSpawners/ControlRigUnitNodeSpawner.cpp
@@ -101,7 +101,7 @@ UEdGraphNode* UControlRigUnitNodeSpawner::Invoke(UEdGraph* ParentGraph, FBinding
 		}
 #endif
 
-		UBlueprint* Blueprint = CastChecked<UBlueprint>(ParentGraph->GetOuter());
+		UBlueprint* Blueprint = FBlueprintEditorUtils::FindBlueprintForGraph(ParentGraph);
 		NewNode = SpawnNode(ParentGraph, Blueprint, StructTemplate, Location);
 	}
 
@@ -147,6 +147,7 @@ UControlRigGraphNode* UControlRigUnitNodeSpawner::SpawnNode(UEdGraph* ParentGrap
 		if (URigVMUnitNode* ModelNode = Controller->AddUnitNode(StructTemplate, TEXT("Execute"), Location, Name.ToString(), bUndo))
 		{
 			NewNode = Cast<UControlRigGraphNode>(RigGraph->FindNodeForModelNodeName(ModelNode->GetFName()));
+			check(NewNode);
 
 			if (NewNode && bUndo)
 			{
diff --git a/Engine/Plugins/Experimental/ControlRig/Source/ControlRigEditor/Private/Graph/NodeSpawners/ControlRigVariableNodeSpawner.cpp b/Engine/Plugins/Experimental/ControlRig/Source/ControlRigEditor/Private/Graph/NodeSpawners/ControlRigVariableNodeSpawner.cpp
--- a/Engine/Plugins/Experimental/ControlRig/Source/ControlRigEditor/Private/Graph/NodeSpawners/ControlRigVariableNodeSpawner.cpp
+++ b/Engine/Plugins/Experimental/ControlRig/Source/ControlRigEditor/Private/Graph/NodeSpawners/ControlRigVariableNodeSpawner.cpp
@@ -84,7 +84,7 @@ UEdGraphNode* UControlRigVariableNodeSpawner::Invoke(UEdGraph* ParentGraph, FBin
 	// First create a backing member for our node
 	UControlRigGraph* RigGraph = Cast<UControlRigGraph>(ParentGraph);
 	check(RigGraph);
-	UControlRigBlueprint* RigBlueprint = Cast<UControlRigBlueprint>(ParentGraph->GetOuter());
+	UControlRigBlueprint* RigBlueprint = Cast<UControlRigBlueprint>(FBlueprintEditorUtils::FindBlueprintForGraph(ParentGraph));
 	check(RigBlueprint);
 
 	FName MemberName = NAME_None;
diff --git a/Engine/Plugins/Experimental/ControlRig/Source/ControlRigEditor/Private/Graph/SControlRigGraphNode.cpp b/Engine/Plugins/Experimental/ControlRig/Source/ControlRigEditor/Private/Graph/SControlRigGraphNode.cpp
--- a/Engine/Plugins/Experimental/ControlRig/Source/ControlRigEditor/Private/Graph/SControlRigGraphNode.cpp
+++ b/Engine/Plugins/Experimental/ControlRig/Source/ControlRigEditor/Private/Graph/SControlRigGraphNode.cpp
@@ -17,6 +17,7 @@
 #include "Engine/Engine.h"
 #include "KismetNodes/KismetNodeInfoContext.h"
 #include "Kismet2/KismetDebugUtilities.h"
+#include "Kismet2/BlueprintEditorUtils.h"
 #include "PropertyPathHelpers.h"
 #include "UObject/PropertyPortFlags.h"
 #include "ControlRigBlueprint.h"
@@ -318,6 +319,23 @@ FReply SControlRigGraphNode::OnMouseButtonDown(const FGeometry& MyGeometry, cons
 	return Reply;
 }
 
+FReply SControlRigGraphNode::OnMouseButtonDoubleClick(const FGeometry& InMyGeometry, const FPointerEvent& InMouseEvent)
+{
+	if (!InMouseEvent.GetModifierKeys().AnyModifiersDown())
+	{
+		if (UControlRigGraphNode* RigNode = Cast<UControlRigGraphNode>(GraphNode))
+		{
+			if (URigVMNode* ModelNode = RigNode->GetModelNode())
+			{
+				UControlRigBlueprint* RigBlueprint = Cast<UControlRigBlueprint>(FBlueprintEditorUtils::FindBlueprintForNode(RigNode));
+				RigBlueprint->BroadcastNodeDoubleClicked(ModelNode);
+				return FReply::Handled();
+			}
+		}
+	}
+	return SGraphNode::OnMouseButtonDoubleClick(InMyGeometry, InMouseEvent);
+}
+
 bool SControlRigGraphNode::UseLowDetailNodeTitles() const
 {
 	return ParentUseLowDetailNodeTitles();
diff --git a/Engine/Plugins/Experimental/ControlRig/Source/ControlRigEditor/Private/Graph/SControlRigGraphNode.h b/Engine/Plugins/Experimental/ControlRig/Source/ControlRigEditor/Private/Graph/SControlRigGraphNode.h
--- a/Engine/Plugins/Experimental/ControlRig/Source/ControlRigEditor/Private/Graph/SControlRigGraphNode.h
+++ b/Engine/Plugins/Experimental/ControlRig/Source/ControlRigEditor/Private/Graph/SControlRigGraphNode.h
@@ -41,6 +41,7 @@ public:
 	}
 	virtual const FSlateBrush * GetNodeBodyBrush() const override;
 	virtual FReply OnMouseButtonDown(const FGeometry& MyGeometry, const FPointerEvent& MouseEvent) override;
+	virtual FReply OnMouseButtonDoubleClick(const FGeometry& InMyGeometry, const FPointerEvent& InMouseEvent) override;
 
 	virtual TSharedRef<SWidget> CreateNodeContentArea() override;
 	virtual TSharedPtr<SGraphPin> GetHoveredPin( const FGeometry& MyGeometry, const FPointerEvent& MouseEvent ) const override;
diff --git a/Engine/Source/Developer/RigVMDeveloper/Private/RigVMModel/Nodes/RigVMCollapseNode.cpp b/Engine/Source/Developer/RigVMDeveloper/Private/RigVMModel/Nodes/RigVMCollapseNode.cpp
--- a/Engine/Source/Developer/RigVMDeveloper/Private/RigVMModel/Nodes/RigVMCollapseNode.cpp
+++ b/Engine/Source/Developer/RigVMDeveloper/Private/RigVMModel/Nodes/RigVMCollapseNode.cpp
@@ -13,3 +13,8 @@ FText URigVMCollapseNode::GetToolTipText() const
 	// todo
 	return URigVMNode::GetToolTipText();
 }
+
+FString URigVMCollapseNode::GetEditorSubGraphName() const
+{
+	return FString::Printf(TEXT("%s_SubGraph"), *GetName());
+}
\ No newline at end of file
diff --git a/Engine/Source/Developer/RigVMDeveloper/Private/RigVMModel/Nodes/RigVMLibraryNode.cpp b/Engine/Source/Developer/RigVMDeveloper/Private/RigVMModel/Nodes/RigVMLibraryNode.cpp
--- a/Engine/Source/Developer/RigVMDeveloper/Private/RigVMModel/Nodes/RigVMLibraryNode.cpp
+++ b/Engine/Source/Developer/RigVMDeveloper/Private/RigVMModel/Nodes/RigVMLibraryNode.cpp
@@ -50,26 +50,10 @@ const TArray<URigVMLink*>& URigVMLibraryNode::GetContainedLinks() const
 
 URigVMFunctionEntryNode* URigVMLibraryNode::GetEntryNode() const
 {
-	TArray<URigVMNode*> ContainedNodes = GetContainedNodes();
-	for (URigVMNode* ContainedNode : ContainedNodes)
-	{
-		if (URigVMFunctionEntryNode* EntryNode = Cast<URigVMFunctionEntryNode>(ContainedNode))
-		{
-			return EntryNode;
-		}
-	}
-	return nullptr;
+	return GetContainedGraph()->GetEntryNode();
 }
 
 URigVMFunctionReturnNode* URigVMLibraryNode::GetReturnNode() const
 {
-	TArray<URigVMNode*> ContainedNodes = GetContainedNodes();
-	for (URigVMNode* ContainedNode : ContainedNodes)
-	{
-		if (URigVMFunctionReturnNode* ReturnNode = Cast<URigVMFunctionReturnNode>(ContainedNode))
-		{
-			return ReturnNode;
-		}
-	}
-	return nullptr;
+	return GetContainedGraph()->GetReturnNode();
 }
diff --git a/Engine/Source/Developer/RigVMDeveloper/Private/RigVMModel/RigVMController.cpp b/Engine/Source/Developer/RigVMDeveloper/Private/RigVMModel/RigVMController.cpp
--- a/Engine/Source/Developer/RigVMDeveloper/Private/RigVMModel/RigVMController.cpp
+++ b/Engine/Source/Developer/RigVMDeveloper/Private/RigVMModel/RigVMController.cpp
@@ -50,24 +50,85 @@ URigVMController::~URigVMController()
 
 URigVMGraph* URigVMController::GetGraph() const
 {
-	return Graph;
+	if (Graphs.Num() == 0)
+	{
+		return nullptr;
+	}
+	return Graphs.Last();
 }
 
 void URigVMController::SetGraph(URigVMGraph* InGraph)
 {
-	if (Graph)
+	ensure(Graphs.Num() < 2);
+
+	URigVMGraph* PreviousGraph = GetGraph();
+	if (PreviousGraph)
 	{
-		Graph->OnModified().RemoveAll(this);
+		PreviousGraph->OnModified().RemoveAll(this);
 	}
 
-	Graph = InGraph;
+	Graphs.Reset();
+	if (InGraph != nullptr)
+	{
+		PushGraph(InGraph, false);
+	}
 
-	if (Graph)
+	URigVMGraph* CurrentGraph = GetGraph();
+	if (CurrentGraph)
 	{
-		Graph->OnModified().AddUObject(this, &URigVMController::HandleModifiedEvent);
+		CurrentGraph->OnModified().AddUObject(this, &URigVMController::HandleModifiedEvent);
 	}
 
-	HandleModifiedEvent(ERigVMGraphNotifType::GraphChanged, Graph, nullptr);
+	HandleModifiedEvent(ERigVMGraphNotifType::GraphChanged, CurrentGraph, nullptr);
+}
+
+void URigVMController::PushGraph(URigVMGraph* InGraph, bool bSetupUndoRedo)
+{
+	check(InGraph);
+	Graphs.Push(InGraph);
+
+	if (bSetupUndoRedo)
+	{
+		ActionStack->AddAction(FRigVMPushGraphAction(InGraph->GetGraphName()));
+	}
+}
+
+URigVMGraph* URigVMController::PopGraph(bool bSetupUndoRedo)
+{
+	ensure(Graphs.Num() > 1);
+	URigVMGraph* LastGraph = GetGraph();
+	Graphs.Pop();
+
+	if (bSetupUndoRedo)
+	{
+		ActionStack->AddAction(FRigVMPopGraphAction(LastGraph->GetGraphName()));
+	}
+
+	return LastGraph;
+}
+
+URigVMGraph* URigVMController::GetTopLevelGraph() const
+{
+	URigVMGraph* Graph = GetGraph();
+	UObject* Outer = Graph->GetOuter();
+	while (Outer)
+	{
+		if (URigVMGraph* OuterGraph = Cast<URigVMGraph>(Outer))
+		{
+			Graph = OuterGraph;
+			Outer = Outer->GetOuter();
+		}
+		else if (URigVMLibraryNode* LibraryNode = Cast<URigVMLibraryNode>(Outer))
+		{
+			Outer = Outer->GetOuter();
+		}
+		else
+		{
+			break;
+		}
+	}
+
+	return Graph;
 }
 
 FRigVMGraphModifiedEvent& URigVMController::OnModified()
@@ -81,7 +142,7 @@ void URigVMController::Notify(ERigVMGraphNotifType InNotifType, UObject* InSubje
 	{
 		return;
 	}
-	if (Graph)
+	if (URigVMGraph* Graph = GetGraph())
 	{
 		Graph->Notify(InNotifType, InSubject);
 	}
@@ -89,7 +150,7 @@ void URigVMController::Notify(ERigVMGraphNotifType InNotifType, UObject* InSubje
 
 void URigVMController::ResendAllNotifications()
 {
-	if (Graph)
+	if (URigVMGraph* Graph = GetGraph())
 	{
 		for (URigVMLink* Link : Graph->Links)
 		{
@@ -181,6 +242,9 @@ URigVMUnitNode* URigVMController::AddUnitNode(UScriptStruct* InScriptStruct, con
 		return nullptr;
 	}
 
+	URigVMGraph* Graph = GetGraph();
+	check(Graph);
+
 	FString FunctionName = FString::Printf(TEXT("F%s::%s"), *InScriptStruct->GetName(), *InMethodName.ToString());
 	FRigVMFunctionPtr Function = FRigVMRegistry::Get().FindFunction(*FunctionName);
 	if (Function == nullptr)
@@ -289,6 +353,9 @@ URigVMVariableNode* URigVMController::AddVariableNode(const FName& InVariableNam
 		return nullptr;
 	}
 
+	URigVMGraph* Graph = GetGraph();
+	check(Graph);
+
 	if (InCPPTypeObject == nullptr)
 	{
 		InCPPTypeObject = URigVMCompiler::GetScriptStructForCPPType(InCPPType);
@@ -411,6 +478,9 @@ void URigVMController::RefreshVariableNode(const FName& InNodeName, const FName&
 		return;
 	}
 
+	URigVMGraph* Graph = GetGraph();
+	check(Graph);
+
 	if (URigVMVariableNode* VariableNode = Cast<URigVMVariableNode>(Graph->FindNodeByName(InNodeName)))
 	{
 		if (URigVMPin* VariablePin = VariableNode->FindPin(URigVMVariableNode::VariableName))
@@ -494,6 +564,9 @@ void URigVMController::OnExternalVariableRemoved(const FName& InVarName, bool bS
 		return;
 	}
 
+	URigVMGraph* Graph = GetGraph();
+	check(Graph);
+
 	FString VarNameStr = InVarName.ToString();
 
 	if (bSetupUndoRedo)
@@ -544,6 +617,9 @@ void URigVMController::OnExternalVariableRenamed(const FName& InOldVarName, cons
 		return;
 	}
 
+	URigVMGraph* Graph = GetGraph();
+	check(Graph);
+
 	FString VarNameStr = InOldVarName.ToString();
 
 	if (bSetupUndoRedo)
@@ -596,6 +672,9 @@ void URigVMController::OnExternalVariableTypeChanged(const FName& InVarName, con
 		return;
 	}
 
+	URigVMGraph* Graph = GetGraph();
+	check(Graph);
+
 	FString VarNameStr = InVarName.ToString();
 
 	if (bSetupUndoRedo)
@@ -644,6 +723,9 @@ URigVMVariableNode* URigVMController::ReplaceParameterNodeWithVariable(const FNa
 		return nullptr;
 	}
 
+	URigVMGraph* Graph = GetGraph();
+	check(Graph);
+
 	if (URigVMParameterNode* ParameterNode = Cast<URigVMParameterNode>(Graph->FindNodeByName(InNodeName)))
 	{
 		URigVMPin* ParameterValuePin = ParameterNode->FindPin(URigVMParameterNode::ValueName);
@@ -688,6 +770,9 @@ URigVMParameterNode* URigVMController::AddParameterNode(const FName& InParameter
 		return nullptr;
 	}
 
+	URigVMGraph* Graph = GetGraph();
+	check(Graph);
+
 	if (InCPPTypeObject == nullptr)
 	{
 		InCPPTypeObject = URigVMCompiler::GetScriptStructForCPPType(InCPPType);
@@ -865,6 +950,9 @@ URigVMCommentNode* URigVMController::AddCommentNode(const FString& InCommentText
 		return nullptr;
 	}
 
+	URigVMGraph* Graph = GetGraph();
+	check(Graph);
+
 	FString Name = GetValidNodeName(InNodeName.IsEmpty() ? FString(TEXT("CommentNode")) : InNodeName);
 	URigVMCommentNode* Node = NewObject<URigVMCommentNode>(Graph, *Name);
 	Node->Position = InPosition;
@@ -943,6 +1031,9 @@ URigVMRerouteNode* URigVMController::AddRerouteNodeOnLinkPath(const FString& InL
 		return nullptr;
 	}
 
+	URigVMGraph* Graph = GetGraph();
+	check(Graph);
+
 	URigVMLink* Link = Graph->FindLink(InLinkPinPathRepresentation);
 	return AddRerouteNodeOnLink(Link, bShowAsFullNode, InPosition, InNodeName, bSetupUndoRedo);
 }
@@ -954,6 +1045,9 @@ URigVMRerouteNode* URigVMController::AddRerouteNodeOnPin(const FString& InPinPat
 		return nullptr;
 	}
 
+	URigVMGraph* Graph = GetGraph();
+	check(Graph);
+
 	URigVMPin* Pin = Graph->FindPin(InPinPath);
 	if(Pin == nullptr)
 	{
@@ -1033,6 +1127,9 @@ URigVMInjectionInfo* URigVMController::AddInjectedNode(const FString& InPinPath,
 		return nullptr;
 	}
 
+	URigVMGraph* Graph = GetGraph();
+	check(Graph);
+
 	URigVMPin* Pin = Graph->FindPin(InPinPath);
 	if (Pin == nullptr)
 	{
@@ -1237,6 +1334,9 @@ URigVMNode* URigVMController::EjectNodeFromPin(const FString& InPinPath, bool bS
 		return nullptr;
 	}
 
+	URigVMGraph* Graph = GetGraph();
+	check(Graph);
+
 	URigVMPin* Pin = Graph->FindPin(InPinPath);
 	if (Pin == nullptr)
 	{
@@ -1402,6 +1502,9 @@ FString URigVMController::ExportNodesToText(const TArray<FName>& InNodeNames)
 		return FString();
 	}
 
+	URigVMGraph* Graph = GetGraph();
+	check(Graph);
+
 	UnMarkAllObjects(EObjectMark(OBJECTMARK_TagExp | OBJECTMARK_TagImp));
 
 	FStringOutputDevice Archive;
@@ -1459,6 +1562,10 @@ FString URigVMController::ExportSelectedNodesToText()
 	{
 		return FString();
 	}
+
+	URigVMGraph* Graph = GetGraph();
+	check(Graph);
+
 	return ExportNodesToText(Graph->GetSelectNodes());
 }
 
@@ -1547,6 +1654,9 @@ TArray<FName> URigVMController::ImportNodesFromText(const FString& InText, bool
 		return NodeNames;
 	}
 
+	URigVMGraph* Graph = GetGraph();
+	check(Graph);
+
 	FRigVMControllerObjectFactory Factory(this);
 	Factory.ProcessBuffer(Graph, RF_Transactional, InText);
 
@@ -1675,6 +1785,9 @@ URigVMCollapseNode* URigVMController::CollapseNodes(const TArray<FName>& InNodeN
 		return nullptr;
 	}
 
+	URigVMGraph* Graph = GetGraph();
+	check(Graph);
+
 	TArray<URigVMNode*> Nodes;
 	for (const FName& NodeName : InNodeNames)
 	{
@@ -1697,6 +1810,9 @@ TArray<URigVMNode*> URigVMController::ExpandLibraryNode(const FName& InNodeName,
 		return TArray<URigVMNode*>();
 	}
 
+	URigVMGraph* Graph = GetGraph();
+	check(Graph);
+
 	URigVMNode* Node = Graph->FindNodeByName(InNodeName);
 	if (Node == nullptr)
 	{
@@ -1718,6 +1834,9 @@ TArray<URigVMNode*> URigVMController::ExpandLibraryNode(const FName& InNodeName,
 
 URigVMCollapseNode* URigVMController::CollapseNodes(const TArray<URigVMNode*>& InNodes, const FString& InCollapseNodeName, bool bSetupUndoRedo)
 {
+	URigVMGraph* Graph = GetGraph();
+	check(Graph);
+
 	TArray<URigVMNode*> Nodes;
 	for (URigVMNode* Node : InNodes)
 	{
@@ -1987,18 +2106,27 @@ URigVMCollapseNode* URigVMController::CollapseNodes(const TArray<URigVMNode*>& I
 
 	Notify(ERigVMGraphNotifType::NodeAdded, CollapseNode);
 
-	URigVMFunctionEntryNode* EntryNode = NewObject<URigVMFunctionEntryNode>(CollapseNode->ContainedGraph, TEXT("Entry"));
-	CollapseNode->ContainedGraph->Nodes.Add(EntryNode);
-	EntryNode->Position = -Diagonal * 0.5f - FVector2D(160.f, 0.f);
-	RefreshFunctionPins(EntryNode);
-
+	URigVMFunctionEntryNode* EntryNode = nullptr;
 	URigVMFunctionReturnNode* ReturnNode = nullptr;
-	if (bContainsOutputs)
 	{
-		ReturnNode = NewObject<URigVMFunctionReturnNode>(CollapseNode->ContainedGraph, TEXT("Return"));
-		CollapseNode->ContainedGraph->Nodes.Add(ReturnNode);
-		ReturnNode->Position = FVector2D(Diagonal.X, -Diagonal.Y) * 0.5f + FVector2D(160.f, 0.f);
-		RefreshFunctionPins(ReturnNode);
+		FRigVMControllerGraphGuard GraphGuard(this, CollapseNode->GetContainedGraph(), false);
+
+		EntryNode = NewObject<URigVMFunctionEntryNode>(CollapseNode->ContainedGraph, TEXT("Entry"));
+		CollapseNode->ContainedGraph->Nodes.Add(EntryNode);
+		EntryNode->Position = -Diagonal * 0.5f - FVector2D(250.f, 0.f);
+		RefreshFunctionPins(EntryNode, false);
+
+		Notify(ERigVMGraphNotifType::NodeAdded, EntryNode);
+
+		if (bContainsOutputs)
+		{
+			ReturnNode = NewObject<URigVMFunctionReturnNode>(CollapseNode->ContainedGraph, TEXT("Return"));
+			CollapseNode->ContainedGraph->Nodes.Add(ReturnNode);
+			ReturnNode->Position = FVector2D(Diagonal.X, -Diagonal.Y) * 0.5f + FVector2D(300.f, 0.f);
+			RefreshFunctionPins(ReturnNode, false);
+
+			Notify(ERigVMGraphNotifType::NodeAdded, ReturnNode);
+		}
 	}
 
 	// create the new nodes within the collapse node
@@ -2006,15 +2134,14 @@ URigVMCollapseNode* URigVMController::CollapseNodes(const TArray<URigVMNode*>& I
 	{
 		FString TextContent = ExportNodesToText(NodeNames);
 
-		TGuardValue<URigVMGraph*> GraphGuard(Graph, CollapseNode->ContainedGraph);
+		FRigVMControllerGraphGuard GraphGuard(this, CollapseNode->GetContainedGraph(), false);
 		ContainedNodeNames = ImportNodesFromText(TextContent, false);
 
 		// move the nodes to the right place
 		for (const FName& ContainedNodeName : ContainedNodeNames)
 		{
-			if (URigVMNode* ContainedNode = CollapseNode->ContainedGraph->FindNodeByName(ContainedNodeName))
+			if (URigVMNode* ContainedNode = CollapseNode->GetContainedGraph()->FindNodeByName(ContainedNodeName))
 			{
-				ContainedNode->Position -= Center;
 				SetNodePosition(ContainedNode, ContainedNode->Position - Center, false, false);
 			}
 		}
@@ -2146,6 +2273,9 @@ TArray<URigVMNode*> URigVMController::ExpandLibraryNode(URigVMLibraryNode* InNod
 		return TArray<URigVMNode*>();
 	}
 
+	URigVMGraph* Graph = GetGraph();
+	check(Graph);
+
 	TArray<URigVMNode*> ContainedNodes = InNode->GetContainedNodes();
 	TArray<URigVMLink*> ContainedLinks = InNode->GetContainedLinks();
 	if (ContainedNodes.Num() == 0)
@@ -2190,7 +2320,7 @@ TArray<URigVMNode*> URigVMController::ExpandLibraryNode(URigVMLibraryNode* InNod
 
 	FString TextContent;
 	{
-		TGuardValue<URigVMGraph*> GraphGuard(Graph, InNode->GetContainedGraph());
+		FRigVMControllerGraphGuard GraphGuard(this, InNode->GetContainedGraph(), false);
 		TextContent = ExportNodesToText(NodeNames);
 	}
 
@@ -2663,17 +2793,22 @@ TArray<URigVMNode*> URigVMController::ExpandLibraryNode(URigVMLibraryNode* InNod
 	return ExpandedNodes;
 }
 
-void URigVMController::RefreshFunctionPins(URigVMNode* InNode)
+void URigVMController::RefreshFunctionPins(URigVMNode* InNode, bool bNotify)
 {
+	if (InNode == nullptr)
+	{
+		return;
+	}
+
 	URigVMFunctionEntryNode* EntryNode = Cast<URigVMFunctionEntryNode>(InNode);
 	URigVMFunctionReturnNode* ReturnNode = Cast<URigVMFunctionReturnNode>(InNode);
 
 	if (EntryNode || ReturnNode)
 	{
 		TArray<URigVMLink*> Links = InNode->GetLinks();
-		DetachLinksFromPinObjects(&Links);
-		RepopulatePinsOnNode(InNode, false);
-		ReattachLinksToPinObjects(false, &Links);
+		DetachLinksFromPinObjects(&Links, bNotify);
+		RepopulatePinsOnNode(InNode, false, bNotify);
+		ReattachLinksToPinObjects(false, &Links, bNotify);
 	}
 }
 
@@ -2684,6 +2819,9 @@ bool URigVMController::RemoveNode(URigVMNode* InNode, bool bSetupUndoRedo, bool
 		return false;
 	}
 
+	URigVMGraph* Graph = GetGraph();
+	check(Graph);
+
 	if (bSetupUndoRedo)
 	{
 		// don't allow deletion of function entry / return nodes
@@ -2740,6 +2878,23 @@ bool URigVMController::RemoveNode(URigVMNode* InNode, bool bSetupUndoRedo, bool
 			BreakAllLinksRecursive(Pin, true, false, bSetupUndoRedo);
 			BreakAllLinksRecursive(Pin, false, false, bSetupUndoRedo);
 		}
+
+		if (URigVMCollapseNode* CollapseNode = Cast<URigVMCollapseNode>(InNode))
+		{
+			URigVMGraph* SubGraph = CollapseNode->GetContainedGraph();
+			FRigVMControllerGraphGuard GraphGuard(this, SubGraph, bSetupUndoRedo);
+
+			TArray<URigVMNode*> ContainedNodes = SubGraph->GetNodes();
+			for (URigVMNode* ContainedNode : ContainedNodes)
+			{
+				if(Cast<URigVMFunctionEntryNode>(ContainedNode) != nullptr ||
+					Cast<URigVMFunctionReturnNode>(ContainedNode) != nullptr)
+				{
+					continue;
+				}
+				RemoveNode(ContainedNode, bSetupUndoRedo, bRecursive);
+			}
+		}
 	}
 
 	if (bSetupUndoRedo)
@@ -2785,9 +2940,65 @@ bool URigVMController::RemoveNodeByName(const FName& InNodeName, bool bSetupUndo
 	{
 		return false;
 	}
+
+	URigVMGraph* Graph = GetGraph();
+	check(Graph);
+
 	return RemoveNode(Graph->FindNodeByName(InNodeName), bSetupUndoRedo, bRecursive);
 }
 
+bool URigVMController::RenameNode(URigVMNode* InNode, const FName& InNewName, bool bSetupUndoRedo)
+{
+	if (!IsValidNodeForGraph(InNode))
+	{
+		return false;
+	}
+
+	FName ValidNewName = *GetValidNodeName(InNewName.ToString());
+	if (InNode->GetFName() == ValidNewName)
+	{
+		return false;
+	}
+
+	FRigVMRenameNodeAction Action;
+	if (bSetupUndoRedo)
+	{
+		Action = FRigVMRenameNodeAction(InNode->GetFName(), ValidNewName);
+		ActionStack->BeginAction(Action);
+	}
+
+	// loop over all links and remove them
+	TArray<URigVMLink*> Links = InNode->GetLinks();
+	for (URigVMLink* Link : Links)
+	{
+		Link->PrepareForCopy();
+		Notify(ERigVMGraphNotifType::LinkRemoved, Link);
+	}
+
+	InNode->PreviousName = InNode->GetFName();
+	if (!InNode->Rename(*ValidNewName.ToString()))
+	{
+		ActionStack->CancelAction(Action);
+		return false;
+	}
+
+	Notify(ERigVMGraphNotifType::NodeRenamed, InNode);
+
+	// update the links once more
+	for (URigVMLink* Link : Links)
+	{
+		Link->PrepareForCopy();
+		Notify(ERigVMGraphNotifType::LinkAdded, Link);
+	}
+
+	if (bSetupUndoRedo)
+	{
+		ActionStack->EndAction(Action);
+	}
+
+	return true;
+}
+
 bool URigVMController::SelectNode(URigVMNode* InNode, bool bSelect, bool bSetupUndoRedo)
 {
 	if (!IsValidNodeForGraph(InNode))
@@ -2800,6 +3011,9 @@ bool URigVMController::SelectNode(URigVMNode* InNode, bool bSelect, bool bSetupU
 		return false;
 	}
 
+	URigVMGraph* Graph = GetGraph();
+	check(Graph);
+
 	TArray<FName> NewSelection = Graph->GetSelectNodes();
 	if (bSelect)
 	{
@@ -2819,6 +3033,10 @@ bool URigVMController::SelectNodeByName(const FName& InNodeName, bool bSelect, b
 	{
 		return false;
 	}
+
+	URigVMGraph* Graph = GetGraph();
+	check(Graph);
+
 	return SelectNode(Graph->FindNodeByName(InNodeName), bSelect, bSetupUndoRedo);
 }
 
@@ -2839,6 +3057,9 @@ bool URigVMController::SetNodeSelection(const TArray<FName>& InNodeNames, bool b
 		return false;
 	}
 
+	URigVMGraph* Graph = GetGraph();
+	check(Graph);
+
 	FRigVMSetNodeSelectionAction Action;
 	if (bSetupUndoRedo)
 	{
@@ -2949,6 +3170,9 @@ bool URigVMController::SetNodePositionByName(const FName& InNodeName, const FVec
 		return false;
 	}
 
+	URigVMGraph* Graph = GetGraph();
+	check(Graph);
+
 	URigVMNode* Node = Graph->FindNodeByName(InNodeName);
 	return SetNodePosition(Node, InPosition, bSetupUndoRedo, bMergeUndoAction);
 }
@@ -2991,6 +3215,9 @@ bool URigVMController::SetNodeSizeByName(const FName& InNodeName, const FVector2
 		return false;
 	}
 
+	URigVMGraph* Graph = GetGraph();
+	check(Graph);
+
 	URigVMNode* Node = Graph->FindNodeByName(InNodeName);
 	return SetNodeSize(Node, InSize, bSetupUndoRedo, bMergeUndoAction);
 }
@@ -3033,6 +3260,9 @@ bool URigVMController::SetNodeColorByName(const FName& InNodeName, const FLinear
 		return false;
 	}
 
+	URigVMGraph* Graph = GetGraph();
+	check(Graph);
+
 	URigVMNode* Node = Graph->FindNodeByName(InNodeName);
 	return SetNodeColor(Node, InColor, bSetupUndoRedo, bMergeUndoAction);
 }
@@ -3080,6 +3310,9 @@ bool URigVMController::SetCommentTextByName(const FName& InNodeName, const FStri
 		return false;
 	}
 
+	URigVMGraph* Graph = GetGraph();
+	check(Graph);
+
 	URigVMNode* Node = Graph->FindNodeByName(InNodeName);
 	return SetCommentText(Node, InCommentText, bSetupUndoRedo);
 }
@@ -3127,6 +3360,9 @@ bool URigVMController::SetRerouteCompactnessByName(const FName& InNodeName, bool
 		return false;
 	}
 
+	URigVMGraph* Graph = GetGraph();
+	check(Graph);
+
 	URigVMNode* Node = Graph->FindNodeByName(InNodeName);
 	return SetRerouteCompactness(Node, bShowAsFullNode, bSetupUndoRedo);
 }
@@ -3144,6 +3380,9 @@ bool URigVMController::RenameVariable(const FName& InOldName, const FName& InNew
 		return false;
 	}
 
+	URigVMGraph* Graph = GetGraph();
+	check(Graph);
+
 	TArray<FRigVMGraphVariableDescription> ExistingVariables = Graph->GetVariableDescriptions();
 	for (const FRigVMGraphVariableDescription& ExistingVariable : ExistingVariables)
 	{
@@ -3212,6 +3451,9 @@ bool URigVMController::RenameParameter(const FName& InOldName, const FName& InNe
 		return false;
 	}
 
+	URigVMGraph* Graph = GetGraph();
+	check(Graph);
+
 	TArray<FRigVMGraphParameterDescription> ExistingParameters = Graph->GetParameterDescriptions();
 	for (const FRigVMGraphParameterDescription& ExistingParameter : ExistingParameters)
 	{
@@ -3304,6 +3546,9 @@ bool URigVMController::SetPinExpansion(const FString& InPinPath, bool bIsExpande
 		return false;
 	}
 
+	URigVMGraph* Graph = GetGraph();
+	check(Graph);
+
 	URigVMPin* Pin = Graph->FindPin(InPinPath);
 	if (Pin == nullptr)
 	{
@@ -3326,6 +3571,9 @@ bool URigVMController::SetPinExpansion(URigVMPin* InPin, bool bIsExpanded, bool
 		return false;
 	}
 
+	URigVMGraph* Graph = GetGraph();
+	check(Graph);
+
 	FRigVMSetPinExpansionAction Action;
 	if (bSetupUndoRedo)
 	{
@@ -3357,6 +3605,9 @@ bool URigVMController::SetPinIsWatched(const FString& InPinPath, bool bIsWatched
 		return false;
 	}
 
+	URigVMGraph* Graph = GetGraph();
+	check(Graph);
+
 	URigVMPin* Pin = Graph->FindPin(InPinPath);
 	if (Pin == nullptr)
 	{
@@ -3384,6 +3635,9 @@ bool URigVMController::SetPinIsWatched(URigVMPin* InPin, bool bIsWatched, bool b
 		return false;
 	}
 
+	URigVMGraph* Graph = GetGraph();
+	check(Graph);
+
 	FRigVMSetPinWatchAction Action;
 	if (bSetupUndoRedo)
 	{
@@ -3415,6 +3669,9 @@ FString URigVMController::GetPinDefaultValue(const FString& InPinPath)
 		return FString();
 	}
 
+	URigVMGraph* Graph = GetGraph();
+	check(Graph);
+
 	URigVMPin* Pin = Graph->FindPin(InPinPath);
 	if (Pin == nullptr)
 	{
@@ -3433,6 +3690,9 @@ bool URigVMController::SetPinDefaultValue(const FString& InPinPath, const FStrin
 		return false;
 	}
 
+	URigVMGraph* Graph = GetGraph();
+	check(Graph);
+
 	URigVMPin* Pin = Graph->FindPin(InPinPath);
 	if (Pin == nullptr)
 	{
@@ -3475,6 +3735,9 @@ void URigVMController::SetPinDefaultValue(URigVMPin* InPin, const FString& InDef
 		return;
 	}
 
+	URigVMGraph* Graph = GetGraph();
+	check(Graph);
+
 	FRigVMSetPinDefaultValueAction Action;
 	if (bSetupUndoRedo)
 	{
@@ -3567,6 +3830,9 @@ bool URigVMController::ResetPinDefaultValue(const FString& InPinPath, bool bSetu
 		return false;
 	}
 
+	URigVMGraph* Graph = GetGraph();
+	check(Graph);
+
 	URigVMPin* Pin = Graph->FindPin(InPinPath);
 	if (Pin == nullptr)
 	{
@@ -3725,6 +3991,9 @@ FString URigVMController::DuplicateArrayPin(const FString& InArrayElementPinPath
 		return FString();
 	}
 
+	URigVMGraph* Graph = GetGraph();
+	check(Graph);
+
 	URigVMPin* ElementPin = Graph->FindPin(InArrayElementPinPath);
 	if (ElementPin == nullptr)
 	{
@@ -3753,6 +4022,9 @@ FString URigVMController::InsertArrayPin(const FString& InArrayPinPath, int32 In
 		return FString();
 	}
 
+	URigVMGraph* Graph = GetGraph();
+	check(Graph);
+
 	URigVMPin* ArrayPin = Graph->FindPin(InArrayPinPath);
 	if (ArrayPin == nullptr)
 	{
@@ -3783,6 +4055,9 @@ URigVMPin* URigVMController::InsertArrayPin(URigVMPin* ArrayPin, int32 InIndex,
 		return nullptr;
 	}
 
+	URigVMGraph* Graph = GetGraph();
+	check(Graph);
+
 	if (InIndex == INDEX_NONE)
 	{
 		InIndex = ArrayPin->GetSubPins().Num();
@@ -3854,6 +4129,9 @@ bool URigVMController::RemoveArrayPin(const FString& InArrayElementPinPath, bool
 		return false;
 	}
 
+	URigVMGraph* Graph = GetGraph();
+	check(Graph);
+
 	URigVMPin* ArrayElementPin = Graph->FindPin(InArrayElementPinPath);
 	if (ArrayElementPin == nullptr)
 	{
@@ -3880,7 +4158,7 @@ bool URigVMController::RemoveArrayPin(const FString& InArrayElementPinPath, bool
 	}
 
 	int32 IndexToRemove = ArrayElementPin->GetPinIndex();
-	if (!RemovePin(ArrayElementPin, bSetupUndoRedo))
+	if (!RemovePin(ArrayElementPin, bSetupUndoRedo, false))
 	{
 		return false;
 	}
@@ -3905,7 +4183,7 @@ bool URigVMController::RemoveArrayPin(const FString& InArrayElementPinPath, bool
 	return true;
 }
 
-bool URigVMController::RemovePin(URigVMPin* InPinToRemove, bool bSetupUndoRedo)
+bool URigVMController::RemovePin(URigVMPin* InPinToRemove, bool bSetupUndoRedo, bool bNotify)
 {
 	if (bSetupUndoRedo)
 	{
@@ -3924,17 +4202,22 @@ bool URigVMController::RemovePin(URigVMPin* InPinToRemove, bool bSetupUndoRedo)
 		Node->Pins.Remove(InPinToRemove);
 	}
 
-	DestroyObject(InPinToRemove);
-
 	TArray<URigVMPin*> SubPins = InPinToRemove->GetSubPins();
 	for (URigVMPin* SubPin : SubPins)
 	{
-		if (!RemovePin(SubPin, bSetupUndoRedo))
+		if (!RemovePin(SubPin, bSetupUndoRedo, bNotify))
 		{
 			return false;
 		}
 	}
 
+	if (bNotify)
+	{
+		Notify(ERigVMGraphNotifType::PinRemoved, InPinToRemove);
+	}
+
+	DestroyObject(InPinToRemove);
+
 	return true;
 }
 
@@ -3950,6 +4233,9 @@ bool URigVMController::SetArrayPinSize(const FString& InArrayPinPath, int32 InSi
 		return false;
 	}
 
+	URigVMGraph* Graph = GetGraph();
+	check(Graph);
+
 	URigVMPin* Pin = Graph->FindPin(InArrayPinPath);
 	if (Pin == nullptr)
 	{
@@ -4032,6 +4318,9 @@ bool URigVMController::BindPinToVariable(const FString& InPinPath, const FString
 		return false;
 	}
 
+	URigVMGraph* Graph = GetGraph();
+	check(Graph);
+
 	URigVMPin* Pin = Graph->FindPin(InPinPath);
 	if (Pin == nullptr)
 	{
@@ -4129,6 +4418,9 @@ bool URigVMController::MakeBindingsFromVariableNode(const FName& InNodeName, boo
 		return false;
 	}
 
+	URigVMGraph* Graph = GetGraph();
+	check(Graph);
+
 	if (URigVMVariableNode* VariableNode = Cast<URigVMVariableNode>(Graph->FindNodeByName(InNodeName)))
 	{
 		return MakeBindingsFromVariableNode(VariableNode, bSetupUndoRedo);
@@ -4230,6 +4522,9 @@ bool URigVMController::PromotePinToVariable(const FString& InPinPath, bool bCrea
 		return false;
 	}
 
+	URigVMGraph* Graph = GetGraph();
+	check(Graph);
+
 	URigVMPin* Pin = Graph->FindPin(InPinPath);
 	if (Pin == nullptr)
 	{
@@ -4319,6 +4614,9 @@ bool URigVMController::AddLink(const FString& InOutputPinPath, const FString& In
 		return false;
 	}
 
+	URigVMGraph* Graph = GetGraph();
+	check(Graph);
+
 	FString OutputPinPath = InOutputPinPath;
 	FString InputPinPath = InInputPinPath;
 
@@ -4369,6 +4667,9 @@ bool URigVMController::AddLink(URigVMPin* OutputPin, URigVMPin* InputPin, bool b
 		return false;
 	}
 
+	URigVMGraph* Graph = GetGraph();
+	check(Graph);
+
 	{
 		// Temporarily remove the bound variable from the pin
 		// if it exists, so that link validation can work.
@@ -4450,6 +4751,9 @@ bool URigVMController::BreakLink(const FString& InOutputPinPath, const FString&
 		return false;
 	}
 
+	URigVMGraph* Graph = GetGraph();
+	check(Graph);
+
 	URigVMPin* OutputPin = Graph->FindPin(InOutputPinPath);
 	if (OutputPin == nullptr)
 	{
@@ -4482,6 +4786,9 @@ bool URigVMController::BreakLink(URigVMPin* OutputPin, URigVMPin* InputPin, bool
 	}
 	ensure(InputPin->IsLinkedTo(OutputPin));
 
+	URigVMGraph* Graph = GetGraph();
+	check(Graph);
+
 	for (URigVMLink* Link : InputPin->Links)
 	{
 		if (Link->SourcePin == OutputPin && Link->TargetPin == InputPin)
@@ -4528,6 +4835,9 @@ bool URigVMController::BreakAllLinks(const FString& InPinPath, bool bAsInput, bo
 		return false;
 	}
 
+	URigVMGraph* Graph = GetGraph();
+	check(Graph);
+
 	URigVMPin* Pin = Graph->FindPin(InPinPath);
 	if (Pin == nullptr)
 	{
@@ -4610,6 +4920,307 @@ void URigVMController::BreakAllLinksRecursive(URigVMPin* Pin, bool bAsInput, boo
 	}
 }
 
+FName URigVMController::AddExposedPin(const FName& InPinName, ERigVMPinDirection InDirection, const FString& InCPPType, const FName& InCPPTypeObjectPath, const FString& InDefaultValue, bool bSetupUndoRedo)
+{
+	if (!IsValidGraph())
+	{
+		return NAME_None;
+	}
+
+	URigVMGraph* Graph = GetGraph();
+	check(Graph);
+
+	if (Graph->IsTopLevelGraph())
+	{
+		ReportError(TEXT("Exposed pins can only be edited on nested graphs."));
+		return NAME_None;
+	}
+
+	URigVMLibraryNode* LibraryNode = Cast<URigVMLibraryNode>(Graph->GetOuter());
+	check(LibraryNode);
+
+	FName PinName = GetUniqueName(InPinName, [LibraryNode](const FName& InName) {
+
+		return LibraryNode->FindPin(InName.ToString()) == nullptr;
+
+	});
+
+	URigVMPin* Pin = NewObject<URigVMPin>(LibraryNode, PinName);
+	Pin->CPPType = InCPPType;
+	Pin->CPPTypeObjectPath = InCPPTypeObjectPath;
+	Pin->bIsConstant = false;
+	Pin->Direction = InDirection;
+	LibraryNode->Pins.Add(Pin);
+
+	if (Pin->IsStruct())
+	{
+		FRigVMControllerGraphGuard GraphGuard(this, LibraryNode->GetGraph(), false);
+
+		FString DefaultValue = InDefaultValue;
+		CreateDefaultValueForStructIfRequired(Pin->GetScriptStruct(), DefaultValue);
+		AddPinsForStruct(Pin->GetScriptStruct(), LibraryNode, Pin, Pin->Direction, DefaultValue, false);
+	}
+
+	FRigVMAddExposedPinAction Action(Pin);
+	if (bSetupUndoRedo)
+	{
+		ActionStack->BeginAction(Action);
+	}
+
+	{
+		FRigVMControllerGraphGuard GraphGuard(this, LibraryNode->GetGraph(), false);
+		Notify(ERigVMGraphNotifType::PinAdded, Pin);
+	}
+
+	if (bSetupUndoRedo)
+	{
+		SetPinDefaultValue(Pin, InDefaultValue, true, bSetupUndoRedo, false);
+	}
+
+	RefreshFunctionPins(Graph->GetEntryNode(), true);
+	RefreshFunctionPins(Graph->GetReturnNode(), true);
+
+	if (bSetupUndoRedo)
+	{
+		ActionStack->EndAction(Action);
+	}
+	
+	return NAME_None;
+}
+
+bool URigVMController::RemoveExposedPin(const FName& InPinName, bool bSetupUndoRedo)
+{
+	if (!IsValidGraph())
+	{
+		return NAME_None;
+	}
+
+	URigVMGraph* Graph = GetGraph();
+	check(Graph);
+
+	if (Graph->IsTopLevelGraph())
+	{
+		ReportError(TEXT("Exposed pins can only be edited on nested graphs."));
+		return NAME_None;
+	}
+
+	URigVMLibraryNode* LibraryNode = Cast<URigVMLibraryNode>(Graph->GetOuter());
+	check(LibraryNode);
+
+	URigVMPin* Pin = LibraryNode->FindPin(InPinName.ToString());
+	if (Pin == nullptr)
+	{
+		return false;
+	}
+
+	FRigVMRemoveExposedPinAction Action(Pin);
+	if (bSetupUndoRedo)
+	{
+		ActionStack->BeginAction(Action);
+	}
+
+	bool bSuccessfullyRemovedPin = false;
+	{
+		FRigVMControllerGraphGuard GraphGuard(this, LibraryNode->GetGraph(), false);
+		bSuccessfullyRemovedPin = RemovePin(Pin, bSetupUndoRedo, true);
+	}
+
+	RefreshFunctionPins(Graph->GetEntryNode(), true);
+	RefreshFunctionPins(Graph->GetReturnNode(), true);
+
+	if (bSetupUndoRedo)
+	{
+		if (bSuccessfullyRemovedPin)
+		{
+			ActionStack->EndAction(Action);
+		}
+		else
+		{
+			ActionStack->CancelAction(Action);
+		}
+	}
+
+	return bSuccessfullyRemovedPin;
+}
+
+bool URigVMController::RenameExposedPin(const FName& InOldPinName, const FName& InNewPinName, bool bSetupUndoRedo)
+{
+	if (!IsValidGraph())
+	{
+		return NAME_None;
+	}
+
+	URigVMGraph* Graph = GetGraph();
+	check(Graph);
+
+	if (Graph->IsTopLevelGraph())
+	{
+		ReportError(TEXT("Exposed pins can only be edited on nested graphs."));
+		return NAME_None;
+	}
+
+	URigVMLibraryNode* LibraryNode = Cast<URigVMLibraryNode>(Graph->GetOuter());
+	check(LibraryNode);
+
+	URigVMPin* Pin = LibraryNode->FindPin(InOldPinName.ToString());
+	if (Pin == nullptr)
+	{
+		return false;
+	}
+
+	if (Pin->GetFName() == InNewPinName)
+	{
+		return false;
+	}
+
+	FName PinName = GetUniqueName(InNewPinName, [LibraryNode](const FName& InName) {
+		return LibraryNode->FindPin(InName.ToString()) == nullptr;
+		});
+
+	FRigVMRenameExposedPinAction Action;
+	if (bSetupUndoRedo)
+	{
+		Action = FRigVMRenameExposedPinAction(Pin->GetFName(), PinName);
+		ActionStack->BeginAction(Action);
+	}
+
+	struct Local
+	{
+		static bool RenamePin(URigVMController* InController, URigVMPin* InPin, const FName& InNewName)
+		{
+			FRigVMControllerGraphGuard GraphGuard(InController, InPin->GetGraph(), false);
+
+			TArray<URigVMLink*> Links;
+			Links.Append(InPin->GetSourceLinks(true));
+			Links.Append(InPin->GetTargetLinks(true));
+
+			// store both the ptr + pin path
+			for (URigVMLink* Link : Links)
+			{
+				Link->PrepareForCopy();
+				InController->Notify(ERigVMGraphNotifType::LinkRemoved, Link);
+			}
+
+			if (!InPin->Rename(*InNewName.ToString()))
+			{
+				return false;
+			}
+
+			// update the eventually stored pin path to the new name
+			for (URigVMLink* Link : Links)
+			{
+				Link->PrepareForCopy();
+			}
+
+			InController->Notify(ERigVMGraphNotifType::PinRenamed, InPin);
+
+			for (URigVMLink* Link : Links)
+			{
+				InController->Notify(ERigVMGraphNotifType::LinkAdded, Link);
+			}
+
+			return true;
+		}
+	};
+
+	if (!Local::RenamePin(this, Pin, PinName))
+	{
+		ActionStack->CancelAction(Action);
+		return false;
+	}
+
+	if (URigVMFunctionEntryNode* EntryNode = Graph->GetEntryNode())
+	{
+		if (URigVMPin* EntryPin = EntryNode->FindPin(InOldPinName.ToString()))
+		{
+			Local::RenamePin(this, EntryPin, PinName);
+		}
+	}
+
+	if (URigVMFunctionReturnNode* ReturnNode = Graph->GetReturnNode())
+	{
+		if (URigVMPin* ReturnPin = ReturnNode->FindPin(InOldPinName.ToString()))
+		{
+			Local::RenamePin(this, ReturnPin, PinName);
+		}
+	}
+
+	if (bSetupUndoRedo)
+	{
+		ActionStack->EndAction(Action);
+	}
+
+	return true;
+}
+
+bool URigVMController::ChangeExposedPinType(const FName& InPinName, const FString& InCPPType, const FName& InCPPTypeObjectPath, bool bSetupUndoRedo)
+{
+	if (!IsValidGraph())
+	{
+		return NAME_None;
+	}
+
+	URigVMGraph* Graph = GetGraph();
+	check(Graph);
+
+	if (Graph->IsTopLevelGraph())
+	{
+		ReportError(TEXT("Exposed pins can only be edited on nested graphs."));
+		return NAME_None;
+	}
+
+	URigVMLibraryNode* LibraryNode = Cast<URigVMLibraryNode>(Graph->GetOuter());
+	check(LibraryNode);
+
+	URigVMPin* Pin = LibraryNode->FindPin(InPinName.ToString());
+	if (Pin == nullptr)
+	{
+		return false;
+	}
+
+	FRigVMBaseAction Action;
+	if (bSetupUndoRedo)
+	{
+		Action.Title = FString::Printf(TEXT("Change Pin Type"));
+		ActionStack->BeginAction(Action);
+	}
+
+	if (bSetupUndoRedo)
+	{
+		FRigVMControllerGraphGuard GraphGuard(this, LibraryNode->GetGraph(), bSetupUndoRedo);
+		if (!ChangePinType(Pin, InCPPType, InCPPTypeObjectPath, bSetupUndoRedo))
+		{
+			if (bSetupUndoRedo)
+			{
+				ActionStack->CancelAction(Action);
+			}
+			return false;
+		}
+	}
+
+	if (URigVMFunctionEntryNode* EntryNode = Graph->GetEntryNode())
+	{
+		if (URigVMPin* EntryPin = EntryNode->FindPin(Pin->GetName()))
+		{
+			ChangePinType(EntryPin, InCPPType, InCPPTypeObjectPath, bSetupUndoRedo);
+		}
+	}
+	if (URigVMFunctionReturnNode* ReturnNode = Graph->GetReturnNode())
+	{
+		if (URigVMPin* ReturnPin = ReturnNode->FindPin(Pin->GetName()))
+		{
+			ChangePinType(ReturnPin, InCPPType, InCPPTypeObjectPath, bSetupUndoRedo);
+		}
+	}
+
+	if (bSetupUndoRedo)
+	{
+		ActionStack->EndAction(Action);
+	}
+
+	return true;
+}
+
 void URigVMController::ExpandPinRecursively(URigVMPin* InPin, bool bSetupUndoRedo)
 {
 	if (InPin == nullptr)
@@ -4662,6 +5273,9 @@ bool URigVMController::SetVariableName(URigVMVariableNode* InVariableNode, const
 		return false;
 	}
 
+	URigVMGraph* Graph = GetGraph();
+	check(Graph);
+
 	TArray<FRigVMGraphVariableDescription> Descriptions = Graph->GetVariableDescriptions();
 	TMap<FName, int32> NameToIndex;
 	for (int32 VariableIndex = 0; VariableIndex < Descriptions.Num(); VariableIndex++)
@@ -4719,6 +5333,9 @@ bool URigVMController::SetParameterName(URigVMParameterNode* InParameterNode, co
 		return false;
 	}
 
+	URigVMGraph* Graph = GetGraph();
+	check(Graph);
+
 	TArray<FRigVMGraphParameterDescription> Descriptions = Graph->GetParameterDescriptions();
 	TMap<FName, int32> NameToIndex;
 	for (int32 ParameterIndex = 0; ParameterIndex < Descriptions.Num(); ParameterIndex++)
@@ -4766,6 +5383,9 @@ URigVMRerouteNode* URigVMController::AddFreeRerouteNode(bool bShowAsFullNode, co
 		return nullptr;
 	}
 
+	URigVMGraph* Graph = GetGraph();
+	check(Graph);
+
 	FRigVMBaseAction Action;
 	if (bSetupUndoRedo)
 	{
@@ -4820,6 +5440,9 @@ URigVMBranchNode* URigVMController::AddBranchNode(const FVector2D& InPosition, c
 		return nullptr;
 	}
 
+	URigVMGraph* Graph = GetGraph();
+	check(Graph);
+
 	FString Name = GetValidNodeName(InNodeName.IsEmpty() ? FString(TEXT("BranchNode")) : InNodeName);
 	URigVMBranchNode* Node = NewObject<URigVMBranchNode>(Graph, *Name);
 	Node->Position = InPosition;
@@ -4870,6 +5493,9 @@ URigVMIfNode* URigVMController::AddIfNode(const FString& InCPPType, const FName&
 		return nullptr;
 	}
 
+	URigVMGraph* Graph = GetGraph();
+	check(Graph);
+
 	ensure(!InCPPType.IsEmpty());
 
 	FString CPPType = InCPPType;
@@ -4963,6 +5589,9 @@ URigVMSelectNode* URigVMController::AddSelectNode(const FString& InCPPType, cons
 		return nullptr;
 	}
 
+	URigVMGraph* Graph = GetGraph();
+	check(Graph);
+
 	ensure(!InCPPType.IsEmpty());
 
 	FString CPPType = InCPPType;
@@ -5035,6 +5664,9 @@ URigVMPrototypeNode* URigVMController::AddPrototypeNode(const FName& InNotation,
 		return nullptr;
 	}
 
+	URigVMGraph* Graph = GetGraph();
+	check(Graph);
+
 	ensure(!InNotation.IsNone());
 
 	const FRigVMPrototype* Prototype = FRigVMRegistry::Get().FindPrototype(InNotation);
@@ -5089,6 +5721,9 @@ URigVMEnumNode* URigVMController::AddEnumNode(const FName& InCPPTypeObjectPath,
 		return nullptr;
 	}
 
+	URigVMGraph* Graph = GetGraph();
+	check(Graph);
+
     UObject* CPPTypeObject = URigVMPin::FindObjectFromCPPTypeObjectPath<UObject>(InCPPTypeObjectPath.ToString());
 	if (CPPTypeObject == nullptr)
 	{
@@ -5159,6 +5794,9 @@ void URigVMController::SetExecuteContextStruct(UStruct* InExecuteContextStruct)
 
 FString URigVMController::GetValidNodeName(const FString& InPrefix)
 {
+	URigVMGraph* Graph = GetGraph();
+	check(Graph);
+
 	return GetUniqueName(*InPrefix, [&](const FName& InName) {
 		return Graph->IsNameAvailable(InName.ToString());
 	}).ToString();
@@ -5166,6 +5804,7 @@ FString URigVMController::GetValidNodeName(const FString& InPrefix)
 
 bool URigVMController::IsValidGraph()
 {
+	URigVMGraph* Graph = GetGraph();
 	if (Graph == nullptr)
 	{
 		ReportError(TEXT("Controller does not have a graph associated - use SetGraph / set_graph."));
@@ -5187,7 +5826,7 @@ bool URigVMController::IsValidNodeForGraph(URigVMNode* InNode)
 		return false;
 	}
 
-	if (InNode->GetGraph() != Graph)
+	if (InNode->GetGraph() != GetGraph())
 	{
 		ReportErrorf(TEXT("InNode '%s' is on a different graph."), *InNode->GetNodePath());
 		return false;
@@ -5240,7 +5879,7 @@ bool URigVMController::IsValidLinkForGraph(URigVMLink* InLink)
 		return false;
 	}
 
-	if (InLink->GetGraph() != Graph)
+	if (InLink->GetGraph() != GetGraph())
 	{
 		ReportError(TEXT("InLink is on a different graph."));
 		return false;
@@ -5276,7 +5915,7 @@ bool URigVMController::IsValidLinkForGraph(URigVMLink* InLink)
 	return true;
 }
 
-void URigVMController::AddPinsForStruct(UStruct* InStruct, URigVMNode* InNode, URigVMPin* InParentPin, ERigVMPinDirection InPinDirection, const FString& InDefaultValue, bool bAutoExpandArrays)
+void URigVMController::AddPinsForStruct(UStruct* InStruct, URigVMNode* InNode, URigVMPin* InParentPin, ERigVMPinDirection InPinDirection, const FString& InDefaultValue, bool bAutoExpandArrays, bool bNotify)
 {
 	TArray<FString> MemberNameValuePairs = SplitDefaultValue(InDefaultValue);
 	TMap<FName, FString> MemberValues;
@@ -5348,12 +5987,17 @@ void URigVMController::AddPinsForStruct(UStruct* InStruct, URigVMNode* InNode, U
 			}
 		}
 		
-if (!Pin->IsArray() && !Pin->IsStruct() && DefaultValuePtr != nullptr)
-{
-	FString DefaultValue = *DefaultValuePtr;
-	PostProcessDefaultValue(Pin, DefaultValue);
-	Pin->DefaultValue = DefaultValue;
-}
+		if (!Pin->IsArray() && !Pin->IsStruct() && DefaultValuePtr != nullptr)
+		{
+			FString DefaultValue = *DefaultValuePtr;
+			PostProcessDefaultValue(Pin, DefaultValue);
+			Pin->DefaultValue = DefaultValue;
+		}
+
+		if (bNotify)
+		{
+			Notify(ERigVMGraphNotifType::PinAdded, Pin);
+		}
 	}
 }
 
@@ -5564,6 +6208,9 @@ FProperty* URigVMController::FindPropertyForPin(const FString& InPinPath)
 		return nullptr;
 	}
 
+	URigVMGraph* Graph = GetGraph();
+	check(Graph);
+
 	URigVMPin* Pin = Graph->FindPin(InPinPath);
 	if (Pin == nullptr)
 	{
@@ -5609,10 +6256,11 @@ FProperty* URigVMController::FindPropertyForPin(const FString& InPinPath)
 	return nullptr;
 }
 
-int32 URigVMController::DetachLinksFromPinObjects(const TArray<URigVMLink*>* InLinks)
+int32 URigVMController::DetachLinksFromPinObjects(const TArray<URigVMLink*>* InLinks, bool bNotify)
 {
+	URigVMGraph* Graph = GetGraph();
 	check(Graph);
-	TGuardValue<bool> SuspendNotifs(bSuspendNotifications, true);
+	TGuardValue<bool> EventuallySuspendNotifs(bSuspendNotifications, !bNotify);
 
 	TArray<URigVMLink*> Links;
 	if (InLinks)
@@ -5626,6 +6274,8 @@ int32 URigVMController::DetachLinksFromPinObjects(const TArray<URigVMLink*>* InL
 
 	for (URigVMLink* Link : Links)
 	{
+		Notify(ERigVMGraphNotifType::LinkRemoved, Link);
+
 		URigVMPin* SourcePin = Link->GetSourcePin();
 		URigVMPin* TargetPin = Link->GetTargetPin();
 
@@ -5651,8 +6301,8 @@ int32 URigVMController::DetachLinksFromPinObjects(const TArray<URigVMLink*>* InL
 		{
 			if (URigVMCollapseNode* CollapseNode = Cast<URigVMCollapseNode>(Node))
 			{
-				TGuardValue<URigVMGraph*> GraphGuard(Graph, CollapseNode->ContainedGraph);
-				DetachLinksFromPinObjects();
+				FRigVMControllerGraphGuard GraphGuard(this, CollapseNode->GetContainedGraph(), false);
+				DetachLinksFromPinObjects(InLinks, bNotify);
 			}
 		}
 	}
@@ -5660,10 +6310,11 @@ int32 URigVMController::DetachLinksFromPinObjects(const TArray<URigVMLink*>* InL
 	return Links.Num();
 }
 
-int32 URigVMController::ReattachLinksToPinObjects(bool bFollowCoreRedirectors, const TArray<URigVMLink*>* InLinks)
+int32 URigVMController::ReattachLinksToPinObjects(bool bFollowCoreRedirectors, const TArray<URigVMLink*>* InLinks, bool bNotify)
 {
+	URigVMGraph* Graph = GetGraph();
 	check(Graph);
-	TGuardValue<bool> SuspendNotifs(bSuspendNotifications, true);
+	TGuardValue<bool> EventuallySuspendNotifs(bSuspendNotifications, !bNotify);
 	FScopeLock Lock(&PinPathCoreRedirectorsLock);
 
 	bool bReplacingAllLinks = false;
@@ -5742,6 +6393,11 @@ int32 URigVMController::ReattachLinksToPinObjects(bool bFollowCoreRedirectors, c
 	if (bReplacingAllLinks)
 	{
 		Graph->Links = NewLinks;
+
+		for (URigVMLink* Link : Graph->Links)
+		{
+			Notify(ERigVMGraphNotifType::LinkAdded, Link);
+		}
 	}
 	else
 	{
@@ -5753,6 +6409,11 @@ int32 URigVMController::ReattachLinksToPinObjects(bool bFollowCoreRedirectors, c
 			if (!NewLinks.Contains(Link))
 			{
 				Graph->Links.Remove(Link);
+				Notify(ERigVMGraphNotifType::LinkRemoved, Link);
+			}
+			else
+			{
+				Notify(ERigVMGraphNotifType::LinkAdded, Link);
 			}
 		}
 	}
@@ -5763,7 +6424,7 @@ int32 URigVMController::ReattachLinksToPinObjects(bool bFollowCoreRedirectors, c
 		{
 			if (URigVMCollapseNode* CollapseNode = Cast<URigVMCollapseNode>(Node))
 			{
-				TGuardValue<URigVMGraph*> GraphGuard(Graph, CollapseNode->ContainedGraph);
+				FRigVMControllerGraphGuard GraphGuard(this, CollapseNode->GetContainedGraph(), false);
 				ReattachLinksToPinObjects(bFollowCoreRedirectors, nullptr);
 			}
 		}
@@ -5782,6 +6443,9 @@ void URigVMController::RemoveStaleNodes()
 		return;
 	}
 
+	URigVMGraph* Graph = GetGraph();
+	check(Graph);
+
 	Graph->Nodes.Remove(nullptr);
 }
 
@@ -5910,6 +6574,9 @@ bool URigVMController::ShouldRedirectPin(UScriptStruct* InOwningStruct, const FS
 
 bool URigVMController::ShouldRedirectPin(const FString& InOldPinPath, FString& InOutNewPinPath) const
 {
+	URigVMGraph* Graph = GetGraph();
+	check(Graph);
+
 	FString PinPathInNode, NodeName;
 	URigVMPin::SplitPinPathAtStart(InOldPinPath, NodeName, PinPathInNode);
 
@@ -5953,7 +6620,7 @@ bool URigVMController::ShouldRedirectPin(const FString& InOldPinPath, FString& I
 	return false;
 }
 
-void URigVMController::RepopulatePinsOnNode(URigVMNode* InNode, bool bFollowCoreRedirectors)
+void URigVMController::RepopulatePinsOnNode(URigVMNode* InNode, bool bFollowCoreRedirectors, bool bNotify)
 {
 	if (InNode == nullptr)
 	{
@@ -5967,9 +6634,10 @@ void URigVMController::RepopulatePinsOnNode(URigVMNode* InNode, bool bFollowCore
 	URigVMFunctionReturnNode* ReturnNode = Cast<URigVMFunctionReturnNode>(InNode);
 	URigVMCollapseNode* CollapseNode = Cast<URigVMCollapseNode>(InNode);
 
-	TGuardValue<bool> SuspendNotifs(bSuspendNotifications, true);
+	TGuardValue<bool> EventuallySuspendNotifs(bSuspendNotifications, !bNotify);
 	FScopeLock Lock(&PinPathCoreRedirectorsLock);
 
+	URigVMGraph* Graph = GetGraph();
 	check(Graph);
 
 	// step 1/3: keep a record of the current state of the node's pins
@@ -5995,7 +6663,7 @@ void URigVMController::RepopulatePinsOnNode(URigVMNode* InNode, bool bFollowCore
 		TArray<URigVMPin*> Pins = InNode->GetPins();
 		for (URigVMPin* Pin : Pins)
 		{
-			RemovePin(Pin, false);
+			RemovePin(Pin, false, bNotify);
 		}
 		InNode->Pins.Reset();
 		Pins.Reset();
@@ -6022,7 +6690,7 @@ void URigVMController::RepopulatePinsOnNode(URigVMNode* InNode, bool bFollowCore
 
 		FString ExportedDefaultValue;
 		CreateDefaultValueForStructIfRequired(ScriptStruct, ExportedDefaultValue);
-		AddPinsForStruct(ScriptStruct, UnitNode, nullptr, ERigVMPinDirection::Invalid, ExportedDefaultValue, false);
+		AddPinsForStruct(ScriptStruct, UnitNode, nullptr, ERigVMPinDirection::Invalid, ExportedDefaultValue, false, bNotify);
 	}
 	else if (RerouteNode != nullptr)
 	{
@@ -6037,7 +6705,7 @@ void URigVMController::RepopulatePinsOnNode(URigVMNode* InNode, bool bFollowCore
 		TArray<URigVMPin*> Pins = ValuePin->SubPins;
 		for (URigVMPin* Pin : Pins)
 		{
-			RemovePin(Pin, false);
+			RemovePin(Pin, false, bNotify);
 		}
 		ValuePin->SubPins.Reset();
 		Pins.Reset();
@@ -6071,7 +6739,7 @@ void URigVMController::RepopulatePinsOnNode(URigVMNode* InNode, bool bFollowCore
 			TArray<URigVMPin*> Pins = InNode->GetPins();
 			for (URigVMPin* Pin : Pins)
 			{
-				RemovePin(Pin, false);
+				RemovePin(Pin, false, bNotify);
 			}
 			InNode->Pins.Reset();
 			Pins.Reset();
@@ -6116,6 +6784,8 @@ void URigVMController::RepopulatePinsOnNode(URigVMNode* InNode, bool bFollowCore
 				{
 					AddPinsForStruct(ExposedPin->GetScriptStruct(), InNode, ExposedPin, ExposedPin->GetDirection(), FString(), false);
 				}
+
+				Notify(ERigVMGraphNotifType::PinAdded, ExposedPin);
 			}
 		}
 		else
@@ -6126,7 +6796,7 @@ void URigVMController::RepopulatePinsOnNode(URigVMNode* InNode, bool bFollowCore
 	}
 	else if (CollapseNode)
 	{
-		TGuardValue<URigVMGraph*> GraphGuard(Graph, CollapseNode->ContainedGraph);
+		FRigVMControllerGraphGuard GraphGuard(this, CollapseNode->GetContainedGraph(), false);
 		for (URigVMNode* ContainedNode : CollapseNode->GetContainedNodes())
 		{
 			RepopulatePinsOnNode(ContainedNode, bFollowCoreRedirectors);
@@ -6340,7 +7010,7 @@ void URigVMController::ReportWarning(const FString& InMessage)
 	}
 
 	FString Message = InMessage;
-	if (Graph != nullptr)
+	if (URigVMGraph* Graph = GetGraph())
 	{
 		if (UPackage* Package = Cast<UPackage>(Graph->GetOutermost()))
 		{
@@ -6359,7 +7029,7 @@ void URigVMController::ReportError(const FString& InMessage)
 	}
 
 	FString Message = InMessage;
-	if (Graph != nullptr)
+	if (URigVMGraph* Graph = GetGraph())
 	{
 		if (UPackage* Package = Cast<UPackage>(Graph->GetOutermost()))
 		{
@@ -6553,6 +7223,9 @@ bool URigVMController::ChangePinType(const FString& InPinPath, const FString& In
 		return false;
 	}
 
+	URigVMGraph* Graph = GetGraph();
+	check(Graph);
+
 	if (URigVMPin* Pin = Graph->FindPin(InPinPath))
 	{
 		return ChangePinType(Pin, InCPPType, InCPPTypeObjectPath, bSetupUndoRedo);
@@ -6568,17 +7241,64 @@ bool URigVMController::ChangePinType(URigVMPin* InPin, const FString& InCPPType,
 		return false;
 	}
 
-	FRigVMChangePinTypeAction Action(InPin, InCPPType, InCPPTypeObjectPath);
+	if (InCPPType == TEXT("None") || InCPPType.IsEmpty())
+	{
+		return false;
+	}
+
+	UObject* CPPTypeObject = URigVMPin::FindObjectFromCPPTypeObjectPath(InCPPTypeObjectPath.ToString());
+	if (CPPTypeObject)
+	{
+		// for now we don't allow UObjects
+		if (!CPPTypeObject->IsA<UEnum>() &&
+			!CPPTypeObject->IsA<UStruct>())
+		{
+			return false;
+		}
+	}
+
+	FRigVMBaseAction Action;
+	if (bSetupUndoRedo)
+	{
+		Action.Title = TEXT("Change pin type");
+		ActionStack->BeginAction(Action);
+	}
+
+	if (bSetupUndoRedo)
+	{
+		BreakAllLinks(InPin, true, true);
+		BreakAllLinks(InPin, false, true);
+		BreakAllLinksRecursive(InPin, true, false, true);
+		BreakAllLinksRecursive(InPin, false, false, true);
+
+		ActionStack->AddAction(FRigVMChangePinTypeAction(InPin, InCPPType, InCPPTypeObjectPath));
+	}
+
+	TArray<URigVMPin*> Pins = InPin->SubPins;
+	for (URigVMPin* Pin : Pins)
+	{
+		RemovePin(Pin, false, true);
+	}
+	InPin->SubPins.Reset();
 
 	InPin->CPPType = InCPPType;
 	InPin->CPPTypeObjectPath = InCPPTypeObjectPath;
-	InPin->CPPTypeObject = URigVMPin::FindObjectFromCPPTypeObjectPath(InPin->CPPTypeObjectPath.ToString());
+	InPin->CPPTypeObject = CPPTypeObject;
+	InPin->DefaultValue = FString();
+
+	if (InPin->IsStruct())
+	{
+		FString DefaultValue = InPin->DefaultValue;
+		CreateDefaultValueForStructIfRequired(InPin->GetScriptStruct(), DefaultValue);
+		AddPinsForStruct(InPin->GetScriptStruct(), InPin->GetNode(), InPin, InPin->Direction, DefaultValue, false, true);
+	}
 
 	Notify(ERigVMGraphNotifType::PinTypeChanged, InPin);
+	Notify(ERigVMGraphNotifType::PinDefaultValueChanged, InPin);
 
 	if (bSetupUndoRedo)
 	{
-		ActionStack->AddAction(Action);
+		ActionStack->EndAction(Action);
 	}
 
 	return true;
diff --git a/Engine/Source/Developer/RigVMDeveloper/Private/RigVMModel/RigVMControllerActions.cpp b/Engine/Source/Developer/RigVMDeveloper/Private/RigVMModel/RigVMControllerActions.cpp
--- a/Engine/Source/Developer/RigVMDeveloper/Private/RigVMModel/RigVMControllerActions.cpp
+++ b/Engine/Source/Developer/RigVMDeveloper/Private/RigVMModel/RigVMControllerActions.cpp
@@ -1399,4 +1399,191 @@ bool FRigVMExpandNodeAction::Redo(URigVMController* InController)
 	}
 #endif
 	return false;
-}
\ No newline at end of file
+}
+
+FRigVMRenameNodeAction::FRigVMRenameNodeAction(const FName& InOldNodeName, const FName& InNewNodeName)
+	: OldNodeName(InOldNodeName.ToString())
+	, NewNodeName(InNewNodeName.ToString())
+{
+}
+
+bool FRigVMRenameNodeAction::Undo(URigVMController* InController)
+{
+	if (!FRigVMBaseAction::Undo(InController))
+	{
+		return false;
+	}
+	if (URigVMNode* Node = InController->GetGraph()->FindNode(NewNodeName))
+	{
+		return InController->RenameNode(Node, *OldNodeName, false);
+	}
+	return false;
+}
+
+bool FRigVMRenameNodeAction::Redo(URigVMController* InController)
+{
+	if (URigVMNode* Node = InController->GetGraph()->FindNode(OldNodeName))
+	{
+		return InController->RenameNode(Node, *NewNodeName, false);
+	}
+	else
+	{
+		return false;
+	}
+	return FRigVMBaseAction::Redo(InController);
+}
+
+FRigVMPushGraphAction::FRigVMPushGraphAction(const FString& InGraphName)
+	: GraphName(InGraphName)
+{
+}
+
+bool FRigVMPushGraphAction::Undo(URigVMController* InController)
+{
+	if (!FRigVMBaseAction::Undo(InController))
+	{
+		return false;
+	}
+	return InController->PopGraph(false) != nullptr;
+}
+
+bool FRigVMPushGraphAction::Redo(URigVMController* InController)
+{
+	if (GraphName.IsEmpty())
+	{
+		InController->PushGraph(InController->GetTopLevelGraph(), false);
+		return FRigVMBaseAction::Redo(InController);
+	}
+	else if (URigVMNode* Node = InController->GetGraph()->FindNode(GraphName))
+	{
+		if (URigVMCollapseNode* CollapseNode = Cast<URigVMCollapseNode>(Node))
+		{
+			InController->PushGraph(CollapseNode->GetContainedGraph(), false);
+			return FRigVMBaseAction::Redo(InController);
+		}
+	}
+	return false;
+}
+
+FRigVMPopGraphAction::FRigVMPopGraphAction(const FString& InGraphName)
+	: GraphName(InGraphName)
+{
+}
+
+bool FRigVMPopGraphAction::Undo(URigVMController* InController)
+{
+	if (!FRigVMBaseAction::Undo(InController))
+	{
+		return false;
+	}
+
+	if (GraphName.IsEmpty())
+	{
+		InController->PushGraph(InController->GetTopLevelGraph(), false);
+		return true;
+	}
+	else if (URigVMNode* Node = InController->GetGraph()->FindNode(GraphName))
+	{
+		if (URigVMCollapseNode* CollapseNode = Cast<URigVMCollapseNode>(Node))
+		{
+			InController->PushGraph(CollapseNode->GetContainedGraph(), false);
+			return true;
+		}
+	}
+	return false;
+}
+
+bool FRigVMPopGraphAction::Redo(URigVMController* InController)
+{
+	if (InController->PopGraph(false) != nullptr)
+	{
+		return FRigVMBaseAction::Redo(InController);
+	}
+	return false;
+}
+
+FRigVMAddExposedPinAction::FRigVMAddExposedPinAction(URigVMPin* InPin)
+	: PinName(InPin->GetName())
+	, Direction(InPin->GetDirection())
+	, CPPType(InPin->GetCPPType())
+	, CPPTypeObjectPath()
+	, DefaultValue(InPin->GetDefaultValue())
+{
+	if (UObject* CPPTypeObject = InPin->GetCPPTypeObject())
+	{
+		CPPTypeObjectPath = CPPTypeObject->GetPathName();
+	}
+}
+
+bool FRigVMAddExposedPinAction::Undo(URigVMController* InController)
+{
+	if (!FRigVMBaseAction::Undo(InController))
+	{
+		return false;
+	}
+	return InController->RemoveExposedPin(*PinName, false);
+}
+
+bool FRigVMAddExposedPinAction::Redo(URigVMController* InController)
+{
+	if (!InController->AddExposedPin(*PinName, Direction, CPPType, *CPPTypeObjectPath, DefaultValue, false).IsNone())
+	{
+		return FRigVMBaseAction::Redo(InController);
+	}
+	return false;
+}
+
+FRigVMRemoveExposedPinAction::FRigVMRemoveExposedPinAction(URigVMPin* InPin)
+	: PinName(InPin->GetName())
+	, Direction(InPin->GetDirection())
+	, CPPType(InPin->GetCPPType())
+	, CPPTypeObjectPath()
+	, DefaultValue(InPin->GetDefaultValue())
+{
+	if (UObject* CPPTypeObject = InPin->GetCPPTypeObject())
+	{
+		CPPTypeObjectPath = CPPTypeObject->GetPathName();
+	}
+}
+
+bool FRigVMRemoveExposedPinAction::Undo(URigVMController* InController)
+{
+	if(!InController->AddExposedPin(*PinName, Direction, CPPType, *CPPTypeObjectPath, DefaultValue, false).IsNone())
+	{
+		return FRigVMBaseAction::Undo(InController);
+	}
+	return false;
+}
+
+bool FRigVMRemoveExposedPinAction::Redo(URigVMController* InController)
+{
+	if(FRigVMBaseAction::Redo(InController))
+	{
+		return InController->RemoveExposedPin(*PinName, false);
+	}
+	return false;
+}
+
+FRigVMRenameExposedPinAction::FRigVMRenameExposedPinAction(const FName& InOldPinName, const FName& InNewPinName)
+	: OldPinName(InOldPinName.ToString())
+	, NewPinName(InNewPinName.ToString())
+{
+}
+
+bool FRigVMRenameExposedPinAction::Undo(URigVMController* InController)
+{
+	if (!FRigVMBaseAction::Undo(InController))
+	{
+		return false;
+	}
+	return InController->RenameExposedPin(*NewPinName, *OldPinName, false);
+}
+
+bool FRigVMRenameExposedPinAction::Redo(URigVMController* InController)
+{
+	if(!InController->RenameExposedPin(*OldPinName, *NewPinName, false))
+	{
+		return false;
+	}
+	return FRigVMBaseAction::Redo(InController);
+}
diff --git a/Engine/Source/Developer/RigVMDeveloper/Private/RigVMModel/RigVMGraph.cpp b/Engine/Source/Developer/RigVMDeveloper/Private/RigVMModel/RigVMGraph.cpp
--- a/Engine/Source/Developer/RigVMDeveloper/Private/RigVMModel/RigVMGraph.cpp
+++ b/Engine/Source/Developer/RigVMDeveloper/Private/RigVMModel/RigVMGraph.cpp
@@ -3,6 +3,8 @@
 #include "RigVMModel/RigVMGraph.h"
 #include "RigVMModel/RigVMController.h"
 #include "RigVMModel/Nodes/RigVMLibraryNode.h"
+#include "RigVMModel/Nodes/RigVMFunctionEntryNode.h"
+#include "RigVMModel/Nodes/RigVMFunctionReturnNode.h"
 #include "UObject/Package.h"
 
 URigVMGraph::URigVMGraph()
@@ -38,6 +40,30 @@ TArray<URigVMGraph*> URigVMGraph::GetContainedGraphs(bool bRecursive) const
 	return Graphs;
 }
 
+URigVMFunctionEntryNode* URigVMGraph::GetEntryNode() const
+{
+	for (URigVMNode* Node : Nodes)
+	{
+		if (URigVMFunctionEntryNode* EntryNode = Cast<URigVMFunctionEntryNode>(Node))
+		{
+			return EntryNode;
+		}
+	}
+	return nullptr;
+}
+
+URigVMFunctionReturnNode* URigVMGraph::GetReturnNode() const
+{
+	for (URigVMNode* Node : Nodes)
+	{
+		if (URigVMFunctionReturnNode* ReturnNode = Cast<URigVMFunctionReturnNode>(Node))
+		{
+			return ReturnNode;
+		}
+	}
+	return nullptr;
+}
+
 TArray<FRigVMGraphVariableDescription> URigVMGraph::GetVariableDescriptions() const
 {
 	TArray<FRigVMGraphVariableDescription> Variables;
@@ -68,7 +94,16 @@ FString URigVMGraph::GetNodePath() const
 {
 	if (URigVMCollapseNode* CollapseNode = Cast<URigVMCollapseNode>(GetOuter()))
 	{
-		return CollapseNode->GetNodePath();
+		return CollapseNode->GetNodePath(true /* recursive */);
+	}
+	return FString();
+}
+
+FString URigVMGraph::GetGraphName() const
+{
+	if (URigVMCollapseNode* CollapseNode = Cast<URigVMCollapseNode>(GetOuter()))
+	{
+		return CollapseNode->GetNodePath(false /* recursive */);
 	}
 	return FString();
 }
@@ -92,9 +127,25 @@ URigVMNode* URigVMGraph::FindNodeByName(const FName& InNodeName) const
 
 URigVMNode* URigVMGraph::FindNode(const FString& InNodePath) const
 {
-	// for now we don't support encapsulation
-	// so we don't walk recursively
-	return FindNodeByName(FName(*InNodePath));
+	if (InNodePath.IsEmpty())
+	{
+		return nullptr;
+	}
+
+	FString Left = InNodePath, Right;
+	URigVMNode::SplitNodePathAtStart(InNodePath, Left, Right);
+
+	if (Right.IsEmpty())
+	{
+		return FindNodeByName(*Left);
+	}
+
+	if (URigVMLibraryNode* LibraryNode = Cast< URigVMLibraryNode>(FindNodeByName(*Left)))
+	{
+		return LibraryNode->GetContainedGraph()->FindNode(Right);
+	}
+
+	return nullptr;
 }
 
 URigVMPin* URigVMGraph::FindPin(const FString& InPinPath) const
diff --git a/Engine/Source/Developer/RigVMDeveloper/Private/RigVMModel/RigVMNode.cpp b/Engine/Source/Developer/RigVMDeveloper/Private/RigVMModel/RigVMNode.cpp
--- a/Engine/Source/Developer/RigVMDeveloper/Private/RigVMModel/RigVMNode.cpp
+++ b/Engine/Source/Developer/RigVMDeveloper/Private/RigVMModel/RigVMNode.cpp
@@ -30,12 +30,66 @@ FString URigVMNode::GetNodePath(bool bRecursive) const
 		FString ParentNodePath = GetGraph()->GetNodePath();
 		if (!ParentNodePath.IsEmpty())
 		{
-			return FString::Printf(TEXT("%s|%s"), *ParentNodePath, *GetName());
+			return JoinNodePath(ParentNodePath, GetName());
 		}
 	}
 	return GetName();
 }
 
+bool URigVMNode::SplitNodePathAtStart(const FString& InNodePath, FString& LeftMost, FString& Right)
+{
+	return InNodePath.Split(TEXT("|"), &LeftMost, &Right, ESearchCase::IgnoreCase, ESearchDir::FromStart);
+}
+
+bool URigVMNode::SplitNodePathAtEnd(const FString& InNodePath, FString& Left, FString& RightMost)
+{
+	return InNodePath.Split(TEXT("|"), &Left, &RightMost, ESearchCase::IgnoreCase, ESearchDir::FromEnd);
+}
+
+bool URigVMNode::SplitNodePath(const FString& InNodePath, TArray<FString>& Parts)
+{
+	int32 OriginalPartsCount = Parts.Num();
+	FString NodePathRemaining = InNodePath;
+	FString Left, Right;
+	Right = NodePathRemaining;
+
+	while (SplitNodePathAtStart(NodePathRemaining, Left, Right))
+	{
+		Parts.Add(Left);
+		Left.Empty();
+		NodePathRemaining = Right;
+	}
+
+	if (!Right.IsEmpty())
+	{
+		Parts.Add(Right);
+	}
+
+	return Parts.Num() > OriginalPartsCount;
+}
+
+FString URigVMNode::JoinNodePath(const FString& Left, const FString& Right)
+{
+	ensure(!Left.IsEmpty() && !Right.IsEmpty());
+	return Left + TEXT("|") + Right;
+}
+
+FString URigVMNode::JoinNodePath(const TArray<FString>& InParts)
+{
+	if (InParts.Num() == 0)
+	{
+		return FString();
+	}
+
+	FString Result = InParts[0];
+	for (int32 PartIndex = 1; PartIndex < InParts.Num(); PartIndex++)
+	{
+		Result += TEXT("|") + InParts[PartIndex];
+	}
+
+	return Result;
+}
+
 int32 URigVMNode::GetNodeIndex() const
 {
 	int32 Index = INDEX_NONE;
diff --git a/Engine/Source/Developer/RigVMDeveloper/Private/RigVMModel/RigVMPin.cpp b/Engine/Source/Developer/RigVMDeveloper/Private/RigVMModel/RigVMPin.cpp
--- a/Engine/Source/Developer/RigVMDeveloper/Private/RigVMModel/RigVMPin.cpp
+++ b/Engine/Source/Developer/RigVMDeveloper/Private/RigVMModel/RigVMPin.cpp
@@ -94,7 +94,7 @@ FString URigVMPin::GetPinPath(bool bUseNodePath) const
 	FString PinPath;
 
 	URigVMPin* ParentPin = GetParentPin();
-	if(ParentPin)
+	if (ParentPin)
 	{
 		PinPath = FString::Printf(TEXT("%s.%s"), *ParentPin->GetPinPath(), *GetName());
 	}
@@ -107,7 +107,10 @@ FString URigVMPin::GetPinPath(bool bUseNodePath) const
 		}
 	}
 #if UE_BUILD_DEBUG
-	CachedPinPath = PinPath;
+	if (!bUseNodePath)
+	{
+		CachedPinPath = PinPath;
+	}
 #endif
 	return PinPath;
 }
diff --git a/Engine/Source/Developer/RigVMDeveloper/Public/RigVMModel/Nodes/RigVMCollapseNode.h b/Engine/Source/Developer/RigVMDeveloper/Public/RigVMModel/Nodes/RigVMCollapseNode.h
--- a/Engine/Source/Developer/RigVMDeveloper/Public/RigVMModel/Nodes/RigVMCollapseNode.h
+++ b/Engine/Source/Developer/RigVMDeveloper/Public/RigVMModel/Nodes/RigVMCollapseNode.h
@@ -24,6 +24,8 @@ public:
 	// Library node interface
 	virtual URigVMGraph* GetContainedGraph() const override { return ContainedGraph; }
 
+	FString GetEditorSubGraphName() const;
+
 private:
 
 	UPROPERTY()
diff --git a/Engine/Source/Developer/RigVMDeveloper/Public/RigVMModel/RigVMController.h b/Engine/Source/Developer/RigVMDeveloper/Public/RigVMModel/RigVMController.h
--- a/Engine/Source/Developer/RigVMDeveloper/Public/RigVMModel/RigVMController.h
+++ b/Engine/Source/Developer/RigVMDeveloper/Public/RigVMModel/RigVMController.h
@@ -55,6 +55,20 @@ public:
 	UFUNCTION(BlueprintCallable, Category = RigVMController)
 	void SetGraph(URigVMGraph* InGraph);
 
+	// Pushes a new graph to the stack
+	// This causes a GraphChanged modified event.
+	UFUNCTION(BlueprintCallable, Category = RigVMController)
+	void PushGraph(URigVMGraph* InGraph, bool bSetupUndoRedo = true);
+
+	// Pops the last graph off the stack
+	// This causes a GraphChanged modified event.
+	UFUNCTION(BlueprintCallable, Category = RigVMController)
+	URigVMGraph* PopGraph(bool bSetupUndoRedo = true);
+	
+	// Returns the top level graph
+	UFUNCTION(BlueprintCallable, Category = RigVMController)
+	URigVMGraph* GetTopLevelGraph() const;
+
 	// The Modified event used to subscribe to changes
 	// happening within the Graph. This is broadcasted to 
 	// for any change happening - not only the changes 
@@ -283,6 +297,11 @@ public:
 	UFUNCTION(BlueprintCallable, Category = RigVMController)
 	bool RemoveNodeByName(const FName& InNodeName, bool bSetupUndoRedo = true, bool bRecursive = false);
 
+	// Renames a node in the graph
+	// This causes a NodeRenamed modified event.
+	UFUNCTION(BlueprintCallable, Category = RigVMController)
+	bool RenameNode(URigVMNode* InNode, const FName& InNewName, bool bSetupUndoRedo = true);
+
 	// Selects a single node in the graph.
 	// This causes a NodeSelected / NodeDeselected modified event.
 	UFUNCTION(BlueprintCallable, Category = RigVMController)
@@ -454,6 +473,22 @@ public:
 	UFUNCTION(BlueprintCallable, Category = RigVMController)
 	bool BreakAllLinks(const FString& InPinPath, bool bAsInput = true, bool bSetupUndoRedo = true);
 
+	// Adds an exposed pin to the graph controlled by this
+	UFUNCTION(BlueprintCallable, Category = RigVMController)
+	FName AddExposedPin(const FName& InPinName, ERigVMPinDirection InDirection, const FString& InCPPType, const FName& InCPPTypeObjectPath, const FString& InDefaultValue, bool bSetupUndoRedo = true);
+
+	// Removes an exposed pin from the graph controlled by this
+	UFUNCTION(BlueprintCallable, Category = RigVMController)
+	bool RemoveExposedPin(const FName& InPinName, bool bSetupUndoRedo = true);
+
+	// Renames an exposed pin in the graph controlled by this
+	UFUNCTION(BlueprintCallable, Category = RigVMController)
+	bool RenameExposedPin(const FName& InOldPinName, const FName& InNewPinName, bool bSetupUndoRedo = true);
+
+	// Changes the type of an exposed pin in the graph controlled by this
+	UFUNCTION(BlueprintCallable, Category = RigVMController)
+	bool ChangeExposedPinType(const FName& InPinName, const FString& InCPPType, const FName& InCPPTypeObjectPath, bool bSetupUndoRedo = true);
+
 	// Sets the execute context struct type to use
 	void SetExecuteContextStruct(UStruct* InExecuteContextStruct);
 
@@ -463,8 +498,8 @@ public:
 	// A delegate to retrieve the list of external variables
 	FRigVMController_GetExternalVariablesDelegate GetExternalVariablesDelegate;
 
-	int32 DetachLinksFromPinObjects(const TArray<URigVMLink*>* InLinks = nullptr);
-	int32 ReattachLinksToPinObjects(bool bFollowCoreRedirectors = false, const TArray<URigVMLink*>* InLinks = nullptr);
+	int32 DetachLinksFromPinObjects(const TArray<URigVMLink*>* InLinks = nullptr, bool bNotify = false);
+	int32 ReattachLinksToPinObjects(bool bFollowCoreRedirectors = false, const TArray<URigVMLink*>* InLinks = nullptr, bool bNotify = false);
 	void AddPinRedirector(bool bInput, bool bOutput, const FString& OldPinPath, const FString& NewPinPath);
 
 	// Removes nodes which went stale.
@@ -475,7 +510,7 @@ public:
 	bool ShouldRedirectPin(UScriptStruct* InOwningStruct, const FString& InOldRelativePinPath, FString& InOutNewRelativePinPath) const;
 	bool ShouldRedirectPin(const FString& InOldPinPath, FString& InOutNewPinPath) const;
 
-	void RepopulatePinsOnNode(URigVMNode* InNode, bool bFollowCoreRedirectors = true);
+	void RepopulatePinsOnNode(URigVMNode* InNode, bool bFollowCoreRedirectors = true, bool bNotify = false);
 #endif
 
 	FRigVMUnitNodeCreatedContext& GetUnitNodeCreatedContext() { return UnitNodeCreatedContext; }
@@ -498,7 +533,7 @@ private:
 	bool IsValidNodeForGraph(URigVMNode* InNode);
 	bool IsValidPinForGraph(URigVMPin* InPin);
 	bool IsValidLinkForGraph(URigVMLink* InLink);
-	void AddPinsForStruct(UStruct* InStruct, URigVMNode* InNode, URigVMPin* InParentPin, ERigVMPinDirection InPinDirection, const FString& InDefaultValue, bool bAutoExpandArrays);
+	void AddPinsForStruct(UStruct* InStruct, URigVMNode* InNode, URigVMPin* InParentPin, ERigVMPinDirection InPinDirection, const FString& InDefaultValue, bool bAutoExpandArrays, bool bNotify = false);
 	void AddPinsForArray(FArrayProperty* InArrayProperty, URigVMNode* InNode, URigVMPin* InParentPin, ERigVMPinDirection InPinDirection, const TArray<FString>& InDefaultValues, bool bAutoExpandArrays);
 	void ConfigurePinFromProperty(FProperty* InProperty, URigVMPin* InOutPin, ERigVMPinDirection InPinDirection = ERigVMPinDirection::Invalid);
 	void ConfigurePinFromPin(URigVMPin* InOutPin, URigVMPin* InPin);
@@ -508,7 +543,7 @@ private:
 	bool ResetPinDefaultValue(URigVMPin* InPin, bool bSetupUndoRedo);
 	static TArray<FString> SplitDefaultValue(const FString& InDefaultValue);
 	URigVMPin* InsertArrayPin(URigVMPin* ArrayPin, int32 InIndex, const FString& InDefaultValue, bool bSetupUndoRedo);
-	bool RemovePin(URigVMPin* InPinToRemove, bool bSetupUndoRedo);
+	bool RemovePin(URigVMPin* InPinToRemove, bool bSetupUndoRedo, bool bNotify);
 	FProperty* FindPropertyForPin(const FString& InPinPath);
 	bool BindPinToVariable(URigVMPin* InPin, const FString& InNewBoundVariablePath, bool bSetupUndoRedo);
 	bool MakeBindingsFromVariableNode(URigVMVariableNode* InNode, bool bSetupUndoRedo);
@@ -527,7 +562,7 @@ private:
 	static void ForEveryPinRecursively(URigVMNode* InNode, TFunction<void(URigVMPin*)> OnEachPinFunction);
 	URigVMCollapseNode* CollapseNodes(const TArray<URigVMNode*>& InNodes, const FString& InCollapseNodeName, bool bSetupUndoRedo);
 	TArray<URigVMNode*> ExpandLibraryNode(URigVMLibraryNode* InNode, bool bSetupUndoRedo);
-	void RefreshFunctionPins(URigVMNode* InNode);
+	void RefreshFunctionPins(URigVMNode* InNode, bool bNotify = true);
 
 	struct FPinState
 	{
@@ -584,7 +619,7 @@ private:
 	TArray<FRigVMExternalVariable> GetExternalVariables();
 
 	UPROPERTY(transient)
-	URigVMGraph* Graph;
+	TArray<URigVMGraph*> Graphs;
 
 	UPROPERTY()
 	UStruct* ExecuteContextStruct;
@@ -646,3 +681,24 @@ private:
 	friend class FRigVMParserAST;
 };
 
+class FRigVMControllerGraphGuard
+{
+public:
+
+	FRigVMControllerGraphGuard(URigVMController* InController, URigVMGraph* InGraph, bool bSetupUndoRedo = true)
+		: Controller(InController)
+		, bUndo(bSetupUndoRedo)
+	{
+		Controller->PushGraph(InGraph, bUndo);
+	}
+
+	~FRigVMControllerGraphGuard()
+	{
+		Controller->PopGraph(bUndo);
+	}
+
+private:
+
+	URigVMController* Controller;
+	bool bUndo;
+};
diff --git a/Engine/Source/Developer/RigVMDeveloper/Public/RigVMModel/RigVMControllerActions.h b/Engine/Source/Developer/RigVMDeveloper/Public/RigVMModel/RigVMControllerActions.h
--- a/Engine/Source/Developer/RigVMDeveloper/Public/RigVMModel/RigVMControllerActions.h
+++ b/Engine/Source/Developer/RigVMDeveloper/Public/RigVMModel/RigVMControllerActions.h
@@ -1198,4 +1198,155 @@ public:
 
 	UPROPERTY()
 	TArray<FString> ExpandedNodePaths;
+};
+
+/**
+ * An action renaming a node in the graph.
+ */
+USTRUCT()
+struct FRigVMRenameNodeAction : public FRigVMBaseAction
+{
+	GENERATED_BODY()
+
+public:
+
+	FRigVMRenameNodeAction() {}
+	FRigVMRenameNodeAction(const FName& InOldNodeName, const FName& InNewNodeName);
+	virtual ~FRigVMRenameNodeAction() {};
+	virtual bool Undo(URigVMController* InController) override;
+	virtual bool Redo(URigVMController* InController) override;
+
+	UPROPERTY()
+	FString OldNodeName;
+
+	UPROPERTY()
+	FString NewNodeName;
+};
+
+/**
+ * An action pushing a graph to the graph stack of the controller
+ */
+USTRUCT()
+struct FRigVMPushGraphAction : public FRigVMBaseAction
+{
+	GENERATED_BODY()
+
+public:
+
+	FRigVMPushGraphAction() {}
+	FRigVMPushGraphAction(const FString& InGraphName);
+	virtual ~FRigVMPushGraphAction() {};
+	virtual bool Undo(URigVMController* InController) override;
+	virtual bool Redo(URigVMController* InController) override;
+
+	UPROPERTY()
+	FString GraphName;
+};
+
+/**
+ * An action popping a graph from the graph stack of the controller
+ */
+USTRUCT()
+struct FRigVMPopGraphAction : public FRigVMBaseAction
+{
+	GENERATED_BODY()
+
+public:
+
+	FRigVMPopGraphAction() {}
+	FRigVMPopGraphAction(const FString& InGraphName);
+	virtual ~FRigVMPopGraphAction() {};
+	virtual bool Undo(URigVMController* InController) override;
+	virtual bool Redo(URigVMController* InController) override;
+
+	UPROPERTY()
+	FString GraphName;
+};
+
+
+/**
+ * An action exposing a pin as a parameter
+ */
+USTRUCT()
+struct FRigVMAddExposedPinAction : public FRigVMBaseAction
+{
+	GENERATED_BODY()
+
+public:
+
+	FRigVMAddExposedPinAction() {}
+	FRigVMAddExposedPinAction(URigVMPin* InPin);
+	virtual ~FRigVMAddExposedPinAction() {};
+	virtual bool Undo(URigVMController* InController) override;
+	virtual bool Redo(URigVMController* InController) override;
+
+	UPROPERTY()
+	FString PinName;
+	
+	UPROPERTY()
+	ERigVMPinDirection Direction;
+
+	UPROPERTY()
+	FString CPPType;
+
+	UPROPERTY()
+	FString CPPTypeObjectPath;
+	
+	UPROPERTY()
+	FString DefaultValue;
+};
+
+/**
+ * An action exposing a pin as a parameter
+ */
+USTRUCT()
+struct FRigVMRemoveExposedPinAction : public FRigVMBaseAction
+{
+	GENERATED_BODY()
+
+public:
+
+	FRigVMRemoveExposedPinAction() {}
+	FRigVMRemoveExposedPinAction(URigVMPin* InPin);
+	virtual ~FRigVMRemoveExposedPinAction() {};
+	virtual bool Undo(URigVMController* InController) override;
+	virtual bool Redo(URigVMController* InController) override;
+
+	UPROPERTY()
+	FString PinName;
+	
+	UPROPERTY()
+	ERigVMPinDirection Direction;
+
+	UPROPERTY()
+	FString CPPType;
+
+	UPROPERTY()
+	FString CPPTypeObjectPath;
+
+	UPROPERTY()
+	FString DefaultValue;
+};
+
+/**
+ * An action renaming an exposed in the graph.
+ */
+USTRUCT()
+struct FRigVMRenameExposedPinAction : public FRigVMBaseAction
+{
+	GENERATED_BODY()
+
+public:
+
+	FRigVMRenameExposedPinAction() {}
+	FRigVMRenameExposedPinAction(const FName& InOldPinName, const FName& InNewPinName);
+	virtual ~FRigVMRenameExposedPinAction() {};
+	virtual bool Undo(URigVMController* InController) override;
+	virtual bool Redo(URigVMController* InController) override;
+
+	UPROPERTY()
+	FString OldPinName;
+
+	UPROPERTY()
+	FString NewPinName;
 };
\ No newline at end of file
diff --git a/Engine/Source/Developer/RigVMDeveloper/Public/RigVMModel/RigVMGraph.h b/Engine/Source/Developer/RigVMDeveloper/Public/RigVMModel/RigVMGraph.h
--- a/Engine/Source/Developer/RigVMDeveloper/Public/RigVMModel/RigVMGraph.h
+++ b/Engine/Source/Developer/RigVMDeveloper/Public/RigVMModel/RigVMGraph.h
@@ -40,6 +40,14 @@ public:
 	UFUNCTION(BlueprintCallable, Category = RigVMGraph)
 	TArray<URigVMGraph*> GetContainedGraphs(bool bRecursive = false) const;
 
+	// Returns the entry node of this graph
+	UFUNCTION(BlueprintCallable, Category = RigVMGraph)
+	URigVMFunctionEntryNode* GetEntryNode() const;
+	
+	// Returns the return node of this graph
+	UFUNCTION(BlueprintCallable, Category = RigVMGraph)
+	URigVMFunctionReturnNode* GetReturnNode() const;
+
 	// Returns a list of unique Variable descriptions within this Graph.
 	// Multiple Variable Nodes can share the same description.
 	UFUNCTION(BlueprintCallable, Category = RigVMGraph)
@@ -51,8 +59,13 @@ public:
 	TArray<FRigVMGraphParameterDescription> GetParameterDescriptions() const;
 
 	// Returns the path of this graph as defined by its invoking nodes
+	UFUNCTION(BlueprintCallable, Category = RigVMGraph)
 	FString GetNodePath() const;
 
+	// Returns the name of this graph (as defined by the node path)
+	UFUNCTION(BlueprintCallable, Category = RigVMGraph)
+	FString GetGraphName() const;
+
 	// Returns a Node given its name (or nullptr).
 	UFUNCTION(BlueprintCallable, Category = RigVMGraph)
 	URigVMNode* FindNodeByName(const FName& InNodeName) const;
diff --git a/Engine/Source/Developer/RigVMDeveloper/Public/RigVMModel/RigVMNode.h b/Engine/Source/Developer/RigVMDeveloper/Public/RigVMModel/RigVMNode.h
--- a/Engine/Source/Developer/RigVMDeveloper/Public/RigVMModel/RigVMNode.h
+++ b/Engine/Source/Developer/RigVMDeveloper/Public/RigVMModel/RigVMNode.h
@@ -35,6 +35,21 @@ public:
 	UFUNCTION(BlueprintCallable, Category = RigVMNode)
 	FString GetNodePath(bool bRecursive = false) const;
 
+	// Splits a NodePath at the start, so for example "CollapseNodeA|CollapseNodeB|CollapseNodeC" becomes "CollapseNodeA" and "CollapseNodeB|CollapseNodeC"
+	static bool SplitNodePathAtStart(const FString& InNodePath, FString& LeftMost, FString& Right);
+
+	// Splits a NodePath at the end, so for example "CollapseNodeA|CollapseNodeB|CollapseNodeC" becomes "CollapseNodeA|CollapseNodeB" and "CollapseNodeC"
+	static bool SplitNodePathAtEnd(const FString& InNodePath, FString& Left, FString& RightMost);
+
+	// Splits a NodePath into all segments, so for example "Node.Color.R" becomes ["Node", "Color", "R"]
+	static bool SplitNodePath(const FString& InNodePath, TArray<FString>& Parts);
+
+	// Joins a NodePath from to segments, so for example "CollapseNodeA" and "CollapseNodeB|CollapseNodeC" becomes "CollapseNodeA|CollapseNodeB|CollapseNodeC"
+	static FString JoinNodePath(const FString& Left, const FString& Right);
+
+	// Joins a NodePath from to segments, so for example ["CollapseNodeA", "CollapseNodeB", "CollapseNodeC"] becomes "CollapseNodeA|CollapseNodeB|CollapseNodeC"
+	static FString JoinNodePath(const TArray<FString>& InParts);
+
 	// Returns the current index of the Node within the Graph.
 	UFUNCTION(BlueprintCallable, Category = RigVMNode)
 	int32 GetNodeIndex() const;
@@ -173,6 +188,10 @@ public:
 	// Returns the number of slices for a given context
 	virtual int32 GetNumSlicesForContext(const FName& InContextName, const FRigVMUserDataArray& InUserData);
 
+	// Returns the name of the node prior to the renaming
+	UFUNCTION(BlueprintCallable, Category = RigVMNode)
+	FName GetPreviousFName() const { return PreviousName; }
+
 private:
 
 	static const FString NodeColorName;
@@ -197,6 +216,9 @@ protected:
 	UPROPERTY()
 	FLinearColor NodeColor;
 
+	UPROPERTY(transient)
+	FName PreviousName;
+
 private:
 	UPROPERTY(transient)
 	int32 InstructionIndex;
diff --git a/Engine/Source/Developer/RigVMDeveloper/Public/RigVMModel/RigVMNotifications.h b/Engine/Source/Developer/RigVMDeveloper/Public/RigVMModel/RigVMNotifications.h
--- a/Engine/Source/Developer/RigVMDeveloper/Public/RigVMModel/RigVMNotifications.h
+++ b/Engine/Source/Developer/RigVMDeveloper/Public/RigVMModel/RigVMNotifications.h
@@ -24,6 +24,9 @@ enum class ERigVMGraphNotifType : uint8
 	NodePositionChanged, // A node's position has changed (Subject == URigVMNode)
 	NodeSizeChanged, // A node's size has changed (Subject == URigVMNode)
 	NodeColorChanged, // A node's color has changed (Subject == URigVMNode)
+	PinAdded, // A pin has been added to a given node (Subject == URigVMPin)
+	PinRemoved, // A pin has been removed from a given node (Subject == URigVMPin)
+	PinRenamed, // A pin has been renamed (Subject == URigVMPin)
 	PinExpansionChanged, // A pin's expansion state has changed(Subject == URigVMPin)
 	PinWatchedChanged, // A pin's watch state has changed (Subject == URigVMPin)
 	PinArraySizeChanged, // An array pin's size has changed (Subject == URigVMPin)
@@ -44,6 +47,7 @@ enum class ERigVMGraphNotifType : uint8
 	InteractionBracketClosed, // A bracket has been opened (Subject == nullptr)
 	InteractionBracketCanceled, // A bracket has been canceled (Subject == nullptr)
 	PinBoundVariableChanged, // A pin has been bound or unbound to / from a variable (Subject == URigVMPin)
+	NodeRenamed, // A node has been renamed in the graph (Subject == URigVMNode)
 	Invalid // The max for this enum (used for guarding)
 };
 
diff --git a/Engine/Source/Editor/Kismet/Private/SBlueprintPalette.cpp b/Engine/Source/Editor/Kismet/Private/SBlueprintPalette.cpp
--- a/Engine/Source/Editor/Kismet/Private/SBlueprintPalette.cpp
+++ b/Engine/Source/Editor/Kismet/Private/SBlueprintPalette.cpp
@@ -1463,6 +1463,11 @@ void SBlueprintPaletteItem::OnNameTextCommitted(const FText& NewText, ETextCommi
 					{
 						return;
 					}
+
+					if (GraphSchema->TryRenameGraph(Graph, *NewText.ToString()))
+					{
+						return;
+					}
 				}
 			}
 
diff --git a/Engine/Source/Editor/Kismet/Private/SMyBlueprint.cpp b/Engine/Source/Editor/Kismet/Private/SMyBlueprint.cpp
--- a/Engine/Source/Editor/Kismet/Private/SMyBlueprint.cpp
+++ b/Engine/Source/Editor/Kismet/Private/SMyBlueprint.cpp
@@ -2601,6 +2601,14 @@ void SMyBlueprint::OnDeleteGraph(UEdGraph* InGraph, EEdGraphSchemaAction_K2Graph
 {
 	if (InGraph && InGraph->bAllowDeletion)
 	{
+		if (const UEdGraphSchema* Schema = InGraph->GetSchema())
+		{
+			if (Schema->TryDeleteGraph(InGraph))
+			{
+				return;
+			}
+		}
+
 		const FScopedTransaction Transaction( LOCTEXT("RemoveGraph", "Remove Graph") );
 		GetBlueprintObj()->Modify();
 
diff --git a/Engine/Source/Runtime/Engine/Classes/EdGraph/EdGraphSchema.h b/Engine/Source/Runtime/Engine/Classes/EdGraph/EdGraphSchema.h
--- a/Engine/Source/Runtime/Engine/Classes/EdGraph/EdGraphSchema.h
+++ b/Engine/Source/Runtime/Engine/Classes/EdGraph/EdGraphSchema.h
@@ -934,6 +934,16 @@ class ENGINE_API UEdGraphSchema : public UObject
 	 */
 	virtual void HandleGraphBeingDeleted(UEdGraph& GraphBeingRemoved) const {}
 
+	/*
+	 * Try to delete the graph through the schema, return true if successful
+	 */
+	virtual bool TryDeleteGraph(UEdGraph* GraphToDelete) const { return false; }
+
+	/*
+	 * Try to rename a graph through the schema, return true if successful
+	 */
+	virtual bool TryRenameGraph(UEdGraph* GraphToRename, const FName& InNewName) const { return false; }
+
 	/**
 	 * Can TestNode be encapsulated into a child graph?
 	 */
