Control Rig variables aren't showing the correct values

The issue here is that we copy the source operand to target - where the source operand might actually be on a sub pin. We need to make sure to remove the registeroffset from the operand used for the debug register - since we only allow debug copies on the root of each register.

Merging

//UE4/Release-4.26/Engine/Source/Developer/RigVMDeveloper/Private/RigVMCompiler/RigVMCompiler.cpp

to //UE5/Main/Engine/Source/Developer/RigVMDeveloper/Private/RigVMCompiler/RigVMCompiler.cpp

diff --git a/Engine/Source/Developer/RigVMDeveloper/Private/RigVMCompiler/RigVMCompiler.cpp b/Engine/Source/Developer/RigVMDeveloper/Private/RigVMCompiler/RigVMCompiler.cpp
--- a/Engine/Source/Developer/RigVMDeveloper/Private/RigVMCompiler/RigVMCompiler.cpp
+++ b/Engine/Source/Developer/RigVMDeveloper/Private/RigVMCompiler/RigVMCompiler.cpp
@@ -577,7 +577,8 @@ int32 URigVMCompiler::TraverseCallExtern(const FRigVMCallExternExprAST* InExpr,
 				FString PinHash = GetPinHash(Pin, PinExpr->To<FRigVMVarExprAST>(), true);
 				if (const FRigVMOperand* Target = WorkData.PinPathToOperand->Find(PinHash))
 				{
-					WorkData.VM->GetByteCode().AddCopyOp(Source, *Target);
+					FRigVMOperand SourceRootOperand(Source.GetMemoryType(), Source.GetRegisterIndex(), INDEX_NONE);
+					WorkData.VM->GetByteCode().AddCopyOp(SourceRootOperand, *Target);
 				}
 			}
 		}
@@ -831,8 +832,8 @@ void URigVMCompiler::TraverseAssign(const FRigVMAssignExprAST* InExpr, FRigVMCom
 				FRigVMOperand WatchOperand = FindOrAddRegister(SourceExpr, WorkData, true /*debug watch*/);
 				if (WatchOperand.IsValid())
 				{
-					Source.RegisterOffset = INDEX_NONE;
-					WorkData.VM->GetByteCode().AddCopyOp(Source, WatchOperand);
+					FRigVMOperand SourceRootOperand(Source.GetMemoryType(), Source.GetRegisterIndex(), INDEX_NONE);
+					WorkData.VM->GetByteCode().AddCopyOp(SourceRootOperand, WatchOperand);
 				}
 			}
 		}
@@ -844,7 +845,7 @@ void URigVMCompiler::TraverseAssign(const FRigVMAssignExprAST* InExpr, FRigVMCom
 				FRigVMOperand WatchOperand = FindOrAddRegister(TargetExpr, WorkData, true /*debug watch*/);
 				if (WatchOperand.IsValid())
 				{
-					Target.RegisterOffset = INDEX_NONE;
+					FRigVMOperand TargetRootOperand(Target.GetMemoryType(), Target.GetRegisterIndex(), INDEX_NONE);
 					WorkData.VM->GetByteCode().AddCopyOp(Target, WatchOperand);
 				}
 			}
