Control Rig variables aren't showing the correct values

The issue here is that we copy the source operand to target - where the source operand might actually be on a sub pin. We need to make sure to remove the registeroffset from the operand used for the debug register - since we only allow debug copies on the root of each register.

Merging

//UE4/Release-4.26/Engine/Source/Developer/RigVMDeveloper/Private/RigVMCompiler/RigVMCompiler.cpp

to //UE5/Main/Engine/Source/Developer/RigVMDeveloper/Private/RigVMCompiler/RigVMCompiler.cpp

diff --git a/Engine/Source/Developer/RigVMDeveloper/Private/RigVMCompiler/RigVMCompiler.cpp b/Engine/Source/Developer/RigVMDeveloper/Private/RigVMCompiler/RigVMCompiler.cpp
--- a/Engine/Source/Developer/RigVMDeveloper/Private/RigVMCompiler/RigVMCompiler.cpp
+++ b/Engine/Source/Developer/RigVMDeveloper/Private/RigVMCompiler/RigVMCompiler.cpp
@@ -846,7 +846,7 @@ void URigVMCompiler::TraverseAssign(const FRigVMAssignExprAST* InExpr, FRigVMCom
 				if (WatchOperand.IsValid())
 				{
 					FRigVMOperand TargetRootOperand(Target.GetMemoryType(), Target.GetRegisterIndex(), INDEX_NONE);
-					WorkData.VM->GetByteCode().AddCopyOp(Target, WatchOperand);
+					WorkData.VM->GetByteCode().AddCopyOp(TargetRootOperand, WatchOperand);
 				}
 			}
 		}
