Control Rig: Added value string validation for setting pin default values + unit tests

diff --git a/Engine/Source/Developer/RigVMDeveloper/Private/RigVMCompiler/RigVMAST.cpp b/Engine/Source/Developer/RigVMDeveloper/Private/RigVMCompiler/RigVMAST.cpp
--- a/Engine/Source/Developer/RigVMDeveloper/Private/RigVMCompiler/RigVMAST.cpp
+++ b/Engine/Source/Developer/RigVMDeveloper/Private/RigVMCompiler/RigVMAST.cpp
@@ -1824,7 +1824,7 @@ bool FRigVMParserAST::FoldConstantValuesToLiterals(URigVMGraph* InGraph, URigVMC
 		URigVMPin* PinForDefaultValue = RootPin;
 		while (PinForDefaultValue != PinToCompute && SegmentNames.Num() > 0)
 		{
-			TArray<FString> SplitDefaultValues = URigVMController::SplitDefaultValue(DefaultValue);
+			TArray<FString> SplitDefaultValues = URigVMPin::SplitDefaultValue(DefaultValue);
 
 			if (PinForDefaultValue->IsArray())
 			{
diff --git a/Engine/Source/Developer/RigVMDeveloper/Private/RigVMCompiler/RigVMCompiler.cpp b/Engine/Source/Developer/RigVMDeveloper/Private/RigVMCompiler/RigVMCompiler.cpp
--- a/Engine/Source/Developer/RigVMDeveloper/Private/RigVMCompiler/RigVMCompiler.cpp
+++ b/Engine/Source/Developer/RigVMDeveloper/Private/RigVMCompiler/RigVMCompiler.cpp
@@ -1353,7 +1353,7 @@ FRigVMOperand URigVMCompiler::FindOrAddRegister(const FRigVMVarExprAST* InVarExp
 
 				int32 DesiredArraySize = VMStruct->GetArraySize(Pin->GetFName(), WorkData.RigVMUserData);
 
-				DefaultValues = URigVMController::SplitDefaultValue(Pin->GetDefaultValue(DefaultValueOverride));
+				DefaultValues = URigVMPin::SplitDefaultValue(Pin->GetDefaultValue(DefaultValueOverride));
 
 				if (DefaultValues.Num() != DesiredArraySize)
 				{
@@ -1372,7 +1372,7 @@ FRigVMOperand URigVMCompiler::FindOrAddRegister(const FRigVMVarExprAST* InVarExp
 			}
 			else
 			{
-				DefaultValues = URigVMController::SplitDefaultValue(Pin->GetDefaultValue(DefaultValueOverride));
+				DefaultValues = URigVMPin::SplitDefaultValue(Pin->GetDefaultValue(DefaultValueOverride));
 			}
 
 			while (DefaultValues.Num() < Pin->GetSubPins().Num())
diff --git a/Engine/Source/Developer/RigVMDeveloper/Private/RigVMModel/RigVMController.cpp b/Engine/Source/Developer/RigVMDeveloper/Private/RigVMModel/RigVMController.cpp
--- a/Engine/Source/Developer/RigVMDeveloper/Private/RigVMModel/RigVMController.cpp
+++ b/Engine/Source/Developer/RigVMDeveloper/Private/RigVMModel/RigVMController.cpp
@@ -3715,24 +3715,36 @@ bool URigVMController::SetPinDefaultValue(const FString& InPinPath, const FStrin
 		}
 	}
 
-	SetPinDefaultValue(Pin, InDefaultValue, bResizeArrays, bSetupUndoRedo, bMergeUndoAction);
+	if (!SetPinDefaultValue(Pin, InDefaultValue, bResizeArrays, bSetupUndoRedo, bMergeUndoAction))
+	{
+		return false;
+	}
+
 	URigVMPin* PinForLink = Pin->GetPinForLink();
 	if (PinForLink != Pin)
 	{
-		SetPinDefaultValue(PinForLink, InDefaultValue, bResizeArrays, false, bMergeUndoAction);
+		if (!SetPinDefaultValue(PinForLink, InDefaultValue, bResizeArrays, false, bMergeUndoAction))
+		{
+			return false;
+		}
 	}
 
 	return true;
 }
 
-void URigVMController::SetPinDefaultValue(URigVMPin* InPin, const FString& InDefaultValue, bool bResizeArrays, bool bSetupUndoRedo, bool bMergeUndoAction)
+bool URigVMController::SetPinDefaultValue(URigVMPin* InPin, const FString& InDefaultValue, bool bResizeArrays, bool bSetupUndoRedo, bool bMergeUndoAction)
 {
 	check(InPin);
 	ensure(!InDefaultValue.IsEmpty());
 
 	if (InPin->GetDirection() == ERigVMPinDirection::Hidden)
 	{
-		return;
+		return false;
+	}
+	
+	if (!InPin->IsValidDefaultValue(InDefaultValue))
+	{
+		return false;
 	}
 
 	URigVMGraph* Graph = GetGraph();
@@ -3751,7 +3763,7 @@ void URigVMController::SetPinDefaultValue(URigVMPin* InPin, const FString& InDef
 	{
 		if (ShouldPinBeUnfolded(InPin))
 		{
-			TArray<FString> Elements = SplitDefaultValue(InDefaultValue);
+			TArray<FString> Elements = URigVMPin::SplitDefaultValue(InDefaultValue);
 
 			if (bResizeArrays)
 			{
@@ -3783,7 +3795,7 @@ void URigVMController::SetPinDefaultValue(URigVMPin* InPin, const FString& InDef
 	}
 	else if (InPin->IsStruct())
 	{
-		TArray<FString> MemberValuePairs = SplitDefaultValue(InDefaultValue);
+		TArray<FString> MemberValuePairs = URigVMPin::SplitDefaultValue(InDefaultValue);
 
 		for (const FString& MemberValuePair : MemberValuePairs)
 		{
@@ -3821,6 +3833,8 @@ void URigVMController::SetPinDefaultValue(URigVMPin* InPin, const FString& InDef
 	{
 		ActionStack->EndAction(Action, bMergeUndoAction);
 	}
+
+	return true;
 }
 
 bool URigVMController::ResetPinDefaultValue(const FString& InPinPath, bool bSetupUndoRedo)
@@ -3919,65 +3933,7 @@ bool URigVMController::ResetPinDefaultValue(URigVMPin* InPin, bool bSetupUndoRed
 	return false;
 }
 
-TArray<FString> URigVMController::SplitDefaultValue(const FString& InDefaultValue)
-{
-	TArray<FString> Parts;
-	if (InDefaultValue.IsEmpty())
-	{
-		return Parts;
-	}
-
-	ensure(InDefaultValue[0] == TCHAR('('));
-	ensure(InDefaultValue[InDefaultValue.Len() - 1] == TCHAR(')'));
-
-	FString Content = InDefaultValue.Mid(1, InDefaultValue.Len() - 2);
-	int32 BraceCount = 0;
-	int32 QuoteCount = 0;
-
-	int32 LastPartStartIndex = 0;
-	for (int32 CharIndex = 0; CharIndex < Content.Len(); CharIndex++)
-	{
-		TCHAR Char = Content[CharIndex];
-		if (QuoteCount > 0)
-		{
-			if (Char == TCHAR('"'))
-			{
-				QuoteCount = 0;
-			}
-		}
-		else if (Char == TCHAR('"'))
-		{
-			QuoteCount = 1;
-		}
-		
-		if (Char == TCHAR('('))
-		{
-			if (QuoteCount == 0)
-			{
-				BraceCount++;
-			}
-		}
-		else if (Char == TCHAR(')'))
-		{
-			if (QuoteCount == 0)
-			{
-				BraceCount--;
-				BraceCount = FMath::Max<int32>(BraceCount, 0);
-			}
-		}
-		else if (Char == TCHAR(',') && BraceCount == 0 && QuoteCount == 0)
-		{
-			Parts.Add(Content.Mid(LastPartStartIndex, CharIndex - LastPartStartIndex));
-			LastPartStartIndex = CharIndex + 1;
-		}
-	}
 
-	if (!Content.IsEmpty())
-	{
-		Parts.Add(Content.Mid(LastPartStartIndex));
-	}
-	return Parts;
-}
 
 FString URigVMController::AddArrayPin(const FString& InArrayPinPath, const FString& InDefaultValue, bool bSetupUndoRedo)
 {
@@ -4097,7 +4053,7 @@ URigVMPin* URigVMController::InsertArrayPin(URigVMPin* ArrayPin, int32 InIndex,
 		FArrayProperty * ArrayProperty = CastField<FArrayProperty>(FindPropertyForPin(Pin->GetPinPath()));
 		if (ArrayProperty)
 		{
-			TArray<FString> ElementDefaultValues = SplitDefaultValue(InDefaultValue);
+			TArray<FString> ElementDefaultValues = URigVMPin::SplitDefaultValue(InDefaultValue);
 			AddPinsForArray(ArrayProperty, Pin->GetNode(), Pin, Pin->Direction, ElementDefaultValues, false);
 		}
 	}
@@ -5917,7 +5873,7 @@ bool URigVMController::IsValidLinkForGraph(URigVMLink* InLink)
 
 void URigVMController::AddPinsForStruct(UStruct* InStruct, URigVMNode* InNode, URigVMPin* InParentPin, ERigVMPinDirection InPinDirection, const FString& InDefaultValue, bool bAutoExpandArrays, bool bNotify)
 {
-	TArray<FString> MemberNameValuePairs = SplitDefaultValue(InDefaultValue);
+	TArray<FString> MemberNameValuePairs = URigVMPin::SplitDefaultValue(InDefaultValue);
 	TMap<FName, FString> MemberValues;
 	for (const FString& MemberNameValuePair : MemberNameValuePairs)
 	{
@@ -5975,7 +5931,7 @@ void URigVMController::AddPinsForStruct(UStruct* InStruct, URigVMNode* InNode, U
 			{
 				if (ShouldPinBeUnfolded(Pin))
 				{
-					TArray<FString> ElementDefaultValues = SplitDefaultValue(*DefaultValuePtr);
+					TArray<FString> ElementDefaultValues = URigVMPin::SplitDefaultValue(*DefaultValuePtr);
 					AddPinsForArray(ArrayProperty, InNode, Pin, Pin->Direction, ElementDefaultValues, bAutoExpandArrays);
 				}
 				else
@@ -6044,7 +6000,7 @@ void URigVMController::AddPinsForArray(FArrayProperty* InArrayProperty, URigVMNo
 		{
 			if (ShouldPinBeUnfolded(Pin))
 			{
-				TArray<FString> ElementDefaultValues = SplitDefaultValue(DefaultValue);
+				TArray<FString> ElementDefaultValues = URigVMPin::SplitDefaultValue(DefaultValue);
 				AddPinsForArray(ArrayProperty, InNode, Pin, Pin->Direction, ElementDefaultValues, bAutoExpandArrays);
 			}
 			else if (!DefaultValue.IsEmpty())
diff --git a/Engine/Source/Developer/RigVMDeveloper/Private/RigVMModel/RigVMPin.cpp b/Engine/Source/Developer/RigVMDeveloper/Private/RigVMModel/RigVMPin.cpp
--- a/Engine/Source/Developer/RigVMDeveloper/Private/RigVMModel/RigVMPin.cpp
+++ b/Engine/Source/Developer/RigVMDeveloper/Private/RigVMModel/RigVMPin.cpp
@@ -9,6 +9,28 @@
 #include "RigVMCore/RigVMExecuteContext.h"
 #include "UObject/Package.h"
 #include "Misc/PackageName.h"
+#include "Misc/OutputDevice.h"
+#include "Misc/DefaultValueHelper.h"
+#include "Logging/LogScopedVerbosityOverride.h"
+
+
+class FRigVMPinDefaultValueImportErrorContext : public FOutputDevice
+{
+public:
+
+	int32 NumErrors;
+
+	FRigVMPinDefaultValueImportErrorContext()
+		: FOutputDevice()
+		, NumErrors(0)
+	{
+	}
+
+	virtual void Serialize(const TCHAR* V, ELogVerbosity::Type Verbosity, const class FName& Category) override
+	{
+		NumErrors++;
+	}
+};
 
 URigVMGraph* URigVMInjectionInfo::GetGraph() const
 {
@@ -72,6 +94,68 @@ FString URigVMPin::JoinPinPath(const TArray<FString>& InParts)
 	return Result;
 }
 
+TArray<FString> URigVMPin::SplitDefaultValue(const FString& InDefaultValue)
+{
+	TArray<FString> Parts;
+	if (InDefaultValue.IsEmpty())
+	{
+		return Parts;
+	}
+
+	ensure(InDefaultValue[0] == TCHAR('('));
+	ensure(InDefaultValue[InDefaultValue.Len() - 1] == TCHAR(')')); 
+
+	FString Content = InDefaultValue.Mid(1, InDefaultValue.Len() - 2);
+	int32 BraceCount = 0;
+	int32 QuoteCount = 0;
+
+	int32 LastPartStartIndex = 0;
+	for (int32 CharIndex = 0; CharIndex < Content.Len(); CharIndex++)
+	{
+		TCHAR Char = Content[CharIndex];
+		if (QuoteCount > 0)
+		{
+			if (Char == TCHAR('"'))
+			{
+				QuoteCount = 0;
+			}
+		}
+		else if (Char == TCHAR('"'))
+		{
+			QuoteCount = 1;
+		}
+
+		if (Char == TCHAR('('))
+		{
+			if (QuoteCount == 0)
+			{
+				BraceCount++;
+			}
+		}
+		else if (Char == TCHAR(')'))
+		{
+			if (QuoteCount == 0)
+			{
+				BraceCount--;
+				BraceCount = FMath::Max<int32>(BraceCount, 0);
+			}
+		}
+		else if (Char == TCHAR(',') && BraceCount == 0 && QuoteCount == 0)
+		{
+			// ignore whitespaces
+			Parts.Add(Content.Mid(LastPartStartIndex, CharIndex - LastPartStartIndex).Replace(TEXT(" "), TEXT("")));
+			LastPartStartIndex = CharIndex + 1;
+		}
+	}
+
+	if (!Content.IsEmpty())
+	{
+		// ignore whitespaces
+		Parts.Add(Content.Mid(LastPartStartIndex).Replace(TEXT(" "), TEXT("")));
+	}
+	return Parts;
+}
+
 URigVMPin::URigVMPin()
 	: Direction(ERigVMPinDirection::Invalid)
 	, bIsExpanded(false)
@@ -358,6 +442,90 @@ FString URigVMPin::GetDefaultValue(const FDefaultValueOverride& InDefaultValueOv
 	return DefaultValue;
 }
 
+bool URigVMPin::IsValidDefaultValue(const FString& InDefaultValue) const
+{
+	TArray<FString> DefaultValues; 
+
+	if (IsArray())
+	{
+		DefaultValues = URigVMPin::SplitDefaultValue(InDefaultValue);
+	}
+	else
+	{
+		DefaultValues.Add(InDefaultValue);
+	}
+
+	FString BaseCPPType = GetCPPType().Replace(TEXT("TArray<"), TEXT("")).Replace(TEXT(">"), TEXT(""));
+
+	for (const FString& Value : DefaultValues)
+	{
+		// perform single value validation
+		if (UScriptStruct* ScriptStruct = Cast<UScriptStruct>(GetCPPTypeObject()))
+		{
+			FRigVMPinDefaultValueImportErrorContext ErrorPipe;
+			TArray<uint8> TempStructBuffer;
+			TempStructBuffer.AddUninitialized(ScriptStruct->GetStructureSize());
+
+			{
+				// force logging to the error pipe for error detection
+				LOG_SCOPE_VERBOSITY_OVERRIDE(LogExec, ELogVerbosity::Verbose); 
+				ScriptStruct->ImportText(*Value, TempStructBuffer.GetData(), nullptr, PPF_None, &ErrorPipe, ScriptStruct->GetName()); 
+			}
+
+			if (ErrorPipe.NumErrors > 0)
+			{
+				return false;
+			}
+		} 
+		else if (UEnum* EnumType = Cast<UEnum>(GetCPPTypeObject()))
+		{
+			FName EnumName(EnumType->GenerateFullEnumName(*Value));
+			if (!EnumType->IsValidEnumName(EnumName))
+			{
+				return false;
+			}
+			else
+			{
+				if (EnumType->HasMetaData(TEXT("Hidden"), EnumType->GetIndexByName(EnumName)))
+				{
+					return false;
+				}
+			}
+		} 
+		else if (BaseCPPType == TEXT("float"))
+		{ 
+			if (!FDefaultValueHelper::IsStringValidFloat(Value))
+			{
+				return false;
+			}
+		}
+		else if (BaseCPPType == TEXT("int32"))
+		{ 
+			if (!FDefaultValueHelper::IsStringValidInteger(Value))
+			{
+				return false;
+			}
+		}
+		else if (BaseCPPType == TEXT("bool"))
+		{
+			if (Value != TEXT("True") && Value != TEXT("False"))
+			{
+				return false;
+			}
+		}
+		else if (BaseCPPType == TEXT("FString"))
+		{ 
+			// anything is allowed
+		}
+		else if (BaseCPPType == TEXT("FName"))
+		{
+			// anything is allowed
+		}
+	}
+
+	return true;
+}
+
 FName URigVMPin::GetCustomWidgetName() const
 {
 	if (IsArrayElement())
diff --git a/Engine/Source/Developer/RigVMDeveloper/Public/RigVMModel/RigVMController.h b/Engine/Source/Developer/RigVMDeveloper/Public/RigVMModel/RigVMController.h
--- a/Engine/Source/Developer/RigVMDeveloper/Public/RigVMModel/RigVMController.h
+++ b/Engine/Source/Developer/RigVMDeveloper/Public/RigVMModel/RigVMController.h
@@ -539,9 +539,8 @@ private:
 	void ConfigurePinFromPin(URigVMPin* InOutPin, URigVMPin* InPin);
 	virtual bool ShouldStructBeUnfolded(const UStruct* InStruct);
 	virtual bool ShouldPinBeUnfolded(URigVMPin* InPin);
-	void SetPinDefaultValue(URigVMPin* InPin, const FString& InDefaultValue, bool bResizeArrays, bool bSetupUndoRedo, bool bMergeUndoAction);
+	bool SetPinDefaultValue(URigVMPin* InPin, const FString& InDefaultValue, bool bResizeArrays, bool bSetupUndoRedo, bool bMergeUndoAction);
 	bool ResetPinDefaultValue(URigVMPin* InPin, bool bSetupUndoRedo);
-	static TArray<FString> SplitDefaultValue(const FString& InDefaultValue);
 	URigVMPin* InsertArrayPin(URigVMPin* ArrayPin, int32 InIndex, const FString& InDefaultValue, bool bSetupUndoRedo);
 	bool RemovePin(URigVMPin* InPinToRemove, bool bSetupUndoRedo, bool bNotify);
 	FProperty* FindPropertyForPin(const FString& InPinPath);
diff --git a/Engine/Source/Developer/RigVMDeveloper/Public/RigVMModel/RigVMPin.h b/Engine/Source/Developer/RigVMDeveloper/Public/RigVMModel/RigVMPin.h
--- a/Engine/Source/Developer/RigVMDeveloper/Public/RigVMModel/RigVMPin.h
+++ b/Engine/Source/Developer/RigVMDeveloper/Public/RigVMModel/RigVMPin.h
@@ -95,6 +95,9 @@ public:
 	// Joins a PinPath from to segments, so for example ["Node", "Color", "R"] becomes "Node.Color.R"
 	static FString JoinPinPath(const TArray<FString>& InParts);
 
+	// Splits the default value into name-value pairs
+	static TArray<FString> SplitDefaultValue(const FString& InDefaultValue);
+
 	// Default constructor
 	URigVMPin();
 
@@ -187,6 +190,9 @@ public:
 	// Returns the default value with an additional override ma
 	FString GetDefaultValue(const FDefaultValueOverride& InDefaultValueOverride) const;
 
+	// Returns true if the default value provided is valid
+	bool IsValidDefaultValue(const FString& InDefaultValue) const;
+
 	// Returns the name of a custom widget to be used
 	// for editing the Pin.
 	UFUNCTION(BlueprintCallable, Category = RigVMPin)
