Control Rig:  now creates the correct default values for rig units with non empty constructors, for example FRigUnit_SetMultiControlRotator. Also added unit test

diff --git a/Engine/Source/Developer/RigVMDeveloper/Private/RigVMModel/RigVMController.cpp b/Engine/Source/Developer/RigVMDeveloper/Private/RigVMModel/RigVMController.cpp
--- a/Engine/Source/Developer/RigVMDeveloper/Private/RigVMModel/RigVMController.cpp
+++ b/Engine/Source/Developer/RigVMDeveloper/Private/RigVMModel/RigVMController.cpp
@@ -5986,6 +5986,13 @@ void URigVMController::AddPinsForArray(FArrayProperty* InArrayProperty, URigVMNo
 		{
 			if (ShouldPinBeUnfolded(Pin))
 			{
+				// DefaultValue before this point only contains parent struct overrides,
+				// see comments in CreateDefaultValueForStructIfRequired
+				UScriptStruct* ScriptStruct = Pin->GetScriptStruct();
+				if (ScriptStruct)
+				{
+					CreateDefaultValueForStructIfRequired(ScriptStruct, DefaultValue);
+				}
 				AddPinsForStruct(StructProperty->Struct, InNode, Pin, Pin->Direction, DefaultValue, bAutoExpandArrays);
 			}
 			else if (!DefaultValue.IsEmpty())
@@ -7022,14 +7029,40 @@ void URigVMController::ReportAndNotifyError(const FString& InMessage)
 }
 
 
-void URigVMController::CreateDefaultValueForStructIfRequired(UScriptStruct* InStruct, FString& OutDefaultValue)
+void URigVMController::CreateDefaultValueForStructIfRequired(UScriptStruct* InStruct, FString& InOutDefaultValue)
 {
-	if ((OutDefaultValue.IsEmpty() || OutDefaultValue == TEXT("()")) && InStruct != nullptr)
+	if (InStruct != nullptr)
 	{
 		TArray<uint8, TAlignedHeapAllocator<16>> TempBuffer;
 		TempBuffer.AddUninitialized(InStruct->GetStructureSize());
+
+		// call the struct constructor to initialize the struct
 		InStruct->InitializeDefaultValue(TempBuffer.GetData());
-		InStruct->ExportText(OutDefaultValue, TempBuffer.GetData(), nullptr, nullptr, PPF_None, nullptr);
+
+		// apply any higher-level value overrides
+		// for example,  
+		// struct B { int Test; B() {Test = 1;}}; ----> This is the constructor initialization, applied first in InitializeDefaultValue() above 
+		// struct A 
+		// {
+		//		Array<B> TestArray;
+		//		A() 
+		//		{
+		//			TestArray.Add(B());
+		//			TestArray[0].Test = 2;  ----> This is the overrride, applied below in ImportText()
+		//		}
+		// }
+		// See UnitTest RigVM->Graph->UnitNodeDefaultValue for more use case.
+		
+		if (!InOutDefaultValue.IsEmpty() && InOutDefaultValue != TEXT("()"))
+		{ 
+			InStruct->ImportText(*InOutDefaultValue, TempBuffer.GetData(), nullptr, PPF_None, nullptr, FString());
+		}
+
+		// in case InOutDefaultValue is not empty, it needs to be cleared
+		// before ExportText() because ExportText() appends to it.
+		InOutDefaultValue.Reset();
+
+		InStruct->ExportText(InOutDefaultValue, TempBuffer.GetData(), nullptr, nullptr, PPF_None, nullptr);
 		InStruct->DestroyStruct(TempBuffer.GetData());
 	}
 }
diff --git a/Engine/Source/Developer/RigVMDeveloper/Public/RigVMModel/RigVMController.h b/Engine/Source/Developer/RigVMDeveloper/Public/RigVMModel/RigVMController.h
--- a/Engine/Source/Developer/RigVMDeveloper/Public/RigVMModel/RigVMController.h
+++ b/Engine/Source/Developer/RigVMDeveloper/Public/RigVMModel/RigVMController.h
@@ -600,7 +600,7 @@ private:
 	}
 
 	static FLinearColor GetColorFromMetadata(const FString& InMetadata);
-	static void CreateDefaultValueForStructIfRequired(UScriptStruct* InStruct, FString& OutDefaultValue);
+	static void CreateDefaultValueForStructIfRequired(UScriptStruct* InStruct, FString& InOutDefaultValue);
 	static void PostProcessDefaultValue(URigVMPin* Pin, FString& OutDefaultValue);
 
 	void PotentiallyResolvePrototypeNode(URigVMPrototypeNode* InNode, bool bSetupUndoRedo);
