Control Rig: Skip validation for pin default value on load as a work around for the jira

diff --git a/Engine/Source/Developer/RigVMDeveloper/Private/RigVMModel/RigVMController.cpp b/Engine/Source/Developer/RigVMDeveloper/Private/RigVMModel/RigVMController.cpp
--- a/Engine/Source/Developer/RigVMDeveloper/Private/RigVMModel/RigVMController.cpp
+++ b/Engine/Source/Developer/RigVMDeveloper/Private/RigVMModel/RigVMController.cpp
@@ -3742,14 +3742,18 @@ bool URigVMController::SetPinDefaultValue(URigVMPin* InPin, const FString& InDef
 		return false;
 	}
 	
-	if (!InPin->IsValidDefaultValue(InDefaultValue))
-	{
-		return false;
-	}
-
 	URigVMGraph* Graph = GetGraph();
 	check(Graph);
 
+	// only perform validation if post load completed for now to avoid crash
+	if (!Graph->HasAnyFlags(RF_NeedPostLoad))
+	{
+		if (!InPin->IsValidDefaultValue(InDefaultValue))
+		{
+			return false;
+		}
+	}
+
 	FRigVMSetPinDefaultValueAction Action;
 	if (bSetupUndoRedo)
 	{
