Control Rig Encapsulation - Function Library Support - First pass

Added FunctionLibrary and FunctionReferenceNode classes
Implemented support and validation for those in the Controller
Implemented unit test for checking the validity of a FunctionReferenceNode.

diff --git a/Engine/Source/Developer/RigVMDeveloper/Private/RigVMModel/Nodes/RigVMCollapseNode.cpp b/Engine/Source/Developer/RigVMDeveloper/Private/RigVMModel/Nodes/RigVMCollapseNode.cpp
--- a/Engine/Source/Developer/RigVMDeveloper/Private/RigVMModel/Nodes/RigVMCollapseNode.cpp
+++ b/Engine/Source/Developer/RigVMDeveloper/Private/RigVMModel/Nodes/RigVMCollapseNode.cpp
@@ -1,6 +1,7 @@
 // Copyright Epic Games, Inc. All Rights Reserved.
 
 #include "RigVMModel/Nodes/RigVMCollapseNode.h"
+#include "RigVMModel/RigVMFunctionLibrary.h"
 
 FString URigVMCollapseNode::GetNodeTitle() const
 {
@@ -14,6 +15,11 @@ FText URigVMCollapseNode::GetToolTipText() const
 	return URigVMNode::GetToolTipText();
 }
 
+URigVMFunctionLibrary* URigVMCollapseNode::GetLibrary() const
+{
+	return Cast<URigVMFunctionLibrary>(GetOuter());
+}
+
 FString URigVMCollapseNode::GetEditorSubGraphName() const
 {
 	return FString::Printf(TEXT("%s_SubGraph"), *GetName());
diff --git a/Engine/Source/Developer/RigVMDeveloper/Private/RigVMModel/Nodes/RigVMFunctionReferenceNode.cpp b/Engine/Source/Developer/RigVMDeveloper/Private/RigVMModel/Nodes/RigVMFunctionReferenceNode.cpp
new file mode 100644
--- /dev/null
+++ b/Engine/Source/Developer/RigVMDeveloper/Private/RigVMModel/Nodes/RigVMFunctionReferenceNode.cpp
@@ -0,0 +1,63 @@
+// Copyright Epic Games, Inc. All Rights Reserved.
+
+#include "RigVMModel/Nodes/RigVMFunctionReferenceNode.h"
+#include "RigVMModel/RigVMFunctionLibrary.h"
+
+FString URigVMFunctionReferenceNode::GetNodeTitle() const
+{
+	if (URigVMLibraryNode* ReferencedNode = GetReferencedNode())
+	{
+		return ReferencedNode->GetNodeTitle();
+	}
+	return Super::GetNodeTitle();
+}
+
+FLinearColor URigVMFunctionReferenceNode::GetNodeColor() const
+{
+	if (URigVMLibraryNode* ReferencedNode = GetReferencedNode())
+	{
+		return ReferencedNode->GetNodeColor();
+	}
+	return Super::GetNodeColor();
+}
+
+FText URigVMFunctionReferenceNode::GetToolTipText() const
+{
+	if (URigVMLibraryNode* ReferencedNode = GetReferencedNode())
+	{
+		return ReferencedNode->GetToolTipText();
+	}
+	return Super::GetToolTipText();
+}
+
+URigVMFunctionLibrary* URigVMFunctionReferenceNode::GetLibrary() const
+{
+	if(URigVMLibraryNode* ReferencedNode = GetReferencedNode())
+	{
+		return ReferencedNode->GetLibrary();
+	}
+	return nullptr;
+}
+
+URigVMGraph* URigVMFunctionReferenceNode::GetContainedGraph() const
+{
+	if (URigVMLibraryNode* ReferencedNode = GetReferencedNode())
+	{
+		return ReferencedNode->GetContainedGraph();
+	}
+	return nullptr;
+}
+
+URigVMLibraryNode* URigVMFunctionReferenceNode::GetReferencedNode() const
+{
+	if (!ReferencedNodePtr.IsValid())
+	{
+		ReferencedNodePtr.LoadSynchronous();
+	}
+	return ReferencedNodePtr.Get();
+}
+
+void URigVMFunctionReferenceNode::SetReferencedNode(URigVMLibraryNode* InReferenceNode)
+{
+	ReferencedNodePtr = InReferenceNode;
+}
\ No newline at end of file
diff --git a/Engine/Source/Developer/RigVMDeveloper/Private/RigVMModel/Nodes/RigVMLibraryNode.cpp b/Engine/Source/Developer/RigVMDeveloper/Private/RigVMModel/Nodes/RigVMLibraryNode.cpp
--- a/Engine/Source/Developer/RigVMDeveloper/Private/RigVMModel/Nodes/RigVMLibraryNode.cpp
+++ b/Engine/Source/Developer/RigVMDeveloper/Private/RigVMModel/Nodes/RigVMLibraryNode.cpp
@@ -50,10 +50,18 @@ const TArray<URigVMLink*>& URigVMLibraryNode::GetContainedLinks() const
 
 URigVMFunctionEntryNode* URigVMLibraryNode::GetEntryNode() const
 {
-	return GetContainedGraph()->GetEntryNode();
+	if (URigVMGraph* Graph = GetContainedGraph())
+	{
+		return Graph->GetEntryNode();
+	}
+	return nullptr;
 }
 
 URigVMFunctionReturnNode* URigVMLibraryNode::GetReturnNode() const
 {
-	return GetContainedGraph()->GetReturnNode();
+	if (URigVMGraph* Graph = GetContainedGraph())
+	{
+		return Graph->GetReturnNode();
+	}
+	return nullptr;
 }
diff --git a/Engine/Source/Developer/RigVMDeveloper/Private/RigVMModel/RigVMController.cpp b/Engine/Source/Developer/RigVMDeveloper/Private/RigVMModel/RigVMController.cpp
--- a/Engine/Source/Developer/RigVMDeveloper/Private/RigVMModel/RigVMController.cpp
+++ b/Engine/Source/Developer/RigVMDeveloper/Private/RigVMModel/RigVMController.cpp
@@ -4,6 +4,7 @@
 #include "RigVMModel/RigVMControllerActions.h"
 #include "RigVMModel/Nodes/RigVMFunctionEntryNode.h"
 #include "RigVMModel/Nodes/RigVMFunctionReturnNode.h"
+#include "RigVMModel/Nodes/RigVMFunctionReferenceNode.h"
 #include "RigVMCore/RigVMRegistry.h"
 #include "RigVMCore/RigVMExecuteContext.h"
 #include "RigVMCompiler/RigVMCompiler.h"
@@ -231,6 +232,13 @@ URigVMUnitNode* URigVMController::AddUnitNode(UScriptStruct* InScriptStruct, con
 	{
 		return nullptr;
 	}
+
+	if (GetGraph()->IsA<URigVMFunctionLibrary>())
+	{
+		ReportError(TEXT("Cannot add unit nodes to function library graphs."));
+		return nullptr;
+	}
+
 	if (InScriptStruct == nullptr)
 	{
 		ReportError(TEXT("InScriptStruct is null."));
@@ -356,6 +364,12 @@ URigVMVariableNode* URigVMController::AddVariableNode(const FName& InVariableNam
 	URigVMGraph* Graph = GetGraph();
 	check(Graph);
 
+	if (Graph->IsA<URigVMFunctionLibrary>())
+	{
+		ReportError(TEXT("Cannot add variables nodes to function library graphs."));
+		return nullptr;
+	}
+
 	if (InCPPTypeObject == nullptr)
 	{
 		InCPPTypeObject = URigVMCompiler::GetScriptStructForCPPType(InCPPType);
@@ -567,6 +581,11 @@ void URigVMController::OnExternalVariableRemoved(const FName& InVarName, bool bS
 	URigVMGraph* Graph = GetGraph();
 	check(Graph);
 
+	if (Graph->IsA<URigVMFunctionLibrary>())
+	{
+		return;
+	}
+
 	FString VarNameStr = InVarName.ToString();
 
 	if (bSetupUndoRedo)
@@ -620,6 +639,11 @@ void URigVMController::OnExternalVariableRenamed(const FName& InOldVarName, cons
 	URigVMGraph* Graph = GetGraph();
 	check(Graph);
 
+	if (Graph->IsA<URigVMFunctionLibrary>())
+	{
+		return;
+	}
+
 	FString VarNameStr = InOldVarName.ToString();
 
 	if (bSetupUndoRedo)
@@ -675,6 +699,11 @@ void URigVMController::OnExternalVariableTypeChanged(const FName& InVarName, con
 	URigVMGraph* Graph = GetGraph();
 	check(Graph);
 
+	if (Graph->IsA<URigVMFunctionLibrary>())
+	{
+		return;
+	}
+
 	FString VarNameStr = InVarName.ToString();
 
 	if (bSetupUndoRedo)
@@ -773,6 +802,12 @@ URigVMParameterNode* URigVMController::AddParameterNode(const FName& InParameter
 	URigVMGraph* Graph = GetGraph();
 	check(Graph);
 
+	if (Graph->IsA<URigVMFunctionLibrary>())
+	{
+		ReportError(TEXT("Cannot add parameter nodes to function library graphs."));
+		return nullptr;
+	}
+
 	if (InCPPTypeObject == nullptr)
 	{
 		InCPPTypeObject = URigVMCompiler::GetScriptStructForCPPType(InCPPType);
@@ -953,6 +988,12 @@ URigVMCommentNode* URigVMController::AddCommentNode(const FString& InCommentText
 	URigVMGraph* Graph = GetGraph();
 	check(Graph);
 
+	if (Graph->IsA<URigVMFunctionLibrary>())
+	{
+		ReportError(TEXT("Cannot add comment nodes to function library graphs."));
+		return nullptr;
+	}
+
 	FString Name = GetValidNodeName(InNodeName.IsEmpty() ? FString(TEXT("CommentNode")) : InNodeName);
 	URigVMCommentNode* Node = NewObject<URigVMCommentNode>(Graph, *Name);
 	Node->Position = InPosition;
@@ -991,6 +1032,12 @@ URigVMRerouteNode* URigVMController::AddRerouteNodeOnLink(URigVMLink* InLink, bo
 		return nullptr;
 	}
 
+	if (GetGraph()->IsA<URigVMFunctionLibrary>())
+	{
+		ReportError(TEXT("Cannot add reroutes to function library graphs."));
+		return nullptr;
+	}
+
 	URigVMPin* SourcePin = InLink->GetSourcePin();
 	URigVMPin* TargetPin = InLink->GetTargetPin();
 
@@ -1048,6 +1095,12 @@ URigVMRerouteNode* URigVMController::AddRerouteNodeOnPin(const FString& InPinPat
 	URigVMGraph* Graph = GetGraph();
 	check(Graph);
 
+	if (Graph->IsA<URigVMFunctionLibrary>())
+	{
+		ReportError(TEXT("Cannot add reroutes to function library graphs."));
+		return nullptr;
+	}
+
 	URigVMPin* Pin = Graph->FindPin(InPinPath);
 	if(Pin == nullptr)
 	{
@@ -1130,6 +1183,12 @@ URigVMInjectionInfo* URigVMController::AddInjectedNode(const FString& InPinPath,
 	URigVMGraph* Graph = GetGraph();
 	check(Graph);
 
+	if (Graph->IsA<URigVMFunctionLibrary>())
+	{
+		ReportError(TEXT("Cannot add injected nodes to function library graphs."));
+		return nullptr;
+	}
+
 	URigVMPin* Pin = Graph->FindPin(InPinPath);
 	if (Pin == nullptr)
 	{
@@ -1337,6 +1396,12 @@ URigVMNode* URigVMController::EjectNodeFromPin(const FString& InPinPath, bool bS
 	URigVMGraph* Graph = GetGraph();
 	check(Graph);
 
+	if (Graph->IsA<URigVMFunctionLibrary>())
+	{
+		ReportError(TEXT("Cannot eject nodes in function library graphs."));
+		return nullptr;
+	}
+
 	URigVMPin* Pin = Graph->FindPin(InPinPath);
 	if (Pin == nullptr)
 	{
@@ -1642,6 +1707,12 @@ bool URigVMController::CanImportNodesFromText(const FString& InText)
 	{
 		return false;
 	}
+
+	if (GetGraph()->IsA<URigVMFunctionLibrary>())
+	{
+		return false;
+	}
+
 	FRigVMControllerObjectFactory Factory(nullptr);
 	return Factory.CanCreateObjectsFromText(InText);
 }
@@ -1657,6 +1728,12 @@ TArray<FName> URigVMController::ImportNodesFromText(const FString& InText, bool
 	URigVMGraph* Graph = GetGraph();
 	check(Graph);
 
+	// can't import nodes to a function library
+	if (Graph->IsA<URigVMFunctionLibrary>())
+	{
+		return NodeNames;
+	}
+
 	FRigVMControllerObjectFactory Factory(this);
 	Factory.ProcessBuffer(Graph, RF_Transactional, InText);
 
@@ -1837,6 +1914,12 @@ URigVMCollapseNode* URigVMController::CollapseNodes(const TArray<URigVMNode*>& I
 	URigVMGraph* Graph = GetGraph();
 	check(Graph);
 
+	if (Graph->IsA<URigVMFunctionLibrary>())
+	{
+		ReportError(TEXT("Cannot collapse nodes in function library graphs."));
+		return nullptr;
+	}
+
 	TArray<URigVMNode*> Nodes;
 	for (URigVMNode* Node : InNodes)
 	{
@@ -2276,6 +2359,12 @@ TArray<URigVMNode*> URigVMController::ExpandLibraryNode(URigVMLibraryNode* InNod
 	URigVMGraph* Graph = GetGraph();
 	check(Graph);
 
+	if (Graph->IsA<URigVMFunctionLibrary>())
+	{
+		ReportError(TEXT("Cannot expand nodes in function library graphs."));
+		return TArray<URigVMNode*>();
+	}
+
 	TArray<URigVMNode*> ContainedNodes = InNode->GetContainedNodes();
 	TArray<URigVMLink*> ContainedLinks = InNode->GetContainedLinks();
 	if (ContainedNodes.Num() == 0)
@@ -3638,6 +3727,12 @@ bool URigVMController::SetPinIsWatched(URigVMPin* InPin, bool bIsWatched, bool b
 	URigVMGraph* Graph = GetGraph();
 	check(Graph);
 
+	if (Graph->IsA<URigVMFunctionLibrary>())
+	{
+		ReportError(TEXT("Cannot watch pins in function library graphs."));
+		return false;
+	}
+
 	FRigVMSetPinWatchAction Action;
 	if (bSetupUndoRedo)
 	{
@@ -4303,6 +4398,12 @@ bool URigVMController::BindPinToVariable(URigVMPin* InPin, const FString& InNewB
 		return false;
 	}
 
+	if (GetGraph()->IsA<URigVMFunctionLibrary>())
+	{
+		ReportError(TEXT("Cannot bind pins to variables in function library graphs."));
+		return false;
+	}
+
 	if (InPin->GetBoundVariablePath() == InNewBoundVariablePath)
 	{
 		return false;
@@ -4499,6 +4600,12 @@ bool URigVMController::PromotePinToVariable(URigVMPin* InPin, bool bCreateVariab
 {
 	check(InPin);
 
+	if (GetGraph()->IsA<URigVMFunctionLibrary>())
+	{
+		ReportError(TEXT("Cannot promote pins to variables in function library graphs."));
+		return false;
+	}
+
 	if (InPin->GetDirection() != ERigVMPinDirection::Input)
 	{
 		return false;
@@ -4630,6 +4737,12 @@ bool URigVMController::AddLink(URigVMPin* OutputPin, URigVMPin* InputPin, bool b
 	URigVMGraph* Graph = GetGraph();
 	check(Graph);
 
+	if (Graph->IsA<URigVMFunctionLibrary>())
+	{
+		ReportError(TEXT("Cannot add links in function library graphs."));
+		return false;
+	}
+
 	{
 		// Temporarily remove the bound variable from the pin
 		// if it exists, so that link validation can work.
@@ -4749,6 +4862,12 @@ bool URigVMController::BreakLink(URigVMPin* OutputPin, URigVMPin* InputPin, bool
 	URigVMGraph* Graph = GetGraph();
 	check(Graph);
 
+	if (Graph->IsA<URigVMFunctionLibrary>())
+	{
+		ReportError(TEXT("Cannot break links in function library graphs."));
+		return false;
+	}
+
 	for (URigVMLink* Link : InputPin->Links)
 	{
 		if (Link->SourcePin == OutputPin && Link->TargetPin == InputPin)
@@ -4896,6 +5015,12 @@ FName URigVMController::AddExposedPin(const FName& InPinName, ERigVMPinDirection
 		return NAME_None;
 	}
 
+	if (Graph->IsA<URigVMFunctionLibrary>())
+	{
+		ReportError(TEXT("Cannot expose pins in function library graphs."));
+		return NAME_None;
+	}
+
 	URigVMLibraryNode* LibraryNode = Cast<URigVMLibraryNode>(Graph->GetOuter());
 	check(LibraryNode);
 
@@ -4952,7 +5077,7 @@ bool URigVMController::RemoveExposedPin(const FName& InPinName, bool bSetupUndoR
 {
 	if (!IsValidGraph())
 	{
-		return NAME_None;
+		return false;
 	}
 
 	URigVMGraph* Graph = GetGraph();
@@ -4961,7 +5086,13 @@ bool URigVMController::RemoveExposedPin(const FName& InPinName, bool bSetupUndoR
 	if (Graph->IsTopLevelGraph())
 	{
 		ReportError(TEXT("Exposed pins can only be edited on nested graphs."));
-		return NAME_None;
+		return false;
+	}
+
+	if (Graph->IsA<URigVMFunctionLibrary>())
+	{
+		ReportError(TEXT("Cannot remove exposed pins in function library graphs."));
+		return false;
 	}
 
 	URigVMLibraryNode* LibraryNode = Cast<URigVMLibraryNode>(Graph->GetOuter());
@@ -5007,7 +5138,7 @@ bool URigVMController::RenameExposedPin(const FName& InOldPinName, const FName&
 {
 	if (!IsValidGraph())
 	{
-		return NAME_None;
+		return false;
 	}
 
 	URigVMGraph* Graph = GetGraph();
@@ -5016,7 +5147,13 @@ bool URigVMController::RenameExposedPin(const FName& InOldPinName, const FName&
 	if (Graph->IsTopLevelGraph())
 	{
 		ReportError(TEXT("Exposed pins can only be edited on nested graphs."));
-		return NAME_None;
+		return false;
+	}
+
+	if (Graph->IsA<URigVMFunctionLibrary>())
+	{
+		ReportError(TEXT("Cannot rename exposed pins in function library graphs."));
+		return false;
 	}
 
 	URigVMLibraryNode* LibraryNode = Cast<URigVMLibraryNode>(Graph->GetOuter());
@@ -5117,7 +5254,7 @@ bool URigVMController::ChangeExposedPinType(const FName& InPinName, const FStrin
 {
 	if (!IsValidGraph())
 	{
-		return NAME_None;
+		return false;
 	}
 
 	URigVMGraph* Graph = GetGraph();
@@ -5126,7 +5263,13 @@ bool URigVMController::ChangeExposedPinType(const FName& InPinName, const FStrin
 	if (Graph->IsTopLevelGraph())
 	{
 		ReportError(TEXT("Exposed pins can only be edited on nested graphs."));
-		return NAME_None;
+		return false;
+	}
+
+	if (Graph->IsA<URigVMFunctionLibrary>())
+	{
+		ReportError(TEXT("Cannot change exposed pin types in function library graphs."));
+		return false;
 	}
 
 	URigVMLibraryNode* LibraryNode = Cast<URigVMLibraryNode>(Graph->GetOuter());
@@ -5181,6 +5324,120 @@ bool URigVMController::ChangeExposedPinType(const FName& InPinName, const FStrin
 	return true;
 }
 
+URigVMFunctionReferenceNode* URigVMController::AddFunctionReferenceNode(URigVMLibraryNode* InFunctionDefinition, const FVector2D& InNodePosition, const FString& InNodeName, bool bSetupUndoRedo)
+{
+	if (!IsValidGraph())
+	{
+		return nullptr;
+	}
+
+	URigVMGraph* Graph = GetGraph();
+	check(Graph);
+
+	if (Graph->IsA<URigVMFunctionLibrary>())
+	{
+		ReportError(TEXT("Cannot add function reference nodes to function library graphs."));
+		return nullptr;
+	}
+
+	if (InFunctionDefinition == nullptr)
+	{
+		ReportError(TEXT("Cannot add a function reference node without a valid function definition."));
+		return nullptr;
+	}
+
+	FString NodeName = GetValidNodeName(InNodeName.IsEmpty() ? InFunctionDefinition->GetName() : InNodeName);
+	URigVMFunctionReferenceNode* FunctionRefNode = NewObject<URigVMFunctionReferenceNode>(Graph, *NodeName);
+	FunctionRefNode->Position = InNodePosition;
+	FunctionRefNode->SetReferencedNode(InFunctionDefinition);
+	Graph->Nodes.Add(FunctionRefNode);
+
+	RepopulatePinsOnNode(FunctionRefNode, false, false);
+
+	if (bSetupUndoRedo)
+	{
+		FRigVMInverseAction InverseAction;
+		InverseAction.Title = TEXT("Add function node");
+
+		ActionStack->BeginAction(InverseAction);
+		ActionStack->AddAction(FRigVMRemoveNodeAction(FunctionRefNode, this));
+		ActionStack->EndAction(InverseAction);
+	}
+
+	return FunctionRefNode;
+}
+
+URigVMLibraryNode* URigVMController::AddFunctionToLibrary(const FName& InFunctionName, const FVector2D& InNodePosition, bool bSetupUndoRedo)
+{
+	if (!IsValidGraph())
+	{
+		return nullptr;
+	}
+
+	URigVMGraph* Graph = GetGraph();
+	check(Graph);
+
+	if (!Graph->IsA<URigVMFunctionLibrary>())
+	{
+		ReportError(TEXT("Can only add function definitions to function library graphs."));
+		return nullptr;
+	}
+
+	FString FunctionName = GetValidNodeName(InFunctionName.IsNone() ? FString(TEXT("Function")) : InFunctionName.ToString());
+	URigVMCollapseNode* CollapseNode = NewObject<URigVMCollapseNode>(Graph, *FunctionName);
+	CollapseNode->ContainedGraph = NewObject<URigVMGraph>(CollapseNode, TEXT("ContainedGraph"));
+	CollapseNode->Position = InNodePosition;
+	Graph->Nodes.Add(CollapseNode);
+
+	{
+		FRigVMControllerGraphGuard GraphGuard(this, CollapseNode->GetContainedGraph(), false);
+
+		URigVMFunctionEntryNode* EntryNode = NewObject<URigVMFunctionEntryNode>(CollapseNode->ContainedGraph, TEXT("Entry"));
+		CollapseNode->ContainedGraph->Nodes.Add(EntryNode);
+		EntryNode->Position = FVector2D(-250.f, 0.f);
+		RefreshFunctionPins(EntryNode, false);
+		Notify(ERigVMGraphNotifType::NodeAdded, EntryNode);
+
+		URigVMFunctionReturnNode* ReturnNode = NewObject<URigVMFunctionReturnNode>(CollapseNode->ContainedGraph, TEXT("Return"));
+		CollapseNode->ContainedGraph->Nodes.Add(ReturnNode);
+		ReturnNode->Position = FVector2D(250.f, 0.f);
+		RefreshFunctionPins(ReturnNode, false);
+		Notify(ERigVMGraphNotifType::NodeAdded, ReturnNode);
+	}
+
+	if (bSetupUndoRedo)
+	{
+		FRigVMInverseAction InverseAction;
+		InverseAction.Title = TEXT("Add function to library");
+
+		ActionStack->BeginAction(InverseAction);
+		ActionStack->AddAction(FRigVMRemoveNodeAction(CollapseNode, this));
+		ActionStack->EndAction(InverseAction);
+	}
+
+	return CollapseNode;
+}
+
+bool URigVMController::RemoveFunctionFromLibrary(const FName& InFunctionName, bool bSetupUndoRedo)
+{
+	if (!IsValidGraph())
+	{
+		return false;
+	}
+
+	URigVMGraph* Graph = GetGraph();
+	check(Graph);
+
+
+	if (!Graph->IsA<URigVMFunctionLibrary>())
+	{
+		ReportError(TEXT("Can only remove function definitions from function library graphs."));
+		return false;
+	}
+
+	return RemoveNodeByName(InFunctionName, bSetupUndoRedo);
+}
+
 void URigVMController::ExpandPinRecursively(URigVMPin* InPin, bool bSetupUndoRedo)
 {
 	if (InPin == nullptr)
@@ -5346,6 +5603,12 @@ URigVMRerouteNode* URigVMController::AddFreeRerouteNode(bool bShowAsFullNode, co
 	URigVMGraph* Graph = GetGraph();
 	check(Graph);
 
+	if (Graph->IsA<URigVMFunctionLibrary>())
+	{
+		ReportError(TEXT("Cannot add reroutes to function library graphs."));
+		return nullptr;
+	}
+
 	FRigVMBaseAction Action;
 	if (bSetupUndoRedo)
 	{
@@ -6600,6 +6863,7 @@ void URigVMController::RepopulatePinsOnNode(URigVMNode* InNode, bool bFollowCore
 	URigVMFunctionEntryNode* EntryNode = Cast<URigVMFunctionEntryNode>(InNode);
 	URigVMFunctionReturnNode* ReturnNode = Cast<URigVMFunctionReturnNode>(InNode);
 	URigVMCollapseNode* CollapseNode = Cast<URigVMCollapseNode>(InNode);
+	URigVMFunctionReferenceNode* FunctionRefNode = Cast<URigVMFunctionReferenceNode>(InNode);
 
 	TGuardValue<bool> EventuallySuspendNotifs(bSuspendNotifications, !bNotify);
 	FScopeLock Lock(&PinPathCoreRedirectorsLock);
@@ -6769,6 +7033,34 @@ void URigVMController::RepopulatePinsOnNode(URigVMNode* InNode, bool bFollowCore
 			RepopulatePinsOnNode(ContainedNode, bFollowCoreRedirectors);
 		}
 	}
+	else if (FunctionRefNode)
+	{
+		if (URigVMLibraryNode* ReferencedNode = FunctionRefNode->GetReferencedNode())
+		{
+			TArray<URigVMPin*> Pins = InNode->GetPins();
+			for (URigVMPin* Pin : Pins)
+			{
+				RemovePin(Pin, false, bNotify);
+			}
+			InNode->Pins.Reset();
+
+			for (URigVMPin* ReferencedPin : ReferencedNode->Pins)
+			{
+				URigVMPin* NewPin = NewObject<URigVMPin>(InNode, ReferencedPin->GetFName());
+				ConfigurePinFromPin(NewPin, ReferencedPin);
+
+				InNode->Pins.Add(NewPin);
+
+				if (NewPin->IsStruct())
+				{
+					AddPinsForStruct(NewPin->GetScriptStruct(), InNode, NewPin, NewPin->GetDirection(), FString(), false);
+				}
+
+				Notify(ERigVMGraphNotifType::PinAdded, NewPin);
+			}
+		}
+	}
+
 	else
 	{
 		return;
diff --git a/Engine/Source/Developer/RigVMDeveloper/Private/RigVMModel/RigVMFunctionLibrary.cpp b/Engine/Source/Developer/RigVMDeveloper/Private/RigVMModel/RigVMFunctionLibrary.cpp
new file mode 100644
--- /dev/null
+++ b/Engine/Source/Developer/RigVMDeveloper/Private/RigVMModel/RigVMFunctionLibrary.cpp
@@ -0,0 +1,28 @@
+// Copyright Epic Games, Inc. All Rights Reserved.
+
+#include "RigVMModel/RigVMFunctionLibrary.h"
+
+URigVMFunctionLibrary::URigVMFunctionLibrary()
+: URigVMGraph()
+{
+}
+
+TArray<URigVMLibraryNode*> URigVMFunctionLibrary::GetFunctions() const
+{
+	TArray<URigVMLibraryNode*> Functions;
+
+	for (URigVMNode* Node : GetNodes())
+	{
+		// we only allow library nodes under a function library graph
+		URigVMLibraryNode* LibraryNode = CastChecked<URigVMLibraryNode>(Node);
+		Functions.Add(LibraryNode);
+	}
+
+	return Functions;
+}
+
+URigVMLibraryNode* URigVMFunctionLibrary::FindFunction(const FName& InFunctionName) const
+{
+	return Cast<URigVMLibraryNode>(FindNodeByName(InFunctionName));
+}
+
diff --git a/Engine/Source/Developer/RigVMDeveloper/Private/RigVMModel/RigVMGraph.cpp b/Engine/Source/Developer/RigVMDeveloper/Private/RigVMModel/RigVMGraph.cpp
--- a/Engine/Source/Developer/RigVMDeveloper/Private/RigVMModel/RigVMGraph.cpp
+++ b/Engine/Source/Developer/RigVMDeveloper/Private/RigVMModel/RigVMGraph.cpp
@@ -1,6 +1,7 @@
 // Copyright Epic Games, Inc. All Rights Reserved.
 
 #include "RigVMModel/RigVMGraph.h"
+#include "RigVMModel/RigVMFunctionLibrary.h"
 #include "RigVMModel/RigVMController.h"
 #include "RigVMModel/Nodes/RigVMLibraryNode.h"
 #include "RigVMModel/Nodes/RigVMFunctionEntryNode.h"
@@ -193,10 +194,31 @@ bool URigVMGraph::IsTopLevelGraph() const
 	{
 		return false;
 	}
-	// todo: implement this for function libraries
 	return true;
 }
 
+URigVMFunctionLibrary* URigVMGraph::GetDefaultFunctionLibrary() const
+{
+	if (DefaultFunctionLibraryPtr.IsValid())
+	{
+		return CastChecked<URigVMFunctionLibrary>(DefaultFunctionLibraryPtr.Get());
+	}
+
+	if (URigVMLibraryNode* OuterLibraryNode = Cast<URigVMLibraryNode>(GetOuter()))
+	{
+		if (URigVMGraph* OuterGraph = OuterLibraryNode->GetGraph())
+		{
+			return OuterGraph->GetDefaultFunctionLibrary();
+		}
+	}
+	return nullptr;
+}
+
+void URigVMGraph::SetDefaultFunctionLibrary(URigVMFunctionLibrary* InFunctionLibrary)
+{
+	DefaultFunctionLibraryPtr = InFunctionLibrary;
+}
+
 FRigVMGraphModifiedEvent& URigVMGraph::OnModified()
 {
 	return ModifiedEvent;
diff --git a/Engine/Source/Developer/RigVMDeveloper/Public/RigVMModel/Nodes/RigVMCollapseNode.h b/Engine/Source/Developer/RigVMDeveloper/Public/RigVMModel/Nodes/RigVMCollapseNode.h
--- a/Engine/Source/Developer/RigVMDeveloper/Public/RigVMModel/Nodes/RigVMCollapseNode.h
+++ b/Engine/Source/Developer/RigVMDeveloper/Public/RigVMModel/Nodes/RigVMCollapseNode.h
@@ -22,6 +22,7 @@ public:
 	virtual FText GetToolTipText() const override;
 
 	// Library node interface
+	virtual URigVMFunctionLibrary* GetLibrary() const override;
 	virtual URigVMGraph* GetContainedGraph() const override { return ContainedGraph; }
 
 	FString GetEditorSubGraphName() const;
diff --git a/Engine/Source/Developer/RigVMDeveloper/Public/RigVMModel/Nodes/RigVMFunctionReferenceNode.h b/Engine/Source/Developer/RigVMDeveloper/Public/RigVMModel/Nodes/RigVMFunctionReferenceNode.h
new file mode 100644
--- /dev/null
+++ b/Engine/Source/Developer/RigVMDeveloper/Public/RigVMModel/Nodes/RigVMFunctionReferenceNode.h
@@ -0,0 +1,43 @@
+// Copyright Epic Games, Inc. All Rights Reserved.
+
+#pragma once
+
+#include "RigVMModel/RigVMGraph.h"
+#include "RigVMModel/RigVMFunctionLibrary.h"
+#include "RigVMModel/Nodes/RigVMLibraryNode.h"
+#include "RigVMFunctionReferenceNode.generated.h"
+
+/**
+ * The Function Reference Node is a library node which references
+ * a library node from a separate function library graph.
+ */
+UCLASS(BlueprintType)
+class RIGVMDEVELOPER_API URigVMFunctionReferenceNode : public URigVMLibraryNode
+{
+	GENERATED_BODY()
+
+public:
+
+	// URigVMNode interface
+	virtual FString GetNodeTitle() const override;
+	virtual FLinearColor GetNodeColor() const override;
+	virtual FText GetToolTipText() const override;
+	// end URigVMNode interface
+
+	// URigVMLibraryNode interface
+	virtual URigVMFunctionLibrary* GetLibrary() const override;
+	virtual URigVMGraph* GetContainedGraph() const override;
+	// end URigVMLibraryNode interface
+
+	URigVMLibraryNode* GetReferencedNode() const;
+
+private:
+
+	void SetReferencedNode(URigVMLibraryNode* InReferenceNode);
+
+	UPROPERTY(AssetRegistrySearchable)
+	mutable TSoftObjectPtr<URigVMLibraryNode> ReferencedNodePtr;
+
+	friend class URigVMController;
+};
+
diff --git a/Engine/Source/Developer/RigVMDeveloper/Public/RigVMModel/Nodes/RigVMLibraryNode.h b/Engine/Source/Developer/RigVMDeveloper/Public/RigVMModel/Nodes/RigVMLibraryNode.h
--- a/Engine/Source/Developer/RigVMDeveloper/Public/RigVMModel/Nodes/RigVMLibraryNode.h
+++ b/Engine/Source/Developer/RigVMDeveloper/Public/RigVMModel/Nodes/RigVMLibraryNode.h
@@ -8,6 +8,7 @@
 class URigVMGraph;
 class URigVMFunctionEntryNode;
 class URigVMFunctionReturnNode;
+class URigVMFunctionLibrary;
 
 /**
  * The Library Node represents a function invocation of a
@@ -27,6 +28,7 @@ public:
 	virtual bool IsDefinedAsVarying() const override;
 
 	// Library node interface
+	virtual URigVMFunctionLibrary* GetLibrary() const { return nullptr; }
 	virtual URigVMGraph* GetContainedGraph() const { return nullptr; }
 	virtual const TArray<URigVMNode*>& GetContainedNodes() const;
 	virtual const TArray<URigVMLink*>& GetContainedLinks() const;
diff --git a/Engine/Source/Developer/RigVMDeveloper/Public/RigVMModel/RigVMController.h b/Engine/Source/Developer/RigVMDeveloper/Public/RigVMModel/RigVMController.h
--- a/Engine/Source/Developer/RigVMDeveloper/Public/RigVMModel/RigVMController.h
+++ b/Engine/Source/Developer/RigVMDeveloper/Public/RigVMModel/RigVMController.h
@@ -3,6 +3,7 @@
 #pragma once
 
 #include "RigVMGraph.h"
+#include "RigVMFunctionLibrary.h"
 #include "RigVMModel/Nodes/RigVMUnitNode.h"
 #include "RigVMModel/Nodes/RigVMVariableNode.h"
 #include "RigVMModel/Nodes/RigVMParameterNode.h"
@@ -14,6 +15,7 @@
 #include "RigVMModel/Nodes/RigVMPrototypeNode.h"
 #include "RigVMModel/Nodes/RigVMEnumNode.h"
 #include "RigVMModel/Nodes/RigVMCollapseNode.h"
+#include "RigVMModel/Nodes/RigVMFunctionReferenceNode.h"
 #include "RigVMController.generated.h"
 
 class URigVMActionStack;
@@ -489,6 +491,18 @@ public:
 	UFUNCTION(BlueprintCallable, Category = RigVMController)
 	bool ChangeExposedPinType(const FName& InPinName, const FString& InCPPType, const FName& InCPPTypeObjectPath, bool bSetupUndoRedo = true);
 
+	// Adds a function reference / invocation to the graph
+	UFUNCTION(BlueprintCallable, Category = RigVMController)
+	URigVMFunctionReferenceNode* AddFunctionReferenceNode(URigVMLibraryNode* InFunctionDefinition, const FVector2D& InNodePosition = FVector2D::ZeroVector, const FString& InNodeName = TEXT(""), bool bSetupUndoRedo = true);
+
+	// Adds a function definition to a function library graph
+	UFUNCTION(BlueprintCallable, Category = RigVMController)
+	URigVMLibraryNode* AddFunctionToLibrary(const FName& InFunctionName, const FVector2D& InNodePosition = FVector2D::ZeroVector, bool bSetupUndoRedo = true);
+
+	// Removes a function from a function library graph
+	UFUNCTION(BlueprintCallable, Category = RigVMController)
+	bool RemoveFunctionFromLibrary(const FName& InFunctionName, bool bSetupUndoRedo = true);
+
 	// Sets the execute context struct type to use
 	void SetExecuteContextStruct(UStruct* InExecuteContextStruct);
 
diff --git a/Engine/Source/Developer/RigVMDeveloper/Public/RigVMModel/RigVMFunctionLibrary.h b/Engine/Source/Developer/RigVMDeveloper/Public/RigVMModel/RigVMFunctionLibrary.h
new file mode 100644
--- /dev/null
+++ b/Engine/Source/Developer/RigVMDeveloper/Public/RigVMModel/RigVMFunctionLibrary.h
@@ -0,0 +1,36 @@
+// Copyright Epic Games, Inc. All Rights Reserved.
+
+#pragma once
+
+#include "RigVMGraph.h"
+#include "RigVMModel/Nodes/RigVMLibraryNode.h"
+#include "RigVMFunctionLibrary.generated.h"
+
+/**
+ * The Function Library is a graph used only to store
+ * the sub graphs used for functions.
+ */
+UCLASS(BlueprintType)
+class RIGVMDEVELOPER_API URigVMFunctionLibrary : public URigVMGraph
+{
+	GENERATED_BODY()
+
+public:
+
+	// Default constructor
+	URigVMFunctionLibrary();
+
+	// Returns all of the stored functions
+	UFUNCTION(BlueprintCallable, Category = RigVMGraph)
+	TArray<URigVMLibraryNode*> GetFunctions() const;
+
+	// Finds a function by name
+	UFUNCTION(BlueprintCallable, Category = RigVMGraph)
+	URigVMLibraryNode* FindFunction(const FName& InFunctionName) const;
+
+private:
+
+	friend class URigVMController;
+	friend class URigVMCompiler;
+};
+
diff --git a/Engine/Source/Developer/RigVMDeveloper/Public/RigVMModel/RigVMGraph.h b/Engine/Source/Developer/RigVMDeveloper/Public/RigVMModel/RigVMGraph.h
--- a/Engine/Source/Developer/RigVMDeveloper/Public/RigVMModel/RigVMGraph.h
+++ b/Engine/Source/Developer/RigVMDeveloper/Public/RigVMModel/RigVMGraph.h
@@ -10,6 +10,8 @@
 #include "RigVMCompiler/RigVMAST.h"
 #include "RigVMGraph.generated.h"
 
+class URigVMFunctionLibrary;
+
 /**
  * The Graph represents a Function definition
  * using Nodes as statements.
@@ -96,6 +98,13 @@ public:
 	UFUNCTION(BlueprintCallable, Category = RigVMGraph)
 	bool IsTopLevelGraph() const;
 
+	// Returns the locally available function library
+	UFUNCTION(BlueprintCallable, Category = RigVMGraph)
+	URigVMFunctionLibrary* GetDefaultFunctionLibrary() const;
+
+	UFUNCTION(BlueprintCallable, Category = RigVMGraph)
+	void SetDefaultFunctionLibrary(URigVMFunctionLibrary* InFunctionLibrary);
+
 	// Returns the modified event, which can be used to 
 	// subscribe to changes happening within the Graph.
 	FRigVMGraphModifiedEvent& OnModified();
@@ -123,6 +132,9 @@ private:
 	UPROPERTY()
 	TArray<FName> SelectedNodes;
 
+	UPROPERTY()
+	TWeakObjectPtr<URigVMGraph> DefaultFunctionLibraryPtr;
+
 	TSharedPtr<FRigVMParserAST> DiagnosticsAST;
 	TSharedPtr<FRigVMParserAST> RuntimeAST;
 
