Control Rig: Added a bValidatePinDefaults flag to the controller such that validation can be turned off in some code paths that is causing crashes

diff --git a/Engine/Plugins/Experimental/ControlRig/Source/ControlRigDeveloper/Private/ControlRigBlueprint.cpp b/Engine/Plugins/Experimental/ControlRig/Source/ControlRigDeveloper/Private/ControlRigBlueprint.cpp
--- a/Engine/Plugins/Experimental/ControlRig/Source/ControlRigDeveloper/Private/ControlRigBlueprint.cpp
+++ b/Engine/Plugins/Experimental/ControlRig/Source/ControlRigDeveloper/Private/ControlRigBlueprint.cpp
@@ -164,6 +164,9 @@ void UControlRigBlueprint::PostLoad()
 {
 	Super::PostLoad();
 
+	// temporarily disable default value validation during load time, serialized values should always be accepted
+	TGuardValue<bool> DisablePinDefaultValueValidation(GetOrCreateController()->bValidatePinDefaults, false);
+
 	HierarchyContainer.ControlHierarchy.PostLoad();
 
 	// correct the offset transforms
@@ -1015,6 +1018,9 @@ void UControlRigBlueprint::PopulateModelFromGraphForBackwardsCompatibility(UCont
 {
 	DECLARE_SCOPE_HIERARCHICAL_COUNTER_FUNC()
 
+	// temporarily disable default value validation during load time, serialized values should always be accepted
+	TGuardValue<bool> DisablePinDefaultValueValidation(GetOrCreateController()->bValidatePinDefaults, false);
+
 	int32 LinkerVersion = GetLinkerCustomVersion(FControlRigObjectVersion::GUID);
 	if (LinkerVersion >= FControlRigObjectVersion::SwitchedToRigVM)
 	{
diff --git a/Engine/Source/Developer/RigVMDeveloper/Private/RigVMModel/RigVMController.cpp b/Engine/Source/Developer/RigVMDeveloper/Private/RigVMModel/RigVMController.cpp
--- a/Engine/Source/Developer/RigVMDeveloper/Private/RigVMModel/RigVMController.cpp
+++ b/Engine/Source/Developer/RigVMDeveloper/Private/RigVMModel/RigVMController.cpp
@@ -26,7 +26,8 @@
 TMap<URigVMController::FControlRigStructPinRedirectorKey, FString> URigVMController::PinPathCoreRedirectors;
 
 URigVMController::URigVMController()
-	: bSuspendNotifications(false)
+	: bValidatePinDefaults(true)
+	, bSuspendNotifications(false)
 	, bReportWarningsAndErrors(true)
 	, bIgnoreRerouteCompactnessChanges(false)
 {
@@ -35,6 +36,7 @@ URigVMController::URigVMController()
 
 URigVMController::URigVMController(const FObjectInitializer& ObjectInitializer)
 	: Super(ObjectInitializer)
+	, bValidatePinDefaults(true)
 	, bReportWarningsAndErrors(true)
 {
 	ActionStack = CreateDefaultSubobject<URigVMActionStack>(TEXT("ActionStack"));
@@ -3839,9 +3841,8 @@ bool URigVMController::SetPinDefaultValue(URigVMPin* InPin, const FString& InDef
 	
 	URigVMGraph* Graph = GetGraph();
 	check(Graph);
-
-	// only perform validation if post load completed for now to avoid crash
-	if (!Graph->HasAnyFlags(RF_NeedPostLoad))
+ 
+	if (bValidatePinDefaults)
 	{
 		if (!InPin->IsValidDefaultValue(InDefaultValue))
 		{
diff --git a/Engine/Source/Developer/RigVMDeveloper/Public/RigVMModel/RigVMController.h b/Engine/Source/Developer/RigVMDeveloper/Public/RigVMModel/RigVMController.h
--- a/Engine/Source/Developer/RigVMDeveloper/Public/RigVMModel/RigVMController.h
+++ b/Engine/Source/Developer/RigVMDeveloper/Public/RigVMModel/RigVMController.h
@@ -534,6 +534,9 @@ public:
 	void SetupDefaultUnitNodeDelegates(TDelegate<FName(FRigVMExternalVariable)> InCreateExternalVariableDelegate);
 	void ResetUnitNodeDelegates();
 
+	// A flag that can be used to turn off pin default value validation if necessary
+	bool bValidatePinDefaults;
+
 private:
 
 	UPROPERTY(BlueprintReadOnly, Category = RigVMController, meta = (ScriptName = "ModifiedEvent", AllowPrivateAccess = "true"))
