Control Rig: Encapsulation support for the rigvm compiler

diff --git a/Engine/Plugins/Experimental/ControlRig/Source/ControlRigDeveloper/Private/ControlRigBlueprint.cpp b/Engine/Plugins/Experimental/ControlRig/Source/ControlRigDeveloper/Private/ControlRigBlueprint.cpp
--- a/Engine/Plugins/Experimental/ControlRig/Source/ControlRigDeveloper/Private/ControlRigBlueprint.cpp
+++ b/Engine/Plugins/Experimental/ControlRig/Source/ControlRigDeveloper/Private/ControlRigBlueprint.cpp
@@ -53,6 +53,7 @@ UControlRigBlueprint::UControlRigBlueprint(const FObjectInitializer& ObjectIniti
 	VMRecompilationBracket = 0;
 
 	Model = ObjectInitializer.CreateDefaultSubobject<URigVMGraph>(this, TEXT("RigVMModel"));
+	FunctionLibrary = ObjectInitializer.CreateDefaultSubobject<URigVMFunctionLibrary>(this, TEXT("RigVMFunctionLibrary"));
 
 	Validator = ObjectInitializer.CreateDefaultSubobject<UControlRigValidator>(this, TEXT("ControlRigValidator"));
 
@@ -431,6 +432,11 @@ TArray<URigVMGraph*> UControlRigBlueprint::GetAllModels() const
 	return Models;
 }
 
+URigVMFunctionLibrary* UControlRigBlueprint::GetLocalFunctionLibrary() const
+{
+	return FunctionLibrary;
+}
+
 URigVMController* UControlRigBlueprint::GetController(URigVMGraph* InGraph) const
 {
 	if (InGraph == nullptr)
@@ -1363,7 +1369,8 @@ void UControlRigBlueprint::HandleModifiedEvent(ERigVMGraphNotifType InNotifType,
 					URigVMPin* RootPin = Pin->GetRootPin();
 					if (const FRigVMOperand* Operand = PinToOperandMap.Find(RootPin->GetPinPath()))
 					{
-						if(const FRigVMExprAST* Expression = InGraph->GetRuntimeAST()->GetExprForSubject(RootPin))
+						FRigVMASTProxy RootPinProxy = FRigVMASTProxy::MakeFromUObject(RootPin);
+						if(const FRigVMExprAST* Expression = InGraph->GetRuntimeAST()->GetExprForSubject(RootPinProxy))
 						{
 							bRequiresRecompile = Expression->NumParents() > 1;
 						}
diff --git a/Engine/Plugins/Experimental/ControlRig/Source/ControlRigDeveloper/Public/ControlRigBlueprint.h b/Engine/Plugins/Experimental/ControlRig/Source/ControlRigDeveloper/Public/ControlRigBlueprint.h
--- a/Engine/Plugins/Experimental/ControlRig/Source/ControlRigDeveloper/Public/ControlRigBlueprint.h
+++ b/Engine/Plugins/Experimental/ControlRig/Source/ControlRigDeveloper/Public/ControlRigBlueprint.h
@@ -109,6 +109,9 @@ public:
 	UFUNCTION(BlueprintCallable, Category = "Control Rig Blueprint")
 	TArray<URigVMGraph*> GetAllModels() const;
 
+	UFUNCTION(BlueprintCallable, Category = "Control Rig Blueprint")
+	URigVMFunctionLibrary* GetLocalFunctionLibrary() const;
+
 	UFUNCTION(BlueprintCallable, Category = "Control Rig Blueprint")
 	URigVMController* GetController(URigVMGraph* InGraph = nullptr) const;
 
@@ -134,6 +137,9 @@ protected:
 	UPROPERTY(BlueprintReadOnly, Category = "VM")
 	URigVMGraph* Model;
 
+	UPROPERTY(BlueprintReadOnly, Category = "VM")
+	URigVMFunctionLibrary* FunctionLibrary;
+
 	UPROPERTY(BlueprintReadOnly, transient, Category = "VM")
 	TMap<URigVMGraph*, URigVMController*> Controllers;
 
diff --git a/Engine/Source/Developer/RigVMDeveloper/Private/RigVMCompiler/RigVMAST.cpp b/Engine/Source/Developer/RigVMDeveloper/Private/RigVMCompiler/RigVMAST.cpp
--- a/Engine/Source/Developer/RigVMDeveloper/Private/RigVMCompiler/RigVMAST.cpp
+++ b/Engine/Source/Developer/RigVMDeveloper/Private/RigVMCompiler/RigVMAST.cpp
@@ -19,7 +19,7 @@
 #include "Stats/StatsHierarchical.h"
 #include "RigVMDeveloperModule.h"
 
-FRigVMExprAST::FRigVMExprAST(EType InType, UObject* InSubject)
+FRigVMExprAST::FRigVMExprAST(EType InType, const FRigVMASTProxy& InProxy)
 	: Name(NAME_None)
 	, Type(InType)
 	, Index(INDEX_NONE)
@@ -568,11 +568,10 @@ bool FRigVMNodeExprAST::IsConstant() const
 	return FRigVMExprAST::IsConstant();
 }
 
-FRigVMNodeExprAST::FRigVMNodeExprAST(EType InType, UObject* InSubject /*= nullptr*/) 
+FRigVMNodeExprAST::FRigVMNodeExprAST(EType InType, const FRigVMASTProxy& InNodeProxy)
 	: FRigVMBlockExprAST(InType)
-	, Node(Cast<URigVMNode>(InSubject))
+	, Proxy(InNodeProxy)
 {
-	check(Node);
 }
 
 FName FRigVMEntryExprAST::GetEventName() const
@@ -628,64 +627,64 @@ bool FRigVMVarExprAST::IsConstant() const
 
 FString FRigVMVarExprAST::GetCPPType() const
 {
-	return Pin->GetCPPType();
+	return GetPin()->GetCPPType();
 }
 
 UObject* FRigVMVarExprAST::GetCPPTypeObject() const
 {
-	return Pin->GetCPPTypeObject();
+	return GetPin()->GetCPPTypeObject();
 }
 
 ERigVMPinDirection FRigVMVarExprAST::GetPinDirection() const
 {
-	return Pin->GetDirection();
+	return GetPin()->GetDirection();
 }
 
 FString FRigVMVarExprAST::GetDefaultValue() const
 {
-	return Pin->GetDefaultValue(GetParser()->GetPinDefaultOverrides());
+	return GetPin()->GetDefaultValue(GetParser()->GetPinDefaultOverrides(), GetProxy());
 }
 
 bool FRigVMVarExprAST::IsExecuteContext() const
 {
-	return Pin->IsExecuteContext();
+	return GetPin()->IsExecuteContext();
 }
 
 bool FRigVMVarExprAST::IsGraphParameter() const
 {
-	if (Cast<URigVMParameterNode>(Pin->GetNode()))
+	if (Cast<URigVMParameterNode>(GetPin()->GetNode()))
 	{
-		return Pin->GetName() == TEXT("Value");
+		return GetPin()->GetName() == TEXT("Value");
 	}
 	return false;
 }
 
 bool FRigVMVarExprAST::IsGraphVariable() const
 {
-	if (Cast<URigVMVariableNode>(Pin->GetNode()))
+	if (Cast<URigVMVariableNode>(GetPin()->GetNode()))
 	{
-		return Pin->GetName() == URigVMVariableNode::ValueName;
+		return GetPin()->GetName() == URigVMVariableNode::ValueName;
 	}
 	return false;
 }
 
 bool FRigVMVarExprAST::IsEnumValue() const
 {
-	if (Cast<URigVMEnumNode>(Pin->GetNode()))
+	if (Cast<URigVMEnumNode>(GetPin()->GetNode()))
 	{
-		return Pin->GetName() == TEXT("EnumIndex");
+		return GetPin()->GetName() == TEXT("EnumIndex");
 	}
 	return false;
 }
 
 bool FRigVMVarExprAST::SupportsSoftLinks() const
 {
-	if (URigVMUnitNode* UnitNode = Cast<URigVMUnitNode>(Pin->GetNode()))
+	if (URigVMUnitNode* UnitNode = Cast<URigVMUnitNode>(GetPin()->GetNode()))
 	{
 		if (UnitNode->IsLoopNode())
 		{
-			if (Pin->GetFName() != FRigVMStruct::ExecuteContextName &&
-				Pin->GetFName() != FRigVMStruct::ForLoopCompletedPinName)
+			if (GetPin()->GetFName() != FRigVMStruct::ExecuteContextName &&
+				GetPin()->GetFName() != FRigVMStruct::ForLoopCompletedPinName)
 			{
 				return true;
 			}
@@ -820,7 +819,7 @@ const FRigVMVarExprAST* FRigVMCallExternExprAST::FindVarWithPinName(const FName&
 	return nullptr;
 }
 
-const TArray<URigVMPin*> FRigVMParserAST::EmptyLinks;
+const TArray<FRigVMASTProxy> FRigVMParserAST::EmptyProxyArray;
 
 FRigVMParserAST::FRigVMParserAST(URigVMGraph* InGraph, URigVMController* InController, const FRigVMParserASTSettings& InSettings, const TArray<FRigVMExternalVariable>& InExternalVariables, const TArray<FRigVMUserDataArray>& InRigVMUserData)
 {
@@ -833,11 +832,12 @@ FRigVMParserAST::FRigVMParserAST(URigVMGraph* InGraph, URigVMController* InContr
 	// construct the inlined nodes and links information
 	Inline(InGraph);
 
-	for (URigVMNode* Node : Nodes)
+	for (const FRigVMASTProxy& NodeProxy : NodeProxies)
 	{
+		URigVMNode* Node = NodeProxy.GetSubjectChecked<URigVMNode>();
 		if(Node->IsEvent())
 		{
-			TraverseMutableNode(Node, nullptr);
+			TraverseMutableNode(NodeProxy, nullptr);
 		}
 	}
 
@@ -846,9 +846,9 @@ FRigVMParserAST::FRigVMParserAST(URigVMGraph* InGraph, URigVMController* InContr
 	for (int32 PassIndex = 0; PassIndex < 2; PassIndex++)
 	{
 		const bool bTraverseMutable = PassIndex == 0;
-		for (int32 NodeIndex = 0; NodeIndex < Nodes.Num(); NodeIndex++)
+		for (int32 NodeIndex = 0; NodeIndex < NodeProxies.Num(); NodeIndex++)
 		{
-			if (const int32* ExprIndex = NodeExpressionIndex.Find(Nodes[NodeIndex]))
+			if (const int32* ExprIndex = NodeExpressionIndex.Find(NodeProxies[NodeIndex]))
 			{
 				if (*ExprIndex != INDEX_NONE)
 				{
@@ -856,15 +856,16 @@ FRigVMParserAST::FRigVMParserAST(URigVMGraph* InGraph, URigVMController* InContr
 				}
 			}
 
-			if (Nodes[NodeIndex]->IsMutable() == bTraverseMutable)
+			URigVMNode* Node = NodeProxies[NodeIndex].GetSubjectChecked<URigVMNode>();
+			if (Node->IsMutable() == bTraverseMutable)
 			{
 				if (bTraverseMutable)
 				{
-					TraverseMutableNode(Nodes[NodeIndex], GetObsoleteBlock());
+					TraverseMutableNode(NodeProxies[NodeIndex], GetObsoleteBlock());
 				}
 				else
 				{
-					TraverseNode(Nodes[NodeIndex], GetObsoleteBlock());
+					TraverseNode(NodeProxies[NodeIndex], GetObsoleteBlock());
 				}
 			}
 		}
@@ -902,18 +903,21 @@ FRigVMParserAST::FRigVMParserAST(URigVMGraph* InGraph, URigVMController* InContr
 	}
 }
 
-FRigVMParserAST::FRigVMParserAST(URigVMGraph* InGraph, const TArray<URigVMNode*>& InNodesToCompute)
+FRigVMParserAST::FRigVMParserAST(URigVMGraph* InGraph, const TArray<FRigVMASTProxy>& InNodesToCompute)
 {
 	LastCycleCheckExpr = nullptr;
 
-	FRigVMBlockExprAST* Block = MakeExpr<FRigVMBlockExprAST>(FRigVMExprAST::EType::Block);
+	FRigVMBlockExprAST* Block = MakeExpr<FRigVMBlockExprAST>(FRigVMExprAST::EType::Block, FRigVMASTProxy());
 	Block->Name = TEXT("NodesToCompute");
 	RootExpressions.Add(Block);
 
+	NodeProxies = InNodesToCompute;
+	
 	Inline(InGraph, InNodesToCompute);
 
-	for (URigVMNode* Node : Nodes)
+	for (const FRigVMASTProxy& NodeProxy : NodeProxies)
 	{
+		URigVMNode* Node = NodeProxy.GetSubjectChecked<URigVMNode>();
 		if (Node->IsEvent())
 		{
 			continue;
@@ -922,10 +926,10 @@ FRigVMParserAST::FRigVMParserAST(URigVMGraph* InGraph, const TArray<URigVMNode*>
 		{
 			continue;
 		}
-		TraverseNode(Node, Block);
+		TraverseNode(NodeProxy, Block);
 	}
 
-	FRigVMExprAST* ExitExpr = MakeExpr<FRigVMExitExprAST>();
+	FRigVMExprAST* ExitExpr = MakeExpr<FRigVMExitExprAST>(FRigVMASTProxy());
 	ExitExpr->AddParent(Block);
 }
 
@@ -942,14 +946,14 @@ FRigVMParserAST::~FRigVMParserAST()
 	RootExpressions.Empty();
 }
 
-FRigVMExprAST* FRigVMParserAST::TraverseMutableNode(URigVMNode* InNode, FRigVMExprAST* InParentExpr)
+FRigVMExprAST* FRigVMParserAST::TraverseMutableNode(const FRigVMASTProxy& InNodeProxy, FRigVMExprAST* InParentExpr)
 {
-	if (SubjectToExpression.Contains(InNode))
+	if (SubjectToExpression.Contains(InNodeProxy))
 	{
-		return SubjectToExpression.FindChecked(InNode);
+		return SubjectToExpression.FindChecked(InNodeProxy);
 	}
 
-	FRigVMExprAST* NodeExpr = CreateExpressionForNode(InNode, InParentExpr);
+	FRigVMExprAST* NodeExpr = CreateExpressionForNode(InNodeProxy, InParentExpr);
 	if (NodeExpr)
 	{
 		if (InParentExpr == nullptr)
@@ -957,16 +961,19 @@ FRigVMExprAST* FRigVMParserAST::TraverseMutableNode(URigVMNode* InNode, FRigVMEx
 			InParentExpr = NodeExpr;
 		}
 
-		TraversePins(InNode, NodeExpr);
+		TraversePins(InNodeProxy, NodeExpr);
 
-		for (URigVMPin* SourcePin : InNode->GetPins())
+		URigVMNode* Node = InNodeProxy.GetSubjectChecked<URigVMNode>();
+		for (URigVMPin* SourcePin : Node->GetPins())
 		{
 			if (SourcePin->GetDirection() == ERigVMPinDirection::Output || SourcePin->GetDirection() == ERigVMPinDirection::IO)
 			{
 				if (SourcePin->IsExecuteContext())
 				{
+					FRigVMASTProxy SourcePinProxy = InNodeProxy.GetSibling(SourcePin);
+
 					bool bIsForLoop = false;
-					if (URigVMUnitNode* UnitNode = Cast<URigVMUnitNode>(InNode))
+					if (URigVMUnitNode* UnitNode = Cast<URigVMUnitNode>(Node))
 					{
 						bIsForLoop = UnitNode->IsLoopNode();
 					}
@@ -974,23 +981,26 @@ FRigVMExprAST* FRigVMParserAST::TraverseMutableNode(URigVMNode* InNode, FRigVMEx
 					FRigVMExprAST* ParentExpr = InParentExpr;
 					if (NodeExpr->IsA(FRigVMExprAST::Branch) || bIsForLoop)
 					{
-						if (FRigVMExprAST** PinExpr = SubjectToExpression.Find(SourcePin))
+						if (FRigVMExprAST** PinExpr = SubjectToExpression.Find(SourcePinProxy))
 						{
-							FRigVMBlockExprAST* BlockExpr = MakeExpr<FRigVMBlockExprAST>(FRigVMExprAST::EType::Block);
+							FRigVMBlockExprAST* BlockExpr = MakeExpr<FRigVMBlockExprAST>(FRigVMExprAST::EType::Block, FRigVMASTProxy());
 							BlockExpr->AddParent(*PinExpr);
 							BlockExpr->Name = SourcePin->GetFName();
 							ParentExpr = BlockExpr;
 						}
 					}
 
-					const TArray<URigVMPin*>& TargetPins = GetTargetPins(SourcePin);
-					for(URigVMPin* TargetPin : TargetPins)
+					const TArray<FRigVMASTProxy>& TargetPins = GetTargetPins(SourcePinProxy);
+					for(const FRigVMASTProxy& TargetPinProxy : TargetPins)
 					{
-						if(ShouldLinkBeSkipped(FRigVMPinPair(SourcePin, TargetPin)))
+						if(ShouldLinkBeSkipped(FRigVMPinProxyPair(SourcePinProxy, TargetPinProxy)))
 						{
 							continue;
 						}
-						TraverseMutableNode(TargetPin->GetNode(), ParentExpr);
+
+						URigVMNode* TargetNode = TargetPinProxy.GetSubjectChecked<URigVMPin>()->GetNode();
+						FRigVMASTProxy TargetNodeProxy = TargetPinProxy.GetSibling(TargetNode);
+						TraverseMutableNode(TargetNodeProxy, ParentExpr);
 					}
 				}
 			}
@@ -1000,63 +1010,66 @@ FRigVMExprAST* FRigVMParserAST::TraverseMutableNode(URigVMNode* InNode, FRigVMEx
 	return NodeExpr;
 }
 
-FRigVMExprAST* FRigVMParserAST::TraverseNode(URigVMNode* InNode, FRigVMExprAST* InParentExpr)
+FRigVMExprAST* FRigVMParserAST::TraverseNode(const FRigVMASTProxy& InNodeProxy, FRigVMExprAST* InParentExpr)
 {
-	if (Cast<URigVMCommentNode>(InNode))
+	URigVMNode* Node = InNodeProxy.GetSubjectChecked<URigVMNode>();
+	if (Cast<URigVMCommentNode>(Node))
 	{
 		return nullptr;
 	
 	}
-	if (SubjectToExpression.Contains(InNode))
+	if (SubjectToExpression.Contains(InNodeProxy))
 	{
-		FRigVMExprAST* NodeExpr = SubjectToExpression.FindChecked(InNode);
+		FRigVMExprAST* NodeExpr = SubjectToExpression.FindChecked(InNodeProxy);
 		NodeExpr->AddParent(InParentExpr);
 		return NodeExpr;
 	}
 
-	FRigVMExprAST* NodeExpr = CreateExpressionForNode(InNode, InParentExpr);
+	FRigVMExprAST* NodeExpr = CreateExpressionForNode(InNodeProxy, InParentExpr);
 	if (NodeExpr)
 	{
-		TraversePins(InNode, NodeExpr);
+		TraversePins(InNodeProxy, NodeExpr);
 	}
 
 	return NodeExpr;
 }
 
-FRigVMExprAST* FRigVMParserAST::CreateExpressionForNode(URigVMNode* InNode, FRigVMExprAST* InParentExpr)
+FRigVMExprAST* FRigVMParserAST::CreateExpressionForNode(const FRigVMASTProxy& InNodeProxy, FRigVMExprAST* InParentExpr)
 {
+	URigVMNode* Node = InNodeProxy.GetSubjectChecked<URigVMNode>();
+
 	FRigVMExprAST* NodeExpr = nullptr;
-	if (InNode->IsEvent())
+	if (Node->IsEvent())
 	{
-		NodeExpr = MakeExpr<FRigVMEntryExprAST>(InNode);
-		NodeExpr->Name = InNode->GetEventName();
+		NodeExpr = MakeExpr<FRigVMEntryExprAST>(InNodeProxy);
+		NodeExpr->Name = Node->GetEventName();
 	}
 	else
 	{
-		if (Cast<URigVMRerouteNode>(InNode) ||
-			Cast<URigVMParameterNode>(InNode) ||
-			Cast<URigVMVariableNode>(InNode) ||
-			Cast<URigVMEnumNode>(InNode))
+		if (InNodeProxy.IsA<URigVMRerouteNode>() ||
+			InNodeProxy.IsA<URigVMParameterNode>() ||
+			InNodeProxy.IsA<URigVMVariableNode>() ||
+			InNodeProxy.IsA<URigVMEnumNode>())
 		{
-			NodeExpr = MakeExpr<FRigVMNoOpExprAST>(InNode);
+			NodeExpr = MakeExpr<FRigVMNoOpExprAST>(InNodeProxy);
 		}
-		else if (Cast<URigVMBranchNode>(InNode))
+		else if (InNodeProxy.IsA<URigVMBranchNode>())
 		{
-			NodeExpr = MakeExpr<FRigVMBranchExprAST>(InNode);
+			NodeExpr = MakeExpr<FRigVMBranchExprAST>(InNodeProxy);
 		}
-		else if (Cast<URigVMIfNode>(InNode))
+		else if (InNodeProxy.IsA<URigVMIfNode>())
 		{
-			NodeExpr = MakeExpr<FRigVMIfExprAST>(InNode);
+			NodeExpr = MakeExpr<FRigVMIfExprAST>(InNodeProxy);
 		}
-		else if (Cast<URigVMSelectNode>(InNode))
+		else if (InNodeProxy.IsA<URigVMSelectNode>())
 		{
-			NodeExpr = MakeExpr<FRigVMSelectExprAST>(InNode);
+			NodeExpr = MakeExpr<FRigVMSelectExprAST>(InNodeProxy);
 		}
 		else
 		{
-			NodeExpr = MakeExpr<FRigVMCallExternExprAST>(InNode);
+			NodeExpr = MakeExpr<FRigVMCallExternExprAST>(InNodeProxy);
 		}
-		NodeExpr->Name = InNode->GetFName();
+		NodeExpr->Name = Node->GetFName();
 	}
 
 	if (InParentExpr != nullptr)
@@ -1067,18 +1080,21 @@ FRigVMExprAST* FRigVMParserAST::CreateExpressionForNode(URigVMNode* InNode, FRig
 	{
 		RootExpressions.Add(NodeExpr);
 	}
-	SubjectToExpression.Add(InNode, NodeExpr);
-	NodeExpressionIndex.Add(InNode, NodeExpr->GetIndex());
+	SubjectToExpression.Add(InNodeProxy, NodeExpr);
+	NodeExpressionIndex.Add(InNodeProxy, NodeExpr->GetIndex());
 
 	return NodeExpr;
 }
 
-TArray<FRigVMExprAST*> FRigVMParserAST::TraversePins(URigVMNode* InNode, FRigVMExprAST* InParentExpr)
+TArray<FRigVMExprAST*> FRigVMParserAST::TraversePins(const FRigVMASTProxy& InNodeProxy, FRigVMExprAST* InParentExpr)
 {
+	URigVMNode* Node = InNodeProxy.GetSubjectChecked<URigVMNode>();
 	TArray<FRigVMExprAST*> PinExpressions;
 
-	for (URigVMPin* Pin : InNode->GetPins())
+	for (URigVMPin* Pin : Node->GetPins())
 	{
+		FRigVMASTProxy PinProxy = InNodeProxy.GetSibling(Pin);
+
 		if (Pin->GetDirection() == ERigVMPinDirection::Input &&
 			InParentExpr->IsA(FRigVMExprAST::EType::Select))
 		{
@@ -1087,27 +1103,30 @@ TArray<FRigVMExprAST*> FRigVMParserAST::TraversePins(URigVMNode* InNode, FRigVME
 				const TArray<URigVMPin*>& CasePins = Pin->GetSubPins();
 				for (URigVMPin* CasePin : CasePins)
 				{
-					PinExpressions.Add(TraversePin(CasePin, InParentExpr));
+					FRigVMASTProxy CasePinProxy = PinProxy.GetSibling(CasePin);
+					PinExpressions.Add(TraversePin(CasePinProxy, InParentExpr));
 				}
 				continue;
 			}
 		}
 
-		PinExpressions.Add(TraversePin(Pin, InParentExpr));
+		PinExpressions.Add(TraversePin(PinProxy, InParentExpr));
 	}
 
 	return PinExpressions;
 }
 
-FRigVMExprAST* FRigVMParserAST::TraversePin(URigVMPin* InPin, FRigVMExprAST* InParentExpr)
+FRigVMExprAST* FRigVMParserAST::TraversePin(const FRigVMASTProxy& InPinProxy, FRigVMExprAST* InParentExpr)
 {
-	ensure(!SubjectToExpression.Contains(InPin));
+	ensure(!SubjectToExpression.Contains(InPinProxy));
+
+	URigVMPin* Pin = InPinProxy.GetSubjectChecked<URigVMPin>();
 
-	TArray<FRigVMPinPair> Links = GetSourceLinks(InPin, true);
+	TArray<FRigVMPinProxyPair> Links = GetSourceLinks(InPinProxy, true);
 	
 	if (LinksToSkip.Num() > 0)
 	{
-		Links.RemoveAll([this](const FRigVMPinPair& LinkToCheck)
+		Links.RemoveAll([this](const FRigVMPinProxyPair& LinkToCheck)
 			{
 				return this->ShouldLinkBeSkipped(LinkToCheck);
 			}
@@ -1130,64 +1149,64 @@ FRigVMExprAST* FRigVMParserAST::TraversePin(URigVMPin* InPin, FRigVMExprAST* InP
 	};
 
 	TArray<URigVMPin*> SubPinsBoundToVariables;
-	Local::LookForPinsBoundToVariables(InPin, SubPinsBoundToVariables);
+	Local::LookForPinsBoundToVariables(Pin, SubPinsBoundToVariables);
 
 	FRigVMExprAST* PinExpr = nullptr;
 
-	if (Cast<URigVMVariableNode>(InPin->GetNode()))
+	if (Cast<URigVMVariableNode>(Pin->GetNode()))
 	{
-		if (InPin->GetName() == URigVMVariableNode::VariableName)
+		if (Pin->GetName() == URigVMVariableNode::VariableName)
 		{
 			return nullptr;
 		}
 	}
-	else if (Cast<URigVMParameterNode>(InPin->GetNode()) ||
-		Cast<URigVMEnumNode>(InPin->GetNode()))
+	else if (Cast<URigVMParameterNode>(Pin->GetNode()) ||
+		Cast<URigVMEnumNode>(Pin->GetNode()))
 	{
-		if (InPin->GetDirection() == ERigVMPinDirection::Visible)
+		if (Pin->GetDirection() == ERigVMPinDirection::Visible)
 		{
 			return nullptr;
 		}
 	}
 
-	if ((InPin->GetDirection() == ERigVMPinDirection::Input ||
-		InPin->GetDirection() == ERigVMPinDirection::Visible) &&
+	if ((Pin->GetDirection() == ERigVMPinDirection::Input ||
+		Pin->GetDirection() == ERigVMPinDirection::Visible) &&
 		Links.Num() == 0 &&
 		SubPinsBoundToVariables.Num() == 0)
 	{
-		if (Cast<URigVMParameterNode>(InPin->GetNode()) ||
-			Cast<URigVMVariableNode>(InPin->GetNode()))
+		if (Cast<URigVMParameterNode>(Pin->GetNode()) ||
+			Cast<URigVMVariableNode>(Pin->GetNode()))
 		{
-			PinExpr = MakeExpr<FRigVMVarExprAST>(FRigVMExprAST::EType::Var, InPin);
-			FRigVMExprAST* PinLiteralExpr = MakeExpr<FRigVMLiteralExprAST>(InPin);
+			PinExpr = MakeExpr<FRigVMVarExprAST>(FRigVMExprAST::EType::Var, InPinProxy);
+			FRigVMExprAST* PinLiteralExpr = MakeExpr<FRigVMLiteralExprAST>(InPinProxy);
 			PinLiteralExpr->Name = PinExpr->Name;
-			FRigVMExprAST* PinCopyExpr = MakeExpr<FRigVMCopyExprAST>(InPin, InPin);
+			FRigVMExprAST* PinCopyExpr = MakeExpr<FRigVMCopyExprAST>(InPinProxy, InPinProxy);
 			PinCopyExpr->AddParent(PinExpr);
 			PinLiteralExpr->AddParent(PinCopyExpr);
 		}
-		else if (InPin->IsBoundToVariable())
+		else if (Pin->IsBoundToVariable())
 		{
-			PinExpr = MakeExpr<FRigVMExternalVarExprAST>(InPin);
+			PinExpr = MakeExpr<FRigVMExternalVarExprAST>(InPinProxy);
 		}
 		else
 		{
-			PinExpr = MakeExpr<FRigVMLiteralExprAST>(InPin);
+			PinExpr = MakeExpr<FRigVMLiteralExprAST>(InPinProxy);
 		}
 	}
-	else if (Cast<URigVMEnumNode>(InPin->GetNode()))
+	else if (Cast<URigVMEnumNode>(Pin->GetNode()))
 	{
-		PinExpr = MakeExpr<FRigVMLiteralExprAST>(InPin);
+		PinExpr = MakeExpr<FRigVMLiteralExprAST>(InPinProxy);
 	}
 	else
 	{
-		PinExpr = MakeExpr<FRigVMVarExprAST>(FRigVMExprAST::EType::Var, InPin);
+		PinExpr = MakeExpr<FRigVMVarExprAST>(FRigVMExprAST::EType::Var, InPinProxy);
 	}
 
 	PinExpr->AddParent(InParentExpr);
-	PinExpr->Name = *InPin->GetPinPath();
-	SubjectToExpression.Add(InPin, PinExpr);
+	PinExpr->Name = *Pin->GetPinPath();
+	SubjectToExpression.Add(InPinProxy, PinExpr);
 
-	if (InPin->IsExecuteContext())
+	if (Pin->IsExecuteContext())
 	{
 		return PinExpr;
 	}
@@ -1197,15 +1216,15 @@ FRigVMExprAST* FRigVMParserAST::TraversePin(URigVMPin* InPin, FRigVMExprAST* InP
 		return PinExpr;
 	}
 
-	if ((InPin->GetDirection() == ERigVMPinDirection::IO ||
-		InPin->GetDirection() == ERigVMPinDirection::Input) 
-		&& !InPin->IsExecuteContext())
+	if ((Pin->GetDirection() == ERigVMPinDirection::IO ||
+		Pin->GetDirection() == ERigVMPinDirection::Input)
+		&& !Pin->IsExecuteContext())
 	{
 		bool bHasSourceLinkToRoot = false;
-		URigVMPin* RootPin = InPin->GetRootPin();
-		for (const FRigVMPinPair& SourceLink : Links)
+		URigVMPin* RootPin = Pin->GetRootPin();
+		for (const FRigVMPinProxyPair& SourceLink : Links)
 		{
-			if (SourceLink.Value == RootPin)
+			if (SourceLink.Value.GetSubject() == RootPin)
 			{
 				bHasSourceLinkToRoot = true;
 				break;
@@ -1213,17 +1232,17 @@ FRigVMExprAST* FRigVMParserAST::TraversePin(URigVMPin* InPin, FRigVMExprAST* InP
 		}
 
 		if (!bHasSourceLinkToRoot && 
-			GetSourcePins(InPin).Num() == 0 &&
-			(InPin->GetDirection() == ERigVMPinDirection::IO || Links.Num() > 0 || SubPinsBoundToVariables.Num() > 0))
+			GetSourcePins(InPinProxy).Num() == 0 &&
+			(Pin->GetDirection() == ERigVMPinDirection::IO || Links.Num() > 0 || SubPinsBoundToVariables.Num() > 0))
 		{
-			FRigVMLiteralExprAST* LiteralExpr = MakeExpr<FRigVMLiteralExprAST>(InPin);
-			FRigVMCopyExprAST* LiteralCopyExpr = MakeExpr<FRigVMCopyExprAST>(InPin, InPin);
-			LiteralCopyExpr->Name = *FString::Printf(TEXT("%s -> %s"), *InPin->GetPinPath(), *InPin->GetPinPath());
+			FRigVMLiteralExprAST* LiteralExpr = MakeExpr<FRigVMLiteralExprAST>(InPinProxy);
+			FRigVMCopyExprAST* LiteralCopyExpr = MakeExpr<FRigVMCopyExprAST>(InPinProxy, InPinProxy);
+			LiteralCopyExpr->Name = *FString::Printf(TEXT("%s -> %s"), *Pin->GetPinPath(), *Pin->GetPinPath());
 			LiteralCopyExpr->AddParent(PinExpr);
 			LiteralExpr->AddParent(LiteralCopyExpr);
-			LiteralExpr->Name = *InPin->GetPinPath();
+			LiteralExpr->Name = *Pin->GetPinPath();
 
-			SubjectToExpression[InPin] = LiteralExpr;
+			SubjectToExpression[InPinProxy] = LiteralExpr;
 		}
 		else
 		{
@@ -1237,25 +1256,26 @@ FRigVMExprAST* FRigVMParserAST::TraversePin(URigVMPin* InPin, FRigVMExprAST* InP
 
 	FRigVMExprAST* ParentExprForLinks = PinExpr;
 
-	if ((InPin->GetDirection() == ERigVMPinDirection::IO || InPin->GetDirection() == ERigVMPinDirection::Input) &&
+	if ((Pin->GetDirection() == ERigVMPinDirection::IO || Pin->GetDirection() == ERigVMPinDirection::Input) &&
 		(InParentExpr->IsA(FRigVMExprAST::If) || InParentExpr->IsA(FRigVMExprAST::Select)) &&
 		Links.Num() > 0)
 	{
-		FRigVMBlockExprAST* BlockExpr = MakeExpr<FRigVMBlockExprAST>(FRigVMExprAST::EType::Block);
+		FRigVMBlockExprAST* BlockExpr = MakeExpr<FRigVMBlockExprAST>(FRigVMExprAST::EType::Block, FRigVMASTProxy());
 		BlockExpr->AddParent(PinExpr);
-		BlockExpr->Name = InPin->GetFName();
+		BlockExpr->Name = Pin->GetFName();
 		ParentExprForLinks = BlockExpr;
 	}
 
-	for (const FRigVMPinPair& SourceLink : Links)
+	for (const FRigVMPinProxyPair& SourceLink : Links)
 	{
 		TraverseLink(SourceLink, ParentExprForLinks);
 	}
 
 	for (URigVMPin* SubPinBoundToVariable : SubPinsBoundToVariables)
 	{
-		FRigVMExternalVarExprAST* ExternalVarExpr = MakeExpr<FRigVMExternalVarExprAST>(SubPinBoundToVariable);
-		FRigVMCopyExprAST* ExternalVarExprCopyExpr = MakeExpr<FRigVMCopyExprAST>(SubPinBoundToVariable, SubPinBoundToVariable);
+		FRigVMASTProxy SubPinBoundToVariableProxy = InPinProxy.GetSibling(SubPinBoundToVariable);
+		FRigVMExternalVarExprAST* ExternalVarExpr = MakeExpr<FRigVMExternalVarExprAST>(SubPinBoundToVariableProxy);
+		FRigVMCopyExprAST* ExternalVarExprCopyExpr = MakeExpr<FRigVMCopyExprAST>(SubPinBoundToVariableProxy, SubPinBoundToVariableProxy);
 		ExternalVarExprCopyExpr->Name = *FString::Printf(TEXT("%s -> %s"), *SubPinBoundToVariable->GetPinPath(), *SubPinBoundToVariable->GetPinPath());
 		ExternalVarExprCopyExpr->AddParent(PinExpr);
 		ExternalVarExpr->AddParent(ExternalVarExprCopyExpr);
@@ -1265,12 +1285,15 @@ FRigVMExprAST* FRigVMParserAST::TraversePin(URigVMPin* InPin, FRigVMExprAST* InP
 	return PinExpr;
 }
 
-FRigVMExprAST* FRigVMParserAST::TraverseLink(const FRigVMPinPair& InLink, FRigVMExprAST* InParentExpr)
+FRigVMExprAST* FRigVMParserAST::TraverseLink(const FRigVMPinProxyPair& InLink, FRigVMExprAST* InParentExpr)
 {
-	URigVMPin* SourcePin = InLink.Key;
-	URigVMPin* TargetPin = InLink.Value;
+	const FRigVMASTProxy& SourceProxy = InLink.Key;
+	const FRigVMASTProxy& TargetProxy = InLink.Value;
+	URigVMPin* SourcePin = SourceProxy.GetSubjectChecked<URigVMPin>();
+	URigVMPin* TargetPin = TargetProxy.GetSubjectChecked<URigVMPin>();
 	URigVMPin* SourceRootPin = SourcePin->GetRootPin();
 	URigVMPin* TargetRootPin = TargetPin->GetRootPin();
+	FRigVMASTProxy SourceNodeProxy = SourceProxy.GetSibling(SourcePin->GetNode());
 
 	bool bRequiresCopy = SourceRootPin != SourcePin || TargetRootPin != TargetPin;
 	if (!bRequiresCopy)
@@ -1285,17 +1308,17 @@ FRigVMExprAST* FRigVMParserAST::TraverseLink(const FRigVMPinPair& InLink, FRigVM
 	FRigVMAssignExprAST* AssignExpr = nullptr;
 	if (bRequiresCopy)
 	{
-		AssignExpr = MakeExpr<FRigVMCopyExprAST>(SourcePin, TargetPin);
+		AssignExpr = MakeExpr<FRigVMCopyExprAST>(SourceProxy, TargetProxy);
 	}
 	else
 	{
-		AssignExpr = MakeExpr<FRigVMAssignExprAST>(FRigVMExprAST::EType::Assign, SourcePin, TargetPin);
+		AssignExpr = MakeExpr<FRigVMAssignExprAST>(FRigVMExprAST::EType::Assign, SourceProxy, TargetProxy);
 	}
 
 	AssignExpr->Name = *GetLinkAsString(InLink);
 	AssignExpr->AddParent(InParentExpr);
 
-	FRigVMExprAST* NodeExpr = TraverseNode(SourcePin->GetNode(), AssignExpr);
+	FRigVMExprAST* NodeExpr = TraverseNode(SourceNodeProxy, AssignExpr);
 	if (NodeExpr)
 	{
 		// if this is a copy expression - we should require the copy to use a ref instead
@@ -1328,7 +1351,7 @@ FRigVMExprAST* FRigVMParserAST::TraverseLink(const FRigVMPinPair& InLink, FRigVM
 
 						if (CacheExpr == nullptr)
 						{
-							CacheExpr = MakeExpr<FRigVMCachedValueExprAST>();
+							CacheExpr = MakeExpr<FRigVMCachedValueExprAST>(FRigVMASTProxy());
 							CacheExpr->Name = AssignExpr->GetName();
 							VarExpr->AddParent(CacheExpr);
 							NodeExpr->AddParent(CacheExpr);
@@ -1364,10 +1387,10 @@ void FRigVMParserAST::FoldEntries()
 				FRigVMEntryExprAST* FoldEntry = EntryByName.FindChecked(Entry->GetEventName());
 
 				// replace the original entry with a noop
-				FRigVMNoOpExprAST* NoOpExpr = MakeExpr<FRigVMNoOpExprAST>(Entry->GetNode());
+				FRigVMNoOpExprAST* NoOpExpr = MakeExpr<FRigVMNoOpExprAST>(Entry->GetProxy());
 				NoOpExpr->AddParent(FoldEntry);
 				NoOpExpr->Name = Entry->Name;
-				SubjectToExpression.FindChecked(Entry->GetNode()) = NoOpExpr;
+				SubjectToExpression.FindChecked(Entry->GetProxy()) = NoOpExpr;
 
 				TArray<FRigVMExprAST*> Children = Entry->Children; // copy since the loop changes the array
 				for (FRigVMExprAST* ChildExpr : Children)
@@ -1422,7 +1445,7 @@ void FRigVMParserAST::InjectExitsToEntries()
 
 			if (!bHasExit)
 			{
-				FRigVMExprAST* ExitExpr = MakeExpr<FRigVMExitExprAST>();
+				FRigVMExprAST* ExitExpr = MakeExpr<FRigVMExitExprAST>(FRigVMASTProxy());
 				ExitExpr->AddParent(RootExpr);
 			}
 		}
@@ -1626,8 +1649,8 @@ void FRigVMParserAST::FoldAssignments()
 			if (GrandParent->IsA(FRigVMExprAST::EType::Assign))
 			{
 				FRigVMAssignExprAST* GrandParentAssign = GrandParent->To<FRigVMAssignExprAST>();
-				GrandParentAssign->SourcePin = AssignExpr->SourcePin;
-				GrandParentAssign->Name = *FString::Printf(TEXT("%s -> %s"), *GrandParentAssign->SourcePin->GetPinPath(), *GrandParentAssign->TargetPin->GetPinPath());
+				GrandParentAssign->SourceProxy = AssignExpr->SourceProxy;
+				GrandParentAssign->Name = *FString::Printf(TEXT("%s -> %s"), *GrandParentAssign->GetSourcePin()->GetPinPath(), *GrandParentAssign->GetTargetPin()->GetPinPath());
 			}
 		}
 
@@ -1661,21 +1684,21 @@ bool FRigVMParserAST::FoldConstantValuesToLiterals(URigVMGraph* InGraph, URigVMC
 	// and remove the tree that created the value.
 
 	TMap<FString, FString> ComputedDefaultValues;
-	TArray<URigVMPin*> PinsToUpdate;
-	TArray<URigVMPin*> RootPinsToUpdate;
-	TArray<URigVMPin*> PinsToCompute;
-	TArray<URigVMNode*> NodesToCompute;
+	TArray<FRigVMASTProxy> PinsToUpdate;
+	TArray<FRigVMASTProxy> RootPinsToUpdate;
+	TArray<FRigVMASTProxy> PinsToCompute;
+	TArray<FRigVMASTProxy> NodesToCompute;
 
-	for (URigVMNode* Node : Nodes)
+	for (const FRigVMASTProxy& NodeProxy : NodeProxies)
 	{
-		if (Cast<URigVMParameterNode>(Node) != nullptr ||
-			Cast<URigVMVariableNode>(Node) != nullptr ||
-			Cast<URigVMEnumNode>(Node) != nullptr)
+		if (NodeProxy.IsA<URigVMParameterNode>() ||
+			NodeProxy.IsA<URigVMVariableNode>() ||
+			NodeProxy.IsA<URigVMEnumNode>())
 		{
 			continue;
 		}
 
-		FRigVMExprAST** NodeExprPtr = SubjectToExpression.Find(Node);
+		FRigVMExprAST** NodeExprPtr = SubjectToExpression.Find(NodeProxy);
 		if(NodeExprPtr == nullptr)
 		{
 			continue;
@@ -1687,6 +1710,7 @@ bool FRigVMParserAST::FoldConstantValuesToLiterals(URigVMGraph* InGraph, URigVMC
 			continue;
 		}
 
+		URigVMNode* Node = NodeProxy.GetSubjectChecked<URigVMNode>();
 		const TArray<URigVMPin*> Pins = Node->GetPins();
 		for (URigVMPin* Pin : Pins)
 		{
@@ -1696,7 +1720,8 @@ bool FRigVMParserAST::FoldConstantValuesToLiterals(URigVMGraph* InGraph, URigVMC
 				continue;
 			}
 
-			FRigVMExprAST** PinExprPtr = SubjectToExpression.Find(Pin);
+			FRigVMASTProxy PinProxy = NodeProxy.GetSibling(Pin);
+			FRigVMExprAST** PinExprPtr = SubjectToExpression.Find(PinProxy);
 			if (PinExprPtr == nullptr)
 			{
 				continue;
@@ -1719,7 +1744,7 @@ bool FRigVMParserAST::FoldConstantValuesToLiterals(URigVMGraph* InGraph, URigVMC
 				}
 			}
 
-			TArray<FRigVMPinPair> SourcePins = GetSourceLinks(Pin, true);
+			TArray<FRigVMPinProxyPair> SourcePins = GetSourceLinks(PinProxy, true);
 			if (SourcePins.Num() == 0)
 			{
 				continue;
@@ -1731,28 +1756,31 @@ bool FRigVMParserAST::FoldConstantValuesToLiterals(URigVMGraph* InGraph, URigVMC
 			}
 
 			bool bFoundValidSourcePin = false;
-			for (const FRigVMPinPair& SourcePin : SourcePins)
+			for (const FRigVMPinProxyPair& SourcePin : SourcePins)
 			{
-				URigVMNode* SourceNode = SourcePin.Key->GetNode();
-				check(SourceNode);
+				const FRigVMASTProxy& SourcePinProxy = SourcePin.Key;
+				URigVMNode* SourceNode = SourcePinProxy.GetSubjectChecked<URigVMPin>()->GetNode();
+				FRigVMASTProxy SourceNodeProxy = SourcePin.Key.GetSibling(SourceNode);
+
+				check(SourceNodeProxy.IsValid());
 
-				if (Cast<URigVMParameterNode>(SourceNode) != nullptr ||
-					Cast<URigVMVariableNode>(SourceNode) != nullptr ||
-					Cast<URigVMRerouteNode>(SourceNode) != nullptr ||
-					Cast<URigVMEnumNode>(SourceNode) != nullptr)
+				if (SourceNodeProxy.IsA<URigVMParameterNode>() ||
+					SourceNodeProxy.IsA<URigVMVariableNode>() ||
+					SourceNodeProxy.IsA<URigVMRerouteNode>() ||
+					SourceNodeProxy.IsA<URigVMEnumNode>())
 				{
 					continue;
 				}
 
-				PinsToCompute.AddUnique(SourcePin.Key);
-				NodesToCompute.AddUnique(SourceNode);
+				PinsToCompute.AddUnique(SourcePinProxy);
+				NodesToCompute.AddUnique(SourceNodeProxy);
 				bFoundValidSourcePin = true;
 			}
 
 			if (bFoundValidSourcePin)
 			{
-				PinsToUpdate.Add(Pin);
-				RootPinsToUpdate.AddUnique(Pin->GetRootPin());
+				PinsToUpdate.Add(PinProxy);
+				RootPinsToUpdate.AddUnique(PinProxy.GetSibling(Pin->GetRootPin()));
 			}
 		}
 	}
@@ -1784,14 +1812,16 @@ bool FRigVMParserAST::FoldConstantValuesToLiterals(URigVMGraph* InGraph, URigVMC
 	}
 
 	// copy the values out of the temp VM and set them on the cached value
-	for (URigVMPin* PinToCompute : PinsToCompute)
+	for (const FRigVMASTProxy& PinToComputeProxy : PinsToCompute)
 	{
 		TGuardValue<bool> GuardControllerNotifs(InController->bSuspendNotifications, true);
 
+		URigVMPin* PinToCompute = PinToComputeProxy.GetSubjectChecked<URigVMPin>();
 		URigVMPin* RootPin = PinToCompute->GetRootPin();
+		FRigVMASTProxy RootPinProxy = PinToComputeProxy.GetSibling(RootPin);
 
 		FRigVMVarExprAST* RootVarExpr = nullptr;
-		FRigVMExprAST** RootPinExprPtr = SubjectToExpression.Find(RootPin);
+		FRigVMExprAST** RootPinExprPtr = SubjectToExpression.Find(RootPinProxy);
 		if (RootPinExprPtr != nullptr)
 		{
 			FRigVMExprAST* RootPinExpr = *RootPinExprPtr;
@@ -1865,18 +1895,18 @@ bool FRigVMParserAST::FoldConstantValuesToLiterals(URigVMGraph* InGraph, URigVMC
 			}
 		}
 
-		const TArray<URigVMPin*>& TargetPins = GetTargetPins(PinToCompute);
-		for (URigVMPin* TargetPin : TargetPins)
+		const TArray<FRigVMASTProxy>& TargetPins = GetTargetPins(PinToComputeProxy);
+		for (const FRigVMASTProxy& TargetPinProxy : TargetPins)
 		{
-			PinDefaultValueOverrides.FindOrAdd(TargetPin) = DefaultValue;
+			PinDefaultValueOverrides.FindOrAdd(TargetPinProxy) = DefaultValue;
 		}
 	}
 
 	// now remove all of the expressions no longer needed
 	TArray<FRigVMExprAST*> ExpressionsToRemove;
-	for (URigVMPin* RootPinToUpdate : RootPinsToUpdate)
+	for (const FRigVMASTProxy& RootPinToUpdateProxy : RootPinsToUpdate)
 	{
-		FRigVMExprAST** PreviousExprPtr = SubjectToExpression.Find(RootPinToUpdate);
+		FRigVMExprAST** PreviousExprPtr = SubjectToExpression.Find(RootPinToUpdateProxy);
 		if (PreviousExprPtr)
 		{
 			FRigVMVarExprAST* PreviousVarExpr = (*PreviousExprPtr)->To<FRigVMVarExprAST>();
@@ -1914,9 +1944,9 @@ bool FRigVMParserAST::FoldConstantValuesToLiterals(URigVMGraph* InGraph, URigVMC
 				}
 			}
 
-			FRigVMLiteralExprAST* LiteralExpr = MakeExpr<FRigVMLiteralExprAST>(RootPinToUpdate);
+			FRigVMLiteralExprAST* LiteralExpr = MakeExpr<FRigVMLiteralExprAST>(RootPinToUpdateProxy);
 			LiteralExpr->Name = PreviousVarExpr->Name;
-			SubjectToExpression[RootPinToUpdate] = LiteralExpr;
+			SubjectToExpression[RootPinToUpdateProxy] = LiteralExpr;
 			PreviousVarExpr->ReplaceBy(LiteralExpr);
 			ExpressionsToRemove.Add(PreviousVarExpr);
 		}
@@ -1933,15 +1963,17 @@ bool FRigVMParserAST::FoldUnreachableBranches(URigVMGraph* InGraph)
 
 	TArray<FRigVMExprAST*> ExpressionsToRemove;
 
-	for (URigVMNode* Node : Nodes)
+	for (const FRigVMASTProxy& NodeProxy : NodeProxies)
 	{
-		if (Cast<URigVMParameterNode>(Node) != nullptr ||
-			Cast<URigVMVariableNode>(Node) != nullptr)
+		if (NodeProxy.IsA<URigVMParameterNode>() ||
+			NodeProxy.IsA<URigVMVariableNode>())
 		{
 			continue;
 		}
 
-		FRigVMExprAST** NodeExprPtr = SubjectToExpression.Find(Node);
+		//URigVMNode* Node = NodeProxy.GetSubjectChecked<URigVMNode>();
+
+		FRigVMExprAST** NodeExprPtr = SubjectToExpression.Find(NodeProxy);
 		if (NodeExprPtr == nullptr)
 		{
 			continue;
@@ -2088,9 +2120,9 @@ void FRigVMParserAST::FoldLiterals()
 	RemoveExpressions(ExpressionsToRemove);
 }
 
-const FRigVMExprAST* FRigVMParserAST::GetExprForSubject(UObject* InSubject)
+const FRigVMExprAST* FRigVMParserAST::GetExprForSubject(const FRigVMASTProxy& InProxy) const
 {
-	if (FRigVMExprAST* const* ExpressionPtr = SubjectToExpression.Find(InSubject))
+	if (FRigVMExprAST* const* ExpressionPtr = SubjectToExpression.Find(InProxy))
 	{
 		return *ExpressionPtr;
 	}
@@ -2105,9 +2137,10 @@ void FRigVMParserAST::PrepareCycleChecking(URigVMPin* InPin)
 		CycleCheckFlags.Reset();
 		return;
 	}
+	FRigVMASTProxy NodeProxy = FRigVMASTProxy::MakeFromUObject(InPin->GetNode());
 
 	const FRigVMExprAST* Expression = nullptr;
-	if (FRigVMExprAST* const* ExpressionPtr = SubjectToExpression.Find(InPin->GetNode()))
+	if (FRigVMExprAST* const* ExpressionPtr = SubjectToExpression.Find(NodeProxy))
 	{
 		Expression = *ExpressionPtr;
 	}
@@ -2146,8 +2179,11 @@ bool FRigVMParserAST::CanLink(URigVMPin* InSourcePin, URigVMPin* InTargetPin, FS
 		return false;
 	}
 
+	FRigVMASTProxy SourceNodeProxy = FRigVMASTProxy::MakeFromUObject(SourceNode);
+	FRigVMASTProxy TargetNodeProxy = FRigVMASTProxy::MakeFromUObject(TargetNode);
+
 	const FRigVMExprAST* SourceExpression = nullptr;
-	if (FRigVMExprAST* const* SourceExpressionPtr = SubjectToExpression.Find(SourceNode))
+	if (FRigVMExprAST* const* SourceExpressionPtr = SubjectToExpression.Find(SourceNodeProxy))
 	{
 		SourceExpression = *SourceExpressionPtr;
 	}
@@ -2161,7 +2197,7 @@ bool FRigVMParserAST::CanLink(URigVMPin* InSourcePin, URigVMPin* InTargetPin, FS
 	}
 
 	const FRigVMVarExprAST* SourceVarExpression = nullptr;
-	if (FRigVMExprAST* const* SourceVarExpressionPtr = SubjectToExpression.Find(InSourcePin->GetRootPin()))
+	if (FRigVMExprAST* const* SourceVarExpressionPtr = SubjectToExpression.Find(SourceNodeProxy.GetSibling(InSourcePin->GetRootPin())))
 	{
 		if ((*SourceVarExpressionPtr)->IsA(FRigVMExprAST::EType::Var))
 		{
@@ -2170,7 +2206,7 @@ bool FRigVMParserAST::CanLink(URigVMPin* InSourcePin, URigVMPin* InTargetPin, FS
 	}
 
 	const FRigVMExprAST* TargetExpression = nullptr;
-	if (FRigVMExprAST* const* TargetExpressionPtr = SubjectToExpression.Find(TargetNode))
+	if (FRigVMExprAST* const* TargetExpressionPtr = SubjectToExpression.Find(TargetNodeProxy))
 	{
 		TargetExpression = *TargetExpressionPtr;
 	}
@@ -2330,7 +2366,7 @@ FRigVMBlockExprAST* FRigVMParserAST::GetObsoleteBlock(bool bCreateIfMissing)
 {
 	if (ObsoleteBlock == nullptr && bCreateIfMissing)
 	{
-		ObsoleteBlock = MakeExpr<FRigVMBlockExprAST>(FRigVMExprAST::EType::Block);
+		ObsoleteBlock = MakeExpr<FRigVMBlockExprAST>(FRigVMExprAST::EType::Block, FRigVMASTProxy());
 		ObsoleteBlock->bIsObsolete = true;
 		RootExpressions.Add(ObsoleteBlock);
 	}
@@ -2342,7 +2378,7 @@ const FRigVMBlockExprAST* FRigVMParserAST::GetObsoleteBlock(bool bCreateIfMissin
 	if (ObsoleteBlock == nullptr && bCreateIfMissing)
 	{
 		FRigVMParserAST* MutableThis = (FRigVMParserAST*)this;
-		MutableThis->ObsoleteBlock = MutableThis->MakeExpr<FRigVMBlockExprAST>(FRigVMExprAST::EType::Block);
+		MutableThis->ObsoleteBlock = MutableThis->MakeExpr<FRigVMBlockExprAST>(FRigVMExprAST::EType::Block, FRigVMASTProxy());
 		MutableThis->ObsoleteBlock->bIsObsolete = true;
 		MutableThis->RootExpressions.Add(MutableThis->ObsoleteBlock);
 	}
@@ -2369,15 +2405,15 @@ void FRigVMParserAST::RemoveExpression(FRigVMExprAST* InExpr, bool bRefreshIndic
 
 	Expressions.Remove(InExpr);
 
-	TArray<UObject*> KeysToRemove;
-	for (TPair<UObject*, FRigVMExprAST*> Pair : SubjectToExpression)
+	TArray<FRigVMASTProxy> KeysToRemove;
+	for (TPair<FRigVMASTProxy, FRigVMExprAST*> Pair : SubjectToExpression)
 	{
 		if (Pair.Value == InExpr)
 		{
 			KeysToRemove.Add(Pair.Key);
 		}
 	}
-	for (UObject* KeyToRemove : KeysToRemove)
+	for (const FRigVMASTProxy& KeyToRemove : KeysToRemove)
 	{
 		SubjectToExpression.Remove(KeyToRemove);
 	}
@@ -2426,60 +2462,64 @@ void FRigVMParserAST::TraverseChildren(const FRigVMExprAST* InExpr, TFunctionRef
 	}
 }
 
-const TArray<URigVMPin*>& FRigVMParserAST::GetSourcePins(URigVMPin* InPin) const
+const TArray<FRigVMASTProxy>& FRigVMParserAST::GetSourcePins(const FRigVMASTProxy& InPinProxy) const
 {
-	if (const TArray<URigVMPin*>* Links = SourceLinks.Find(InPin))
+	if (const TArray<FRigVMASTProxy>* LinkedPins = SourceLinks.Find(InPinProxy))
 	{
-		return *Links;
+		return *LinkedPins;
 	}
-	return EmptyLinks;
+	return EmptyProxyArray;
 }
 
-const TArray<URigVMPin*>& FRigVMParserAST::GetTargetPins(URigVMPin* InPin) const
+const TArray<FRigVMASTProxy>& FRigVMParserAST::GetTargetPins(const FRigVMASTProxy& InPinProxy) const
 {
-	if (const TArray<URigVMPin*>* Links = TargetLinks.Find(InPin))
+	if (const TArray<FRigVMASTProxy>* LinkedPins = TargetLinks.Find(InPinProxy))
 	{
-		return *Links;
+		return *LinkedPins;
 	}
-	return EmptyLinks;
+	return EmptyProxyArray;
 }
 
-TArray<FRigVMPinPair> FRigVMParserAST::GetSourceLinks(URigVMPin* InPin, bool bRecursive) const
+TArray<FRigVMPinProxyPair> FRigVMParserAST::GetSourceLinks(const FRigVMASTProxy& InPinProxy, bool bRecursive) const
 {
-	const TArray<URigVMPin*>& SourcePins = GetSourcePins(InPin);
+	const TArray<FRigVMASTProxy>& SourcePins = GetSourcePins(InPinProxy);
 
-	TArray<FRigVMPinPair> Pairs;
-	for (URigVMPin* SourcePin : SourcePins)
+	TArray<FRigVMPinProxyPair> Pairs;
+	for (const FRigVMASTProxy& SourcePin : SourcePins)
 	{
-		Pairs.Add(FRigVMPinPair(SourcePin, InPin));
+		Pairs.Add(FRigVMPinProxyPair(SourcePin, InPinProxy));
 	}
 
 	if (bRecursive)
 	{
-		for (URigVMPin* SubPin : InPin->GetSubPins())
+		URigVMPin* Pin = InPinProxy.GetSubjectChecked<URigVMPin>();
+		for (URigVMPin* SubPin : Pin->GetSubPins())
 		{
-			Pairs.Append(GetSourceLinks(SubPin, true));
+			FRigVMASTProxy SubPinProxy = InPinProxy.GetSibling(SubPin);
+			Pairs.Append(GetSourceLinks(SubPinProxy, true));
 		}
 	}
 
 	return Pairs;
 }
 
-TArray<FRigVMPinPair> FRigVMParserAST::GetTargetLinks(URigVMPin* InPin, bool bRecursive) const
+TArray<FRigVMPinProxyPair> FRigVMParserAST::GetTargetLinks(const FRigVMASTProxy& InPinProxy, bool bRecursive) const
 {
-	const TArray<URigVMPin*>& TargetPins = GetTargetPins(InPin);
+	const TArray<FRigVMASTProxy>& TargetPins = GetTargetPins(InPinProxy);
 
-	TArray<FRigVMPinPair> Pairs;
-	for (URigVMPin* TargetPin : TargetPins)
+	TArray<FRigVMPinProxyPair> Pairs;
+	for (const FRigVMASTProxy& TargetPin : TargetPins)
 	{
-		Pairs.Add(FRigVMPinPair(InPin, TargetPin));
+		Pairs.Add(FRigVMPinProxyPair(InPinProxy, TargetPin));
 	}
 
 	if (bRecursive)
 	{
-		for (URigVMPin* SubPin : InPin->GetSubPins())
+		URigVMPin* Pin = InPinProxy.GetSubjectChecked<URigVMPin>();
+		for (URigVMPin* SubPin : Pin->GetSubPins())
 		{
-			Pairs.Append(GetTargetLinks(SubPin, true));
+			FRigVMASTProxy SubPinProxy = InPinProxy.GetSibling(SubPin);
+			Pairs.Append(GetTargetLinks(SubPinProxy, true));
 		}
 	}
 
@@ -2488,79 +2528,71 @@ TArray<FRigVMPinPair> FRigVMParserAST::GetTargetLinks(URigVMPin* InPin, bool bRe
 
 void FRigVMParserAST::Inline(URigVMGraph* InGraph)
 {
-	Inline(InGraph, InGraph->GetNodes());
+	TArray<FRigVMASTProxy> LocalNodeProxies;
+	for (URigVMNode* LocalNode : InGraph->GetNodes())
+	{
+		LocalNodeProxies.Add(FRigVMASTProxy::MakeFromUObject(LocalNode));
+	}
+	Inline(InGraph, LocalNodeProxies);
 }
 
-void FRigVMParserAST::Inline(URigVMGraph* InGraph, const TArray<URigVMNode*>& InNodes)
+void FRigVMParserAST::Inline(URigVMGraph* InGraph, const TArray<FRigVMASTProxy>& InNodeProxies)
 {
 	struct LocalPinTraversalInfo
 	{
 		URigVMPin::FDefaultValueOverride* PinDefaultValueOverrides;
-		TMap<URigVMPin*, URigVMPin*> SourcePins;
-		TMap<URigVMPin*, TArray<URigVMPin*>>* TargetLinks;
-		TMap<URigVMPin*, TArray<URigVMPin*>>* SourceLinks;
-		TArray<URigVMLibraryNode*> LibraryNodeStack;
+		TMap<FRigVMASTProxy, FRigVMASTProxy> SourcePins;
+		TMap<FRigVMASTProxy, TArray<FRigVMASTProxy>>* TargetLinks;
+		TMap<FRigVMASTProxy, TArray<FRigVMASTProxy>>* SourceLinks;
+		TArray<FRigVMASTProxy> LibraryNodeCallstack;
 
-		class FLibraryNodeGuard
+		static bool ShouldRecursePin(const FRigVMASTProxy& InPinProxy)
 		{
-		public:
-			FLibraryNodeGuard(URigVMLibraryNode* InNode, TArray<URigVMLibraryNode*>& InStack)
-				: Stack(InStack)
-			{
-				Stack.Push(InNode);
-			}
-
-			~FLibraryNodeGuard()
-			{
-				Stack.Pop();
-			}
-
-		private:
-
-			TArray<URigVMLibraryNode*>& Stack;
-		};
-
-		static bool ShouldRecursePin(URigVMPin* InPin)
-		{
-			URigVMNode* Node = InPin->GetNode();
+			URigVMPin* Pin = InPinProxy.GetSubjectChecked<URigVMPin>();
+			URigVMNode* Node = Pin->GetNode();
 			return Node->IsA<URigVMRerouteNode>() ||
 				Node->IsA<URigVMLibraryNode>() ||
 				Node->IsA<URigVMFunctionEntryNode>() ||
 				Node->IsA<URigVMFunctionReturnNode>();
 		}
 
-		static bool IsValidPinForAST(URigVMPin* InPin)
+		static bool IsValidPinForAST(const FRigVMASTProxy& InPinProxy)
 		{
-			return !ShouldRecursePin(InPin);
+			return !ShouldRecursePin(InPinProxy);
 		}
 
-		static bool IsValidLinkForAST(URigVMPin* InSourcePin, URigVMPin* InTargetPin)
+		static bool IsValidLinkForAST(const FRigVMASTProxy& InSourcePinProxy, const FRigVMASTProxy& InTargetPinProxy)
 		{
-			return IsValidPinForAST(InSourcePin) && IsValidPinForAST(InTargetPin);
+			return IsValidPinForAST(InSourcePinProxy) && IsValidPinForAST(InTargetPinProxy);
 		}
 
-		static URigVMPin* FindSourcePin(URigVMPin* InPin, LocalPinTraversalInfo& OutTraversalInfo)
+		static FRigVMASTProxy FindSourcePin(const FRigVMASTProxy& InPinProxy, LocalPinTraversalInfo& OutTraversalInfo)
 		{
+			URigVMPin* Pin = InPinProxy.GetSubjectChecked<URigVMPin>();
+
 			// if this pin is a root on a library
-			if (InPin->GetParentPin() == nullptr)
+			if (Pin->GetParentPin() == nullptr)
 			{
-				if (InPin->GetDirection() == ERigVMPinDirection::Output ||
-					InPin->GetDirection() == ERigVMPinDirection::IO)
+				if (Pin->GetDirection() == ERigVMPinDirection::Output ||
+					Pin->GetDirection() == ERigVMPinDirection::IO)
 				{
-					URigVMNode* Node = InPin->GetNode();
+					URigVMNode* Node = Pin->GetNode();
 					if (URigVMLibraryNode* LibraryNode = Cast<URigVMLibraryNode>(Node))
 					{
-						if (!OutTraversalInfo.LibraryNodeStack.Contains(LibraryNode))
+						FRigVMASTProxy LibraryNodeProxy = InPinProxy.GetSibling(LibraryNode);
+						if (!OutTraversalInfo.LibraryNodeCallstack.Contains(LibraryNodeProxy))
 						{
 							if (URigVMFunctionReturnNode* ReturnNode = LibraryNode->GetReturnNode())
 							{
-								if (URigVMPin* ReturnPin = ReturnNode->FindPin(InPin->GetName()))
+								if (URigVMPin* ReturnPin = ReturnNode->FindPin(Pin->GetName()))
 								{
-									FLibraryNodeGuard NodeGuard(LibraryNode, OutTraversalInfo.LibraryNodeStack);
-									URigVMPin* SourcePin = FindSourcePin(ReturnPin, OutTraversalInfo);
-									SourcePin = SourcePin == nullptr ? ReturnPin : SourcePin;
-									OutTraversalInfo.SourcePins.FindOrAdd(InPin) = SourcePin;
-									return SourcePin;
+									OutTraversalInfo.LibraryNodeCallstack.Push(LibraryNodeProxy);
+									FRigVMASTProxy ReturnPinProxy = LibraryNodeProxy.GetChild(ReturnPin);
+									FRigVMASTProxy SourcePinProxy = FindSourcePin(ReturnPinProxy, OutTraversalInfo);
+									SourcePinProxy = SourcePinProxy.IsValid() ? SourcePinProxy : ReturnPinProxy;
+									OutTraversalInfo.SourcePins.FindOrAdd(InPinProxy) = SourcePinProxy;
+									OutTraversalInfo.LibraryNodeCallstack.Pop();
+									return SourcePinProxy;
 
 								}
 							}
@@ -2568,17 +2600,24 @@ void FRigVMParserAST::Inline(URigVMGraph* InGraph, const TArray<URigVMNode*>& In
 					}
 					else if (URigVMFunctionEntryNode* EntryNode = Cast<URigVMFunctionEntryNode>(Node))
 					{
-						for (int32 LibraryNodeIndex = OutTraversalInfo.LibraryNodeStack.Num() - 1; LibraryNodeIndex >= 0; LibraryNodeIndex--)
+						for (int32 LibraryNodeIndex = OutTraversalInfo.LibraryNodeCallstack.Num() - 1; LibraryNodeIndex >= 0; LibraryNodeIndex--)
 						{
-							URigVMLibraryNode* LastLibraryNode = OutTraversalInfo.LibraryNodeStack[LibraryNodeIndex];
+							FRigVMASTProxy LibraryNodeProxy = OutTraversalInfo.LibraryNodeCallstack[LibraryNodeIndex];
+							URigVMLibraryNode* LastLibraryNode = LibraryNodeProxy.GetSubject<URigVMLibraryNode>();
+							if (LastLibraryNode == nullptr)
+							{
+								continue;
+							}
+
 							if(LastLibraryNode->GetEntryNode() == EntryNode)
 							{
-								if (URigVMPin* LibraryPin = LastLibraryNode->FindPin(InPin->GetName()))
+								if (URigVMPin* LibraryPin = LastLibraryNode->FindPin(Pin->GetName()))
 								{
-									URigVMPin* SourcePin = FindSourcePin(LibraryPin, OutTraversalInfo);
-									SourcePin = SourcePin == nullptr ? LibraryPin : SourcePin;
-									OutTraversalInfo.SourcePins.FindOrAdd(InPin) = SourcePin;
-									return SourcePin;
+									FRigVMASTProxy LibraryPinProxy = LibraryNodeProxy.GetSibling(LibraryPin);
+									FRigVMASTProxy SourcePinProxy = FindSourcePin(LibraryPinProxy, OutTraversalInfo);
+									SourcePinProxy = SourcePinProxy.IsValid() ? SourcePinProxy : LibraryPinProxy;
+									OutTraversalInfo.SourcePins.FindOrAdd(InPinProxy) = SourcePinProxy;
+									return SourcePinProxy;
 								}
 							}
 						}
@@ -2586,18 +2625,18 @@ void FRigVMParserAST::Inline(URigVMGraph* InGraph, const TArray<URigVMNode*>& In
 				}
 			}
 
-			if (InPin->GetDirection() != ERigVMPinDirection::Input &&
-				InPin->GetDirection() != ERigVMPinDirection::IO)
+			if (Pin->GetDirection() != ERigVMPinDirection::Input &&
+				Pin->GetDirection() != ERigVMPinDirection::IO)
 			{
-				return nullptr;
+				return FRigVMASTProxy();
 			}
 
 			bool bIOPinOnLeftOfLibraryNode = false;
-			if (InPin->GetDirection() == ERigVMPinDirection::IO)
+			if (Pin->GetDirection() == ERigVMPinDirection::IO)
 			{
-				if (URigVMLibraryNode* LibraryNode = Cast<URigVMLibraryNode>(InPin->GetNode()))
+				if (URigVMLibraryNode* LibraryNode = Cast<URigVMLibraryNode>(Pin->GetNode()))
 				{
-					bIOPinOnLeftOfLibraryNode = OutTraversalInfo.LibraryNodeStack.Contains(LibraryNode);
+					bIOPinOnLeftOfLibraryNode = OutTraversalInfo.LibraryNodeCallstack.Contains(InPinProxy.GetSibling(LibraryNode));
 				}
 			}
 
@@ -2605,29 +2644,31 @@ void FRigVMParserAST::Inline(URigVMGraph* InGraph, const TArray<URigVMNode*>& In
 			{
 				// note: this map isn't going to work for functions which are referenced.
 				// (since the pin objects are shared between multiple invocation nodes)
-				if (URigVMPin* const* SourcePin = OutTraversalInfo.SourcePins.Find(InPin))
+				if (const FRigVMASTProxy* SourcePinProxy = OutTraversalInfo.SourcePins.Find(InPinProxy))
 				{
-					return *SourcePin;
+					return *SourcePinProxy;
 				}
 			}
 
 			TArray<FString> SegmentPath;
-			URigVMPin* SourcePin = nullptr;
+			FRigVMASTProxy SourcePinProxy;
 
-			URigVMPin* ChildPin = InPin;
+			URigVMPin* ChildPin = Pin;
 			while (ChildPin != nullptr)
 			{
 				TArray<URigVMLink*> SourceLinks = ChildPin->GetSourceLinks(false /* recursive */);
 				if (SourceLinks.Num() > 0)
 				{
-					SourcePin = SourceLinks[0]->GetSourcePin();
+					URigVMPin* SourcePin = SourceLinks[0]->GetSourcePin();
+					SourcePinProxy = InPinProxy.GetSibling(SourcePin);
 
 					// only continue the recursion on reroutes
-					if (ShouldRecursePin(SourcePin))
+					if (ShouldRecursePin(SourcePinProxy))
 					{
-						if (URigVMPin* SourceSourcePin = FindSourcePin(SourcePin, OutTraversalInfo))
+						FRigVMASTProxy SourceSourcePinProxy = FindSourcePin(SourcePinProxy, OutTraversalInfo);
+						if(SourceSourcePinProxy.IsValid())
 						{
-							SourcePin = SourceSourcePin;
+							SourcePinProxy = SourceSourcePinProxy;
 						}
 					}
 
@@ -2637,15 +2678,17 @@ void FRigVMParserAST::Inline(URigVMGraph* InGraph, const TArray<URigVMNode*>& In
 				URigVMPin* ParentPin = ChildPin->GetParentPin();
 				if (ParentPin)
 				{
+					FRigVMASTProxy ParentPinProxy = InPinProxy.GetSibling(ParentPin);
+
 					// if we found a parent pin which has a source that is not a reroute
-					if (URigVMPin** ParentSourcePinPtr = OutTraversalInfo.SourcePins.Find(ParentPin))
+					if (FRigVMASTProxy* ParentSourcePinProxyPtr = OutTraversalInfo.SourcePins.Find(ParentPinProxy))
 					{
-						URigVMPin* ParentSourcePin = *ParentSourcePinPtr;
-						if (ParentSourcePin)
+						FRigVMASTProxy& ParentSourcePinProxy = *ParentSourcePinProxyPtr;
+						if (ParentSourcePinProxy.IsValid())
 						{
-							if (!ShouldRecursePin(ParentSourcePin))
+							if (!ShouldRecursePin(ParentSourcePinProxy))
 							{
-								SourcePin = nullptr;
+								SourcePinProxy = FRigVMASTProxy();
 								break;
 							}
 						}
@@ -2656,28 +2699,30 @@ void FRigVMParserAST::Inline(URigVMGraph* InGraph, const TArray<URigVMNode*>& In
 				ChildPin = ParentPin;
 			}
 			
-			if (SourcePin)
+			if (SourcePinProxy.IsValid())
 			{
 				while (!SegmentPath.IsEmpty())
 				{
 					FString Segment = SegmentPath.Pop();
 
+					URigVMPin* SourcePin = SourcePinProxy.GetSubjectChecked<URigVMPin>();
 					if (URigVMPin* SourceSubPin = SourcePin->FindSubPin(Segment))
 					{
-						SourcePin = SourceSubPin;
+						SourcePinProxy = SourcePinProxy.GetSibling(SourceSubPin);
 
 						// only continue the recursion on reroutes
-						if (ShouldRecursePin(SourcePin))
+						if (ShouldRecursePin(SourcePinProxy))
 						{
-							if (URigVMPin* SourceSourceSubPin = FindSourcePin(SourcePin, OutTraversalInfo))
+							FRigVMASTProxy SourceSourceSubPinProxy = FindSourcePin(SourcePinProxy, OutTraversalInfo);
+							if (SourceSourceSubPinProxy.IsValid())
 							{
-								SourcePin = SourceSourceSubPin;
+								SourcePinProxy = SourceSourceSubPinProxy;
 							}
 						}
 					}
 					else
 					{
-						SourcePin = nullptr;
+						SourcePinProxy = FRigVMASTProxy();
 						break;
 					}
 				}
@@ -2685,14 +2730,15 @@ void FRigVMParserAST::Inline(URigVMGraph* InGraph, const TArray<URigVMNode*>& In
 
 			if (!bIOPinOnLeftOfLibraryNode)
 			{
-				OutTraversalInfo.SourcePins.FindOrAdd(InPin) = SourcePin;
+				OutTraversalInfo.SourcePins.FindOrAdd(InPinProxy) = SourcePinProxy;
 			}
-			return SourcePin;
+			return SourcePinProxy;
 		}
 
-		static void VisitPin(URigVMPin* InPin, LocalPinTraversalInfo& OutTraversalInfo)
+		static void VisitPin(const FRigVMASTProxy& InPinProxy, LocalPinTraversalInfo& OutTraversalInfo)
 		{
-			if (URigVMPin* SourcePin = FindSourcePin(InPin, OutTraversalInfo))
+			FRigVMASTProxy SourcePinProxy = FindSourcePin(InPinProxy, OutTraversalInfo);
+			if (SourcePinProxy.IsValid())
 			{
 				// The source pin is the final determined source pin, since
 				// FindSourcePin is recursive.
@@ -2702,65 +2748,72 @@ void FRigVMParserAST::Inline(URigVMGraph* InGraph, const TArray<URigVMNode*>& In
 				// same goes for library nodes or return nodes - we'll 
 				// just use the default pin value in that case.
 
+				URigVMPin* SourcePin = SourcePinProxy.GetSubjectChecked<URigVMPin>();
 				URigVMNode* SourceNode = SourcePin->GetNode();
 				if (SourceNode->IsA<URigVMRerouteNode>() ||
 					SourceNode->IsA<URigVMLibraryNode>() ||
 					SourceNode->IsA<URigVMFunctionReturnNode>())
 				{
-					OutTraversalInfo.PinDefaultValueOverrides->FindOrAdd(InPin) = SourcePin->GetDefaultValue();
+					OutTraversalInfo.PinDefaultValueOverrides->FindOrAdd(InPinProxy) = SourcePin->GetDefaultValue();
 				}
 				else
 				{
-					if (IsValidLinkForAST(SourcePin, InPin))
+					if (IsValidLinkForAST(SourcePinProxy, InPinProxy))
 					{
-						OutTraversalInfo.SourceLinks->FindOrAdd(InPin).Add(SourcePin);
-						OutTraversalInfo.TargetLinks->FindOrAdd(SourcePin).Add(InPin);
+						OutTraversalInfo.SourceLinks->FindOrAdd(InPinProxy).Add(SourcePinProxy);
+						OutTraversalInfo.TargetLinks->FindOrAdd(SourcePinProxy).Add(InPinProxy);
 					}
 				}
 			}
 
-			for (URigVMPin* SubPin : InPin->GetSubPins())
+			URigVMPin* Pin = InPinProxy.GetSubjectChecked<URigVMPin>();
+			for (URigVMPin* SubPin : Pin->GetSubPins())
 			{
-				VisitPin(SubPin, OutTraversalInfo);
+				FRigVMASTProxy SubPinProxy = InPinProxy.GetSibling(SubPin);
+				VisitPin(SubPinProxy, OutTraversalInfo);
 			}
 		}
 
-		static void VisitNode(URigVMNode* InNode, LocalPinTraversalInfo& OutTraversalInfo)
+		static void VisitNode(const FRigVMASTProxy& InNodeProxy, LocalPinTraversalInfo& OutTraversalInfo)
 		{
-			if (InNode->IsA<URigVMRerouteNode>() ||
-				InNode->IsA<URigVMFunctionEntryNode>() || 
-				InNode->IsA<URigVMFunctionReturnNode>())
+			if (InNodeProxy.IsA<URigVMRerouteNode>() ||
+				InNodeProxy.IsA<URigVMFunctionEntryNode>() ||
+				InNodeProxy.IsA<URigVMFunctionReturnNode>())
 			{
 				return;
 			}
 
-			if (URigVMLibraryNode* LibraryNode = Cast<URigVMLibraryNode>(InNode))
+			if (URigVMLibraryNode* LibraryNode = InNodeProxy.GetSubject<URigVMLibraryNode>())
 			{
-				FLibraryNodeGuard NodeGuard(LibraryNode, OutTraversalInfo.LibraryNodeStack);
-
+				OutTraversalInfo.LibraryNodeCallstack.Push(InNodeProxy);
 				TArray<URigVMNode*> ContainedNodes = LibraryNode->GetContainedNodes();
 				for (URigVMNode* ContainedNode : ContainedNodes)
 				{
-					VisitNode(ContainedNode, OutTraversalInfo);
+					// create a proxy which uses the previous node as a callstack
+					FRigVMASTProxy ContainedNodeProxy = InNodeProxy.GetChild(ContainedNode);
+					VisitNode(ContainedNodeProxy, OutTraversalInfo);
 				}
+				OutTraversalInfo.LibraryNodeCallstack.Pop();
 			}
 			else
 			{
-				for (URigVMPin* Pin : InNode->GetPins())
+				URigVMNode* Node = InNodeProxy.GetSubjectChecked<URigVMNode>();
+				for (URigVMPin* Pin : Node->GetPins())
 				{
-					LocalPinTraversalInfo::VisitPin(Pin, OutTraversalInfo);
+					FRigVMASTProxy PinProxy = InNodeProxy.GetSibling(Pin);
+					LocalPinTraversalInfo::VisitPin(PinProxy, OutTraversalInfo);
 				}
 			}
 		}
 	};
 
-	Nodes.Reset();
+	NodeProxies.Reset();
 	SourceLinks.Reset();
 	TargetLinks.Reset();
 
 	// a) find all of the relevant nodes,
 	//    inline and traverse into library nodes
-	Nodes = InNodes;
+	NodeProxies = InNodeProxies;
 
 	// c) flatten links from an entry node / to a return node
 	//    also traverse links along reroutes and flatten them
@@ -2769,18 +2822,21 @@ void FRigVMParserAST::Inline(URigVMGraph* InGraph, const TArray<URigVMNode*>& In
 	TraversalInfo.TargetLinks = &TargetLinks;
 	TraversalInfo.SourceLinks = &SourceLinks;
 
-	for (URigVMNode* Node : Nodes)
+	for (const FRigVMASTProxy& NodeProxy : NodeProxies)
 	{
-		LocalPinTraversalInfo::VisitNode(Node, TraversalInfo);
+		LocalPinTraversalInfo::VisitNode(NodeProxy, TraversalInfo);
 	}
 }
 
-bool FRigVMParserAST::ShouldLinkBeSkipped(const FRigVMPinPair& InLink) const
+bool FRigVMParserAST::ShouldLinkBeSkipped(const FRigVMPinProxyPair& InLink) const
 {
+	URigVMPin* SourcePin = InLink.Key.GetSubjectChecked<URigVMPin>();
+	URigVMPin* TargetPin = InLink.Value.GetSubjectChecked<URigVMPin>();
+
 	for (URigVMLink* LinkToSkip : LinksToSkip)
 	{
-		if (LinkToSkip->GetSourcePin() == InLink.Key &&
-			LinkToSkip->GetTargetPin() == InLink.Value)
+		if (LinkToSkip->GetSourcePin() == SourcePin &&
+			LinkToSkip->GetTargetPin() == TargetPin)
 		{
 			return true;
 		}
@@ -2788,7 +2844,10 @@ bool FRigVMParserAST::ShouldLinkBeSkipped(const FRigVMPinPair& InLink) const
 	return false;
 }
 
-FString FRigVMParserAST::GetLinkAsString(const FRigVMPinPair& InLink)
+FString FRigVMParserAST::GetLinkAsString(const FRigVMPinProxyPair& InLink)
 {
-	return FString::Printf(TEXT("%s -> %s"), *InLink.Key->GetPinPath(), *InLink.Value->GetPinPath());
+	URigVMPin* SourcePin = InLink.Key.GetSubjectChecked<URigVMPin>();
+	URigVMPin* TargetPin = InLink.Value.GetSubjectChecked<URigVMPin>();
+
+	return FString::Printf(TEXT("%s -> %s"), *SourcePin->GetPinPath(), *TargetPin->GetPinPath());
 }
\ No newline at end of file
diff --git a/Engine/Source/Developer/RigVMDeveloper/Private/RigVMCompiler/RigVMASTProxy.cpp b/Engine/Source/Developer/RigVMDeveloper/Private/RigVMCompiler/RigVMASTProxy.cpp
new file mode 100644
--- /dev/null
+++ b/Engine/Source/Developer/RigVMDeveloper/Private/RigVMCompiler/RigVMASTProxy.cpp
@@ -0,0 +1,95 @@
+// Copyright Epic Games, Inc. All Rights Reserved.
+
+#include "RigVMCompiler/RigVMASTProxy.h"
+#include "RigVMModel/RigVMGraph.h"
+#include "RigVMModel/RigVMNode.h"
+#include "RigVMModel/RigVMPin.h"
+#include "RigVMModel/Nodes/RigVMLibraryNode.h"
+
+FString FRigVMCallstack::GetCallPath(bool bIncludeLast) const
+{
+	TArray<FString> Segments;
+	for (UObject* Entry : Stack)
+	{
+		if (URigVMNode* Node = Cast<URigVMNode>(Entry))
+		{
+			Segments.Add(Node->GetName());
+		}
+		else if (URigVMPin* Pin = Cast<URigVMPin>(Entry))
+		{
+			Segments.Add(Pin->GetPinPath());
+		}
+	}
+
+	if (!bIncludeLast)
+	{
+		Segments.Pop();
+	}
+
+	if (Segments.Num() == 0)
+	{
+		return FString();
+	}
+
+	FString Result = Segments[0];
+	for (int32 PartIndex = 1; PartIndex < Segments.Num(); PartIndex++)
+	{
+		Result += TEXT("|") + Segments[PartIndex];
+	}
+
+	return Result;
+}
+
+int32 FRigVMCallstack::Num() const
+{
+	return Stack.Num();
+}
+
+const UObject* FRigVMCallstack::operator[](int32 InIndex) const
+{
+	return Stack[InIndex];
+}
+
+bool FRigVMCallstack::Contains(UObject* InEntry) const
+{
+	return Stack.Contains(InEntry);
+}
+
+FRigVMASTProxy FRigVMASTProxy::MakeFromUObject(UObject* InSubject)
+{
+	UObject* Subject = InSubject;
+
+	FRigVMASTProxy Proxy;
+	Proxy.Callstack.Stack.Reset();
+
+	while (Subject != nullptr)
+	{
+		if (URigVMPin* Pin = Cast<URigVMPin>(Subject))
+		{
+			Subject = Pin->GetNode();
+		}
+		else if (URigVMNode* Node = Cast<URigVMNode>(Subject))
+		{
+			Subject = Node->GetGraph();
+		}
+		else if (URigVMGraph* Graph = Cast<URigVMGraph>(Subject))
+		{
+			Subject = Cast<URigVMLibraryNode>(Graph->GetOuter());
+			if (Subject)
+			{
+				Proxy.Callstack.Stack.Insert(Subject, 0);
+			}
+		}
+		else
+		{
+			break;
+		}
+	}
+
+	Proxy.Callstack.Stack.Push(InSubject);
+
+#if UE_BUILD_DEBUG
+	Proxy.DebugName = Proxy.Callstack.GetCallPath();
+#endif
+	return Proxy;
+}
diff --git a/Engine/Source/Developer/RigVMDeveloper/Private/RigVMCompiler/RigVMCompiler.cpp b/Engine/Source/Developer/RigVMDeveloper/Private/RigVMCompiler/RigVMCompiler.cpp
--- a/Engine/Source/Developer/RigVMDeveloper/Private/RigVMCompiler/RigVMCompiler.cpp
+++ b/Engine/Source/Developer/RigVMDeveloper/Private/RigVMCompiler/RigVMCompiler.cpp
@@ -145,7 +145,9 @@ bool URigVMCompiler::Compile(URigVMGraph* InGraph, URigVMController* InControlle
 	{
 		AST = MakeShareable(new FRigVMParserAST(InGraph, InController, Settings.ASTSettings, InExternalVariables, UserData));
 		InGraph->RuntimeAST = AST;
+#if UE_BUILD_DEBUG
 		//UE_LOG(LogRigVMDeveloper, Display, TEXT("%s"), *AST->DumpDot());
+#endif
 	}
 	ensure(AST.IsValid());
 
@@ -192,7 +194,8 @@ bool URigVMCompiler::Compile(URigVMGraph* InGraph, URigVMController* InControlle
 	{
 		if (URigVMParameterNode* ParameterNode = Cast<URigVMParameterNode>(Node))
 		{
-			const FRigVMExprAST* ParameterExpr = AST->GetExprForSubject(ParameterNode);
+			FRigVMASTProxy ParameterNodeProxy = FRigVMASTProxy::MakeFromUObject(ParameterNode);
+			const FRigVMExprAST* ParameterExpr = AST->GetExprForSubject(ParameterNodeProxy);
 			if (!ParameterExpr || ParameterExpr->IsA(FRigVMExprAST::NoOp))
 			{
 				FName Name = ParameterNode->GetParameterName();
@@ -500,7 +503,8 @@ int32 URigVMCompiler::TraverseCallExtern(const FRigVMCallExternExprAST* InExpr,
 			{
 				if (Pin->RequiresWatch())
 				{
-					FRigVMVarExprAST TempVarExpr(FRigVMExprAST::EType::Var, Pin);
+					FRigVMASTProxy PinProxy = FRigVMASTProxy::MakeFromUObject(Pin);
+					FRigVMVarExprAST TempVarExpr(FRigVMExprAST::EType::Var, PinProxy);
 					TempVarExpr.ParserPtr = InExpr->GetParser();
 					FindOrAddRegister(&TempVarExpr, WorkData, true);
 				}
@@ -1176,6 +1180,23 @@ FString URigVMCompiler::GetPinHash(URigVMPin* InPin, const FRigVMVarExprAST* InV
 			return FString::Printf(TEXT("%sVariable::%s%s"), *Prefix, *VariableNode->GetVariableName().ToString(), *Suffix);
 		}
 	}
+	else
+	{
+		if (InVarExpr)
+		{
+			FRigVMASTProxy NodeProxy = InVarExpr->GetProxy().GetSibling(Node);
+			if (const FRigVMExprAST* NodeExpr = InVarExpr->GetParser()->GetExprForSubject(NodeProxy))
+			{
+				// rely on the proxy callstack to differentiate registers
+				FString CallStackPath = NodeProxy.GetCallstack().GetCallPath(false /* include last */);
+				if (!CallStackPath.IsEmpty())
+				{
+					Prefix += CallStackPath + TEXT("|");
+				}
+			}
+		}
+	}
+
 	return FString::Printf(TEXT("%s%s%s"), *Prefix, *InPin->GetPinPath(true /*  full node path */), *Suffix);
 }
 
@@ -1353,14 +1374,14 @@ FRigVMOperand URigVMCompiler::FindOrAddRegister(const FRigVMVarExprAST* InVarExp
 
 				int32 DesiredArraySize = VMStruct->GetArraySize(Pin->GetFName(), WorkData.RigVMUserData);
 
-				DefaultValues = URigVMPin::SplitDefaultValue(Pin->GetDefaultValue(DefaultValueOverride));
+				DefaultValues = URigVMPin::SplitDefaultValue(Pin->GetDefaultValue(DefaultValueOverride, InVarExpr->GetProxy()));
 
 				if (DefaultValues.Num() != DesiredArraySize)
 				{
 					FString DefaultValue;
 					if (Pin->GetArraySize() > 0)
 					{
-						DefaultValue = Pin->GetSubPins()[0]->GetDefaultValue(DefaultValueOverride);
+						DefaultValue = Pin->GetSubPins()[0]->GetDefaultValue(DefaultValueOverride, InVarExpr->GetProxy());
 					}
 
 					DefaultValues.Reset();
@@ -1372,7 +1393,7 @@ FRigVMOperand URigVMCompiler::FindOrAddRegister(const FRigVMVarExprAST* InVarExp
 			}
 			else
 			{
-				DefaultValues = URigVMPin::SplitDefaultValue(Pin->GetDefaultValue(DefaultValueOverride));
+				DefaultValues = URigVMPin::SplitDefaultValue(Pin->GetDefaultValue(DefaultValueOverride, InVarExpr->GetProxy()));
 			}
 
 			while (DefaultValues.Num() < Pin->GetSubPins().Num())
@@ -1382,7 +1403,7 @@ FRigVMOperand URigVMCompiler::FindOrAddRegister(const FRigVMVarExprAST* InVarExp
 		}
 		else if (URigVMEnumNode* EnumNode = Cast<URigVMEnumNode>(Pin->GetNode()))
 		{
-			FString EnumValueStr = EnumNode->GetDefaultValue(DefaultValueOverride);
+			FString EnumValueStr = EnumNode->GetDefaultValue(DefaultValueOverride, InVarExpr->GetProxy());
 			if (UEnum* Enum = EnumNode->GetEnum())
 			{
 				DefaultValues.Add(FString::FromInt((int32)Enum->GetValueByNameString(EnumValueStr)));
@@ -1394,7 +1415,7 @@ FRigVMOperand URigVMCompiler::FindOrAddRegister(const FRigVMVarExprAST* InVarExp
 		}
 		else
 		{
-			DefaultValues.Add(Pin->GetDefaultValue(DefaultValueOverride));
+			DefaultValues.Add(Pin->GetDefaultValue(DefaultValueOverride, InVarExpr->GetProxy()));
 		}
 
 		UScriptStruct* ScriptStruct = Pin->GetScriptStruct();
diff --git a/Engine/Source/Developer/RigVMDeveloper/Private/RigVMModel/Nodes/RigVMEnumNode.cpp b/Engine/Source/Developer/RigVMDeveloper/Private/RigVMModel/Nodes/RigVMEnumNode.cpp
--- a/Engine/Source/Developer/RigVMDeveloper/Private/RigVMModel/Nodes/RigVMEnumNode.cpp
+++ b/Engine/Source/Developer/RigVMDeveloper/Private/RigVMModel/Nodes/RigVMEnumNode.cpp
@@ -44,12 +44,12 @@ UObject* URigVMEnumNode::GetCPPTypeObject() const
 	return EnumValuePin->GetCPPTypeObject();
 }
 
-FString URigVMEnumNode::GetDefaultValue(const URigVMPin::FDefaultValueOverride& InDefaultValueOverride) const
+FString URigVMEnumNode::GetDefaultValue(const URigVMPin::FDefaultValueOverride& InDefaultValueOverride, const FRigVMASTProxy& InProxy) const
 {
 	URigVMPin* EnumValuePin = FindPin(EnumValueName);
 	if (EnumValuePin == nullptr)
 	{
 		return FString();
 	}
-	return EnumValuePin->GetDefaultValue(InDefaultValueOverride);
+	return EnumValuePin->GetDefaultValue(InDefaultValueOverride, InProxy);
 }
diff --git a/Engine/Source/Developer/RigVMDeveloper/Private/RigVMModel/RigVMController.cpp b/Engine/Source/Developer/RigVMDeveloper/Private/RigVMModel/RigVMController.cpp
--- a/Engine/Source/Developer/RigVMDeveloper/Private/RigVMModel/RigVMController.cpp
+++ b/Engine/Source/Developer/RigVMDeveloper/Private/RigVMModel/RigVMController.cpp
@@ -203,7 +203,8 @@ void URigVMController::HandleModifiedEvent(ERigVMGraphNotifType InNotifType, URi
 			if (InGraph->RuntimeAST.IsValid())
 			{
 				URigVMPin* RootPin = CastChecked<URigVMPin>(InSubject)->GetRootPin();
-				const FRigVMExprAST* Expression = InGraph->GetRuntimeAST()->GetExprForSubject(RootPin);
+				FRigVMASTProxy RootPinProxy = FRigVMASTProxy::MakeFromUObject(RootPin);
+				const FRigVMExprAST* Expression = InGraph->GetRuntimeAST()->GetExprForSubject(RootPinProxy);
 				if (Expression == nullptr)
 				{
 					InGraph->ClearAST();
diff --git a/Engine/Source/Developer/RigVMDeveloper/Private/RigVMModel/RigVMPin.cpp b/Engine/Source/Developer/RigVMDeveloper/Private/RigVMModel/RigVMPin.cpp
--- a/Engine/Source/Developer/RigVMDeveloper/Private/RigVMModel/RigVMPin.cpp
+++ b/Engine/Source/Developer/RigVMDeveloper/Private/RigVMModel/RigVMPin.cpp
@@ -366,20 +366,12 @@ bool URigVMPin::IsExecuteContext() const
 
 FString URigVMPin::GetDefaultValue() const
 {
-	return GetDefaultValue(FDefaultValueOverride());
+	return GetDefaultValue(FDefaultValueOverride(), FRigVMASTProxy());
 }
 
-FString URigVMPin::GetDefaultValue(const FDefaultValueOverride& InDefaultValueOverride) const
+FString URigVMPin::GetDefaultValue(const FDefaultValueOverride& InDefaultValueOverride, const FRigVMASTProxy& InProxy) const
 {
-	/*
-	URigVMPin* PinForDefaultValue = URigVMCompiler::FollowPinForDefaultValue((URigVMPin*)this);
-	if (PinForDefaultValue != this)
-	{
-		return PinForDefaultValue->GetDefaultValue(InDefaultValueOverride);
-	}
-	*/
-
-	if (const FString* Override = InDefaultValueOverride.Find(this))
+	if (const FString* Override = InDefaultValueOverride.Find(InProxy.GetSibling((URigVMPin*)this)))
 	{
 		return *Override;
 	}
@@ -391,7 +383,7 @@ FString URigVMPin::GetDefaultValue(const FDefaultValueOverride& InDefaultValueOv
 			TArray<FString> ElementDefaultValues;
 			for (URigVMPin* SubPin : SubPins)
 			{
-				FString ElementDefaultValue = SubPin->GetDefaultValue(InDefaultValueOverride);
+				FString ElementDefaultValue = SubPin->GetDefaultValue(InDefaultValueOverride, InProxy);
 				if (SubPin->IsStringType())
 				{
 					ElementDefaultValue = TEXT("\"") + ElementDefaultValue + TEXT("\"");
@@ -418,7 +410,7 @@ FString URigVMPin::GetDefaultValue(const FDefaultValueOverride& InDefaultValueOv
 			TArray<FString> MemberDefaultValues;
 			for (URigVMPin* SubPin : SubPins)
 			{
-				FString MemberDefaultValue = SubPin->GetDefaultValue(InDefaultValueOverride);
+				FString MemberDefaultValue = SubPin->GetDefaultValue(InDefaultValueOverride, InProxy);
 				if (SubPin->IsStringType())
 				{
 					MemberDefaultValue = TEXT("\"") + MemberDefaultValue + TEXT("\"");
diff --git a/Engine/Source/Developer/RigVMDeveloper/Public/RigVMCompiler/RigVMAST.h b/Engine/Source/Developer/RigVMDeveloper/Public/RigVMCompiler/RigVMAST.h
--- a/Engine/Source/Developer/RigVMDeveloper/Public/RigVMCompiler/RigVMAST.h
+++ b/Engine/Source/Developer/RigVMDeveloper/Public/RigVMCompiler/RigVMAST.h
@@ -6,6 +6,8 @@
 #include "RigVMCore/RigVMRegistry.h"
 #include "RigVMModel/RigVMPin.h"
 #include "RigVMModel/RigVMLink.h"
+#include "RigVMModel/RigVMNode.h"
+#include "RigVMASTProxy.h"
 #include "RigVMAST.generated.h"
 
 class FRigVMParserAST;
@@ -30,8 +32,6 @@ class URigVMNode;
 class URigVMGraph;
 class URigVMController;
 
-typedef TPair<URigVMPin*, URigVMPin*> FRigVMPinPair;
-
 /*
  * Base class for an expression within an abstract syntax tree.
  * The base implements parent / child relationships as well as a
@@ -208,7 +208,7 @@ public:
 protected:
 
 	// default constructor (protected so that only parser can access it)
-	FRigVMExprAST(EType InType = EType::Invalid, UObject* InSubject = nullptr);
+	FRigVMExprAST(EType InType = EType::Invalid, const FRigVMASTProxy& InProxy = FRigVMASTProxy());
 
 	// adds a parent to this expression
 	// this in consequence also adds this as a child to the parent
@@ -460,8 +460,8 @@ public:
 protected:
 
 	// default constructor (protected so that only parser can access it)
-	FRigVMBlockExprAST(EType InType = EType::Block, UObject* InSubject = nullptr)
-		: FRigVMExprAST(InType, InSubject)
+	FRigVMBlockExprAST(EType InType = EType::Block, const FRigVMASTProxy& InProxy = FRigVMASTProxy())
+		: FRigVMExprAST(InType, InProxy)
 		, bIsObsolete(false)
 	{}
 
@@ -486,16 +486,19 @@ public:
 	// disable copy constructor
 	FRigVMNodeExprAST(const FRigVMNodeExprAST&) = delete;
 
+	// returns the proxy this expression is using
+	const FRigVMASTProxy& GetProxy() const { return Proxy; }
+
 	// returns the node from the model this expression is referencing
 	// @return the node from the model this expression is referencing
-	URigVMNode* GetNode() const { return Node; }
+	URigVMNode* GetNode() const { return GetProxy().GetSubjectChecked<URigVMNode>(); }
 
 	virtual bool IsConstant() const override;
 
 protected:
 
 	// default constructor (protected so that only parser can access it)
-	FRigVMNodeExprAST(EType InType, UObject* InSubject = nullptr);
+	FRigVMNodeExprAST(EType InType, const FRigVMASTProxy& InNodeProxy);
 
 	// overload of the type checking mechanism
 	virtual bool IsA(EType InType) const override
@@ -505,7 +508,7 @@ protected:
 
 private:
 
-	URigVMNode* Node;
+	FRigVMASTProxy Proxy;
 	friend class FRigVMParserAST;
 };
 
@@ -541,8 +544,8 @@ public:
 protected:
 
 	// default constructor (protected so that only parser can access it)
-	FRigVMEntryExprAST(UObject* InSubject = nullptr)
-		: FRigVMNodeExprAST(EType::Entry, InSubject)
+	FRigVMEntryExprAST(const FRigVMASTProxy& InNodeProxy)
+		: FRigVMNodeExprAST(EType::Entry, InNodeProxy)
 	{}
 
 private:
@@ -578,8 +581,8 @@ public:
 protected:
 
 	// default constructor (protected so that only parser can access it)
-	FRigVMCallExternExprAST(UObject* InSubject = nullptr)
-		: FRigVMNodeExprAST(EType::CallExtern, InSubject)
+	FRigVMCallExternExprAST(const FRigVMASTProxy& InNodeProxy)
+		: FRigVMNodeExprAST(EType::CallExtern, InNodeProxy)
 	{}
 
 private:
@@ -613,8 +616,8 @@ public:
 protected:
 
 	// default constructor (protected so that only parser can access it)
-	FRigVMNoOpExprAST(UObject* InSubject = nullptr)
-		: FRigVMNodeExprAST(EType::NoOp, InSubject)
+	FRigVMNoOpExprAST(const FRigVMASTProxy& InNodeProxy)
+		: FRigVMNodeExprAST(EType::NoOp, InNodeProxy)
 	{}
 
 private:
@@ -646,9 +649,12 @@ public:
 
 	virtual bool IsConstant() const override;
 
+	// returns the proxy this expression is using
+	const FRigVMASTProxy& GetProxy() const { return Proxy; }
+
 	// returns the pin in the model this variable is representing
 	// @return the pin in the model this variable is representing
-	URigVMPin* GetPin() const { return Pin; }
+	URigVMPin* GetPin() const { return GetProxy().GetSubjectChecked<URigVMPin>(); }
 
 	// returns the C++ data type of this variable
 	// @return the C++ data type of this variable
@@ -687,16 +693,15 @@ public:
 protected:
 
 	// default constructor (protected so that only parser can access it)
-	FRigVMVarExprAST(EType InType = EType::Var, UObject* InSubject = nullptr)
+	FRigVMVarExprAST(EType InType, const FRigVMASTProxy& InPinProxy)
 		: FRigVMExprAST(InType)
-		, Pin(Cast<URigVMPin>(InSubject))
+		, Proxy(InPinProxy)
 	{
-		check(Pin);
 	}
 
 private:
 
-	URigVMPin* Pin;
+	FRigVMASTProxy Proxy;
 
 	friend class FRigVMParserAST;
 	friend class URigVMCompiler;
@@ -736,8 +741,8 @@ public:
 protected:
 
 	// default constructor (protected so that only parser can access it)
-	FRigVMLiteralExprAST(UObject* InSubject = nullptr)
-		: FRigVMVarExprAST(EType::Literal, InSubject)
+	FRigVMLiteralExprAST(const FRigVMASTProxy& InPinProxy)
+		: FRigVMVarExprAST(EType::Literal, InPinProxy)
 	{}
 
 private:
@@ -772,8 +777,8 @@ public:
 protected:
 
 	// default constructor (protected so that only parser can access it)
-	FRigVMExternalVarExprAST(UObject* InSubject = nullptr)
-		: FRigVMVarExprAST(EType::ExternalVar, InSubject)
+	FRigVMExternalVarExprAST(const FRigVMASTProxy& InPinProxy)
+		: FRigVMVarExprAST(EType::ExternalVar, InPinProxy)
 	{}
 
 private:
@@ -798,13 +803,19 @@ public:
 	// disable copy constructor
 	FRigVMAssignExprAST(const FRigVMAssignExprAST&) = delete;
 
+	// returns the source proxy this expression is using
+	const FRigVMASTProxy& GetSourceProxy() const { return SourceProxy; }
+
 	// returns the source pin for this assignment
 	// @return the source pin for this assignment
-	URigVMPin* GetSourcePin() const { return SourcePin; }
+	URigVMPin* GetSourcePin() const { return SourceProxy.GetSubjectChecked<URigVMPin>(); }
+
+	// returns the target proxy this expression is using
+	const FRigVMASTProxy& GetTargetProxy() const { return TargetProxy; }
 
 	// returns the target pin for this assignment
 	// @return the target pin for this assignment
-	URigVMPin* GetTargetPin() const { return TargetPin; }
+	URigVMPin* GetTargetPin() const { return TargetProxy.GetSubjectChecked<URigVMPin>(); }
 
 	// overload of the type checking mechanism
 	virtual bool IsA(EType InType) const override
@@ -815,19 +826,17 @@ public:
 protected:
 
 	// default constructor (protected so that only parser can access it)
-	FRigVMAssignExprAST(EType InType, UObject* InSubjectA, UObject* InSubjectB)
+	FRigVMAssignExprAST(EType InType, const FRigVMASTProxy& InSourceProxy, const FRigVMASTProxy& InTargetProxy)
 		: FRigVMExprAST(InType)
-		, SourcePin(Cast<URigVMPin>(InSubjectA))
-		, TargetPin(Cast<URigVMPin>(InSubjectB))
+		, SourceProxy(InSourceProxy)
+		, TargetProxy(InTargetProxy)
 	{
-		check(SourcePin);
-		check(TargetPin);
 	}
 
 private:
 
-	URigVMPin* SourcePin;
-	URigVMPin* TargetPin;
+	FRigVMASTProxy SourceProxy;
+	FRigVMASTProxy TargetProxy;
 	friend class FRigVMParserAST;
 };
 
@@ -862,8 +871,8 @@ public:
 protected:
 
 	// default constructor (protected so that only parser can access it)
-	FRigVMCopyExprAST(UObject* InSubjectA, UObject* InSubjectB)
-		: FRigVMAssignExprAST(EType::Copy, InSubjectA, InSubjectB)
+	FRigVMCopyExprAST(const FRigVMASTProxy& InSourceProxy, const FRigVMASTProxy& InTargetProxy)
+		: FRigVMAssignExprAST(EType::Copy, InSourceProxy, InTargetProxy)
 	{}
 
 
@@ -905,8 +914,8 @@ public:
 protected:
 
 	// default constructor (protected so that only parser can access it)
-	FRigVMCachedValueExprAST(UObject* InSubject = nullptr)
-		: FRigVMExprAST(EType::CachedValue, InSubject)
+	FRigVMCachedValueExprAST(const FRigVMASTProxy& InProxy)
+		: FRigVMExprAST(EType::CachedValue, InProxy)
 	{}
 
 
@@ -943,8 +952,8 @@ public:
 protected:
 
 	// default constructor (protected so that only parser can access it)
-	FRigVMExitExprAST(UObject* InSubject = nullptr)
-		: FRigVMExprAST(EType::Exit, InSubject)
+	FRigVMExitExprAST(const FRigVMASTProxy& InProxy)
+		: FRigVMExprAST(EType::Exit, InProxy)
 	{}
 
 private:
@@ -985,8 +994,8 @@ public:
 protected:
 
 	// default constructor (protected so that only parser can access it)
-	FRigVMBranchExprAST(UObject* InSubject = nullptr)
-		: FRigVMNodeExprAST(EType::Branch, InSubject)
+	FRigVMBranchExprAST(const FRigVMASTProxy& InProxy)
+		: FRigVMNodeExprAST(EType::Branch, InProxy)
 	{}
 
 
@@ -1028,8 +1037,8 @@ public:
 protected:
 
 	// default constructor (protected so that only parser can access it)
-	FRigVMIfExprAST(UObject* InSubject = nullptr)
-		: FRigVMNodeExprAST(EType::If, InSubject)
+	FRigVMIfExprAST(const FRigVMASTProxy& InProxy)
+		: FRigVMNodeExprAST(EType::If, InProxy)
 	{}
 
 
@@ -1070,8 +1079,8 @@ public:
 protected:
 
 	// default constructor (protected so that only parser can access it)
-	FRigVMSelectExprAST(UObject* InSubject = nullptr)
-		: FRigVMNodeExprAST(EType::Select, InSubject)
+	FRigVMSelectExprAST(const FRigVMASTProxy& InProxy)
+		: FRigVMNodeExprAST(EType::Select, InProxy)
 	{}
 
 
@@ -1168,7 +1177,7 @@ public:
 	// returns the expression for a given subject. subjects include nodes and pins.
 	// @param InSubject the subject to retrieve the expression for (node or pin)
 	// @return the expressoin for the given subject (or nullptr)
-	const FRigVMExprAST* GetExprForSubject(UObject* InSubject);
+	const FRigVMExprAST* GetExprForSubject(const FRigVMASTProxy& InProxy) const;
 
 	// Prepares the parser for cycle checking on a given pin.
 	// This marks up the parents and childen of the corresponding expression in the graph,
@@ -1203,13 +1212,23 @@ public:
 private:
 
 	// private constructor for a partial build
-	FRigVMParserAST(URigVMGraph* InGraph, const TArray<URigVMNode*>& InNodesToCompute);
+	FRigVMParserAST(URigVMGraph* InGraph, const TArray<FRigVMASTProxy>& InNodesToCompute);
+
+	// make function to create an expression
+	template<class ExprType>
+	ExprType* MakeExpr(FRigVMExprAST::EType InType, const FRigVMASTProxy& InProxy)
+	{
+		ExprType* Expr = new ExprType(InType, InProxy);
+		Expr->ParserPtr = this;
+		Expr->Index = Expressions.Add(Expr);
+		return Expr;
+	}
 
 	// make function to create an expression
 	template<class ExprType>
-	ExprType* MakeExpr(FRigVMExprAST::EType InType, UObject* InSubject = nullptr)
+	ExprType* MakeExpr(FRigVMExprAST::EType InType, const FRigVMASTProxy& InProxyA, const FRigVMASTProxy& InProxyB)
 	{
-		ExprType* Expr = new ExprType(InType, InSubject);
+		ExprType* Expr = new ExprType(InType, InProxyA, InProxyB);
 		Expr->ParserPtr = this;
 		Expr->Index = Expressions.Add(Expr);
 		return Expr;
@@ -1217,9 +1236,9 @@ private:
 
 	// make function to create an expression
 	template<class ExprType>
-	ExprType* MakeExpr(FRigVMExprAST::EType InType, UObject* InSubjectA, UObject* InSubjectB)
+	ExprType* MakeExpr(const FRigVMASTProxy& InProxyA, const FRigVMASTProxy& InProxyBB)
 	{
-		ExprType* Expr = new ExprType(InType, InSubjectA, InSubjectB);
+		ExprType* Expr = new ExprType(InProxyA, InProxyBB);
 		Expr->ParserPtr = this;
 		Expr->Index = Expressions.Add(Expr);
 		return Expr;
@@ -1227,9 +1246,9 @@ private:
 
 	// make function to create an expression
 	template<class ExprType>
-	ExprType* MakeExpr(UObject* InSubjectA, UObject* InSubjectB)
+	ExprType* MakeExpr(const FRigVMASTProxy& InProxy)
 	{
-		ExprType* Expr = new ExprType(InSubjectA, InSubjectB);
+		ExprType* Expr = new ExprType(InProxy);
 		Expr->ParserPtr = this;
 		Expr->Index = Expressions.Add(Expr);
 		return Expr;
@@ -1237,9 +1256,9 @@ private:
 
 	// make function to create an expression
 	template<class ExprType>
-	ExprType* MakeExpr(UObject* InSubject = nullptr)
+	ExprType* MakeExpr()
 	{
-		ExprType* Expr = new ExprType(InSubject);
+		ExprType* Expr = new ExprType();
 		Expr->ParserPtr = this;
 		Expr->Index = Expressions.Add(Expr);
 		return Expr;
@@ -1300,46 +1319,46 @@ private:
 
 	// helper function to inline all contributing nodes of the graph
 	void Inline(URigVMGraph* InGraph);
-	void Inline(URigVMGraph* InGraph, const TArray<URigVMNode*>& InNodes);
+	void Inline(URigVMGraph* InGraph, const TArray<FRigVMASTProxy>& InNodeProxies);
 
 	// helper functions to retrieve links for a given pin
-	const TArray<URigVMPin*>& GetSourcePins(URigVMPin* InPin) const;
-	const TArray<URigVMPin*>& GetTargetPins(URigVMPin* InPin) const;
-	TArray<FRigVMPinPair> GetSourceLinks(URigVMPin* InPin, bool bRecursive = false) const;
-	TArray<FRigVMPinPair> GetTargetLinks(URigVMPin* InPin, bool bRecursive = false) const;
+	const TArray<FRigVMASTProxy>& GetSourcePins(const FRigVMASTProxy& InPinProxy) const;
+	const TArray<FRigVMASTProxy>& GetTargetPins(const FRigVMASTProxy& InPinProxy) const;
+	TArray<FRigVMPinProxyPair> GetSourceLinks(const FRigVMASTProxy& InPinProxy, bool bRecursive = false) const;
+	TArray<FRigVMPinProxyPair> GetTargetLinks(const FRigVMASTProxy& InPinProxy, bool bRecursive = false) const;
 
 	// traverse a single mutable node (constructs entry, call extern and other expressions)
-	FRigVMExprAST* TraverseMutableNode(URigVMNode* InNode, FRigVMExprAST* InParentExpr);
+	FRigVMExprAST* TraverseMutableNode(const FRigVMASTProxy& InNodeProxy, FRigVMExprAST* InParentExpr);
 
 	// traverse a single pure node (constructs call extern expressions)
-	FRigVMExprAST* TraverseNode(URigVMNode* InNode, FRigVMExprAST* InParentExpr);
+	FRigVMExprAST* TraverseNode(const FRigVMASTProxy& InNodeProxy, FRigVMExprAST* InParentExpr);
 
 	// returns the expression for a given node
-	FRigVMExprAST* CreateExpressionForNode(URigVMNode* InNode, FRigVMExprAST* InParentExpr);
+	FRigVMExprAST* CreateExpressionForNode(const FRigVMASTProxy& InNodeProxy, FRigVMExprAST* InParentExpr);
 
 	// traverse an array of pins for a given node
-	TArray<FRigVMExprAST*> TraversePins(URigVMNode* InNode, FRigVMExprAST* InParentExpr);
+	TArray<FRigVMExprAST*> TraversePins(const FRigVMASTProxy& InNodeProxy, FRigVMExprAST* InParentExpr);
 
 	// traverse a single pin (constructs var + literal expressions)
-	FRigVMExprAST* TraversePin(URigVMPin* InPin, FRigVMExprAST* InParentExpr);
+	FRigVMExprAST* TraversePin(const FRigVMASTProxy& InPinProxy, FRigVMExprAST* InParentExpr);
 
 	// traverse a single link (constructs assign + copy expressions)
-	FRigVMExprAST* TraverseLink(const FRigVMPinPair& InLink, FRigVMExprAST* InParentExpr);
+	FRigVMExprAST* TraverseLink(const FRigVMPinProxyPair& InLink, FRigVMExprAST* InParentExpr);
 
-	bool ShouldLinkBeSkipped(const FRigVMPinPair& InLink) const;
+	bool ShouldLinkBeSkipped(const FRigVMPinProxyPair& InLink) const;
 
-	static FString GetLinkAsString(const FRigVMPinPair& InLink);
+	static FString GetLinkAsString(const FRigVMPinProxyPair& InLink);
 
-	TMap<UObject*, FRigVMExprAST*> SubjectToExpression;
-	TMap<UObject*, int32> NodeExpressionIndex;
+	TMap<FRigVMASTProxy, FRigVMExprAST*> SubjectToExpression;
+	TMap<FRigVMASTProxy, int32> NodeExpressionIndex;
 	TArray<FRigVMExprAST*> Expressions;
 	TArray<FRigVMExprAST*> RootExpressions;
 	FRigVMBlockExprAST* ObsoleteBlock;
 
-	TArray<URigVMNode*> Nodes;
-	TMap<URigVMPin*, TArray<URigVMPin*>> TargetLinks;
-	TMap<URigVMPin*, TArray<URigVMPin*>> SourceLinks;
-	static const TArray<URigVMPin*> EmptyLinks;
+	TArray<FRigVMASTProxy> NodeProxies;
+	TMap<FRigVMASTProxy, TArray<FRigVMASTProxy>> TargetLinks;
+	TMap<FRigVMASTProxy, TArray<FRigVMASTProxy>> SourceLinks;
+	static const TArray<FRigVMASTProxy> EmptyProxyArray;
 	URigVMPin::FDefaultValueOverride PinDefaultValueOverrides;
 
 	enum ETraverseRelationShip
diff --git a/Engine/Source/Developer/RigVMDeveloper/Public/RigVMCompiler/RigVMASTProxy.h b/Engine/Source/Developer/RigVMDeveloper/Public/RigVMCompiler/RigVMASTProxy.h
new file mode 100644
--- /dev/null
+++ b/Engine/Source/Developer/RigVMDeveloper/Public/RigVMCompiler/RigVMASTProxy.h
@@ -0,0 +1,224 @@
+// Copyright Epic Games, Inc. All Rights Reserved.
+
+#pragma once
+
+#include "CoreMinimal.h"
+
+/*
+ * A stack describing the model object which led to this invocation.
+ * Objects on the stack can be URigVMGraph or URigVMNode.
+ */
+class RIGVMDEVELOPER_API FRigVMCallstack
+{
+public:
+
+	// Returns the path
+	FString GetCallPath(bool bIncludeLast = true) const;
+
+	// Returns the number of entries on this callstack
+	int32 Num() const;
+
+	// Returns the element at a given index
+	const UObject* operator[](int32 InIndex) const;
+
+	// Returns true if the stack contains a given entry
+	bool Contains(UObject* InEntry) const;
+
+	friend FORCEINLINE uint32 GetTypeHash(const FRigVMCallstack& Callstack)
+	{
+		return Callstack.GetEntryTypeHash(Callstack.Num() - 1);
+	}
+
+	FORCEINLINE bool operator ==(const FRigVMCallstack& Other) const
+	{
+		if (Num() == Other.Num())
+		{
+			for (int32 Index = 0; Index < Num(); Index++)
+			{
+				if (operator[](Index) != Other[Index])
+				{
+					return false;
+				}
+			}
+			return true;
+		}
+		return false;
+	}
+
+	FORCEINLINE bool operator !=(const FRigVMCallstack& Other) const
+	{
+		return !(*this == Other);
+	}
+
+	FORCEINLINE bool operator <(const FRigVMCallstack& Other) const
+	{
+		if (Num() < Other.Num())
+		{
+			return true;
+		}
+
+		if (Num() > Other.Num())
+		{
+			return false;
+		}
+
+		if (Num() == 0)
+		{
+			return false;
+		}
+
+		return operator[](0) < Other[0];
+	}
+
+	FORCEINLINE bool operator >(const FRigVMCallstack& Other) const
+	{
+		if (Num() > Other.Num())
+		{
+			return true;
+		}
+
+		if (Num() < Other.Num())
+		{
+			return false;
+		}
+
+		if (Num() == 0)
+		{
+			return true;
+		}
+
+		return operator[](0) > Other[0];
+	}
+
+private:
+
+	FORCEINLINE uint32 GetEntryTypeHash(int32 Index) const
+	{
+		if (Index < 0)
+		{
+			return 0;
+		}
+		if (Index == 0)
+		{
+			return GetTypeHash(Stack[Index]);
+		}
+		return HashCombine(GetEntryTypeHash(Index - 1), GetTypeHash(Stack[Index]));
+	}
+
+	TArray<UObject*> Stack;
+
+	friend class FRigVMParserAST;
+	friend class FRigVMCallstackGuard;
+	friend class FRigVMASTProxy;
+};
+
+/*
+ * A proxy which describes an occurence of a subject within
+ * a graph. The subject can be a pin, a node, a link, etc.
+ * The context is the callstack under which is occured.
+ */
+class RIGVMDEVELOPER_API FRigVMASTProxy
+{
+public:
+
+	FRigVMASTProxy()
+	{
+#if UE_BUILD_DEBUG
+		DebugName = FString();
+#endif
+		Callstack.Stack.Push(nullptr);
+	}
+
+	FRigVMASTProxy(const FRigVMASTProxy& InOther)
+	{
+		Callstack.Stack = InOther.Callstack.Stack;
+#if UE_BUILD_DEBUG
+		DebugName = Callstack.GetCallPath();
+#endif
+	}
+
+	static FRigVMASTProxy MakeFromUObject(UObject* InSubject);
+
+	FRigVMASTProxy GetSibling(UObject* InSubject) const
+	{
+		FRigVMASTProxy Sibling = *this;
+		Sibling.Callstack.Stack.Last() = InSubject;
+#if UE_BUILD_DEBUG
+		Sibling.DebugName = Sibling.Callstack.GetCallPath();
+#endif
+		return Sibling;
+	}
+
+	FRigVMASTProxy GetParent() const
+	{
+		ensure(Callstack.Num() > 1);
+		
+		FRigVMASTProxy Parent = *this;
+		Parent.Callstack.Stack.Pop();
+#if UE_BUILD_DEBUG
+		Parent.DebugName = Parent.Callstack.GetCallPath();
+#endif
+		return Parent;
+	}
+
+	FRigVMASTProxy GetChild(UObject* InSubject) const
+	{
+		FRigVMASTProxy Child = *this;
+		Child.Callstack.Stack.Push(InSubject);
+#if UE_BUILD_DEBUG
+		Child.DebugName = Child.Callstack.GetCallPath();
+#endif
+		return Child;
+	}
+
+	const FRigVMCallstack& GetCallstack() const { return Callstack; }
+
+	UObject* GetSubject() const { return Callstack.Stack.Last(); }
+
+	template<class T>
+	T* GetSubject() const { return Cast<T>(GetSubject()); }
+
+	template<class T>
+	T* GetSubjectChecked() const { return CastChecked<T>(GetSubject()); }
+
+	bool IsValid() const { return Callstack.Num() > 0 && GetSubject() != nullptr; }
+
+	template<class T>
+	bool IsA() const { return GetSubject() ? GetSubject()->IsA<T>() : false; }
+
+	friend FORCEINLINE uint32 GetTypeHash(const FRigVMASTProxy& Proxy)
+	{
+		return GetTypeHash(Proxy.GetCallstack());
+	}
+
+	FORCEINLINE bool operator ==(const FRigVMASTProxy& Other) const
+	{
+		return GetCallstack() == Other.GetCallstack();
+	}
+
+	FORCEINLINE bool operator !=(const FRigVMASTProxy& Other) const
+	{
+		return GetCallstack() != Other.GetCallstack();
+	}
+
+	FORCEINLINE bool operator <(const FRigVMASTProxy& Other) const
+	{
+		return GetCallstack() < Other.GetCallstack();
+	}
+
+	FORCEINLINE bool operator >(const FRigVMASTProxy& Other) const
+	{
+		return GetCallstack() > Other.GetCallstack();
+	}
+
+private:
+
+#if UE_BUILD_DEBUG
+	FString DebugName;
+#endif
+
+	FRigVMCallstack Callstack;
+};
+
+typedef TPair<FRigVMASTProxy, FRigVMASTProxy> FRigVMPinProxyPair;
+typedef TMap<FRigVMASTProxy, FString> FRigVMPinDefaultValueOverride;
diff --git a/Engine/Source/Developer/RigVMDeveloper/Public/RigVMModel/Nodes/RigVMEnumNode.h b/Engine/Source/Developer/RigVMDeveloper/Public/RigVMModel/Nodes/RigVMEnumNode.h
--- a/Engine/Source/Developer/RigVMDeveloper/Public/RigVMModel/Nodes/RigVMEnumNode.h
+++ b/Engine/Source/Developer/RigVMDeveloper/Public/RigVMModel/Nodes/RigVMEnumNode.h
@@ -34,7 +34,7 @@ public:
 	UObject* GetCPPTypeObject() const;
 
 	// Returns the default value of the parameter as a string
-	FString GetDefaultValue(const URigVMPin::FDefaultValueOverride& InDefaultValueOverride = URigVMPin::FDefaultValueOverride()) const;
+	FString GetDefaultValue(const URigVMPin::FDefaultValueOverride& InDefaultValueOverride = URigVMPin::FDefaultValueOverride(), const FRigVMASTProxy& InProxy = FRigVMASTProxy()) const;
 
 	// Override of node title
 	virtual FLinearColor GetNodeColor() const override { return FLinearColor::Blue; }
diff --git a/Engine/Source/Developer/RigVMDeveloper/Public/RigVMModel/RigVMPin.h b/Engine/Source/Developer/RigVMDeveloper/Public/RigVMModel/RigVMPin.h
--- a/Engine/Source/Developer/RigVMDeveloper/Public/RigVMModel/RigVMPin.h
+++ b/Engine/Source/Developer/RigVMDeveloper/Public/RigVMModel/RigVMPin.h
@@ -5,6 +5,7 @@
 #include "CoreMinimal.h"
 #include "UObject/ObjectMacros.h"
 #include "RigVMCore/RigVMRegistry.h"
+#include "RigVMCompiler/RigVMASTProxy.h"
 #include "RigVMPin.generated.h"
 
 class URigVMGraph;
@@ -78,7 +79,7 @@ class RIGVMDEVELOPER_API URigVMPin : public UObject
 public:
 
 	// A map used to override pin default values
-	typedef TMap<URigVMPin*, FString> FDefaultValueOverride;
+	typedef TMap<FRigVMASTProxy, FString> FDefaultValueOverride;
 
 	// Splits a PinPath at the start, so for example "Node.Color.R" becomes "Node" and "Color.R"
 	static bool SplitPinPathAtStart(const FString& InPinPath, FString& LeftMost, FString& Right);
@@ -188,7 +189,7 @@ public:
 	FString GetDefaultValue() const;
 
 	// Returns the default value with an additional override ma
-	FString GetDefaultValue(const FDefaultValueOverride& InDefaultValueOverride) const;
+	FString GetDefaultValue(const FDefaultValueOverride& InDefaultValueOverride, const FRigVMASTProxy& InProxy) const;
 
 	// Returns true if the default value provided is valid
 	bool IsValidDefaultValue(const FString& InDefaultValue) const;
