Fix uninitialized struct members

diff --git a/Engine/Plugins/Animation/IKRig/Source/IKRig/Public/AnimNodes/AnimNode_IKRig.h b/Engine/Plugins/Animation/IKRig/Source/IKRig/Public/AnimNodes/AnimNode_IKRig.h
--- a/Engine/Plugins/Animation/IKRig/Source/IKRig/Public/AnimNodes/AnimNode_IKRig.h
+++ b/Engine/Plugins/Animation/IKRig/Source/IKRig/Public/AnimNodes/AnimNode_IKRig.h
@@ -31,7 +31,7 @@ struct IKRIG_API FAnimNode_IKRig : public FAnimNode_Base
 
 	/** For every evaluation, we'll start with refpose, this is safer if you have joints that don't exists in the target */
 	UPROPERTY(EditAnywhere, Category = Solver)
-	bool bStartFromRefPose;
+	bool bStartFromRefPose = false;
 
 #if WITH_EDITORONLY_DATA
 	/** Toggle drawing of axes to debug joint rotation*/
diff --git a/Engine/Plugins/Animation/IKRig/Source/IKRig/Public/IKRigDataTypes.h b/Engine/Plugins/Animation/IKRig/Source/IKRig/Public/IKRigDataTypes.h
--- a/Engine/Plugins/Animation/IKRig/Source/IKRig/Public/IKRigDataTypes.h
+++ b/Engine/Plugins/Animation/IKRig/Source/IKRig/Public/IKRigDataTypes.h
@@ -96,7 +96,7 @@ struct IKRIG_API FIKRigEffector
 	GENERATED_BODY()
 
 	UPROPERTY()
-	FGuid Guid;
+	FGuid Guid = FGuid();
 
 	UPROPERTY(EditAnywhere, Category = FIKRigEffector)
 	FName Bone;
diff --git a/Engine/Source/Developer/RigVMDeveloper/Public/RigVMModel/RigVMControllerActions.h b/Engine/Source/Developer/RigVMDeveloper/Public/RigVMModel/RigVMControllerActions.h
--- a/Engine/Source/Developer/RigVMDeveloper/Public/RigVMModel/RigVMControllerActions.h
+++ b/Engine/Source/Developer/RigVMDeveloper/Public/RigVMModel/RigVMControllerActions.h
@@ -1284,7 +1284,7 @@ public:
 	FString PinName;
 	
 	UPROPERTY()
-	ERigVMPinDirection Direction;
+	ERigVMPinDirection Direction = ERigVMPinDirection::Input;
 
 	UPROPERTY()
 	FString CPPType;
@@ -1316,7 +1316,7 @@ public:
 	FString PinName;
 	
 	UPROPERTY()
-	ERigVMPinDirection Direction;
+	ERigVMPinDirection Direction = ERigVMPinDirection::Input;
 
 	UPROPERTY()
 	FString CPPType;
diff --git a/Engine/Source/Runtime/Engine/Classes/Animation/AnimNode_AssetPlayerBase.h b/Engine/Source/Runtime/Engine/Classes/Animation/AnimNode_AssetPlayerBase.h
--- a/Engine/Source/Runtime/Engine/Classes/Animation/AnimNode_AssetPlayerBase.h
+++ b/Engine/Source/Runtime/Engine/Classes/Animation/AnimNode_AssetPlayerBase.h
@@ -61,7 +61,7 @@ struct ENGINE_API FAnimNode_AssetPlayerBase : public FAnimNode_Base
 
 	// How synchronization is determined
 	UPROPERTY()
-	EAnimSyncMethod Method;
+	EAnimSyncMethod Method = EAnimSyncMethod::DoNotSync;
 
 	/** If true, "Relevant anim" nodes that look for the highest weighted animation in a state will ignore
 	 *  this node
