CIS: Fix missing includes in RigVMASTProxy.h

diff --git a/Engine/Source/Developer/RigVMDeveloper/Public/RigVMCompiler/RigVMASTProxy.h b/Engine/Source/Developer/RigVMDeveloper/Public/RigVMCompiler/RigVMASTProxy.h
--- a/Engine/Source/Developer/RigVMDeveloper/Public/RigVMCompiler/RigVMASTProxy.h
+++ b/Engine/Source/Developer/RigVMDeveloper/Public/RigVMCompiler/RigVMASTProxy.h
@@ -3,6 +3,8 @@
 #pragma once
 
 #include "CoreMinimal.h"
+#include "Templates/Casts.h"
+#include "UObject/UObjectBaseUtility.h"
 
 /*
  * A stack describing the model object which led to this invocation.
