Control Rig: Special case ExecuteContext for Exposed Pin actions in Controller

diff --git a/Engine/Source/Developer/RigVMDeveloper/Private/RigVMModel/RigVMController.cpp b/Engine/Source/Developer/RigVMDeveloper/Private/RigVMModel/RigVMController.cpp
--- a/Engine/Source/Developer/RigVMDeveloper/Private/RigVMModel/RigVMController.cpp
+++ b/Engine/Source/Developer/RigVMDeveloper/Private/RigVMModel/RigVMController.cpp
@@ -5094,6 +5094,38 @@ FName URigVMController::AddExposedPin(const FName& InPinName, ERigVMPinDirection
 	URigVMLibraryNode* LibraryNode = Cast<URigVMLibraryNode>(Graph->GetOuter());
 	check(LibraryNode);
 
+	UObject* CPPTypeObject = nullptr;
+	if (!InCPPTypeObjectPath.IsNone())
+	{
+		if (CPPTypeObject == nullptr)
+		{
+			CPPTypeObject = URigVMCompiler::GetScriptStructForCPPType(InCPPTypeObjectPath.ToString());
+		}
+		if (CPPTypeObject == nullptr)
+		{
+			CPPTypeObject = URigVMPin::FindObjectFromCPPTypeObjectPath<UObject>(InCPPTypeObjectPath.ToString());
+		}
+	}
+
+	if (CPPTypeObject)
+	{
+		if (UScriptStruct* ScriptStruct = Cast<UScriptStruct>(CPPTypeObject))
+		{
+			if (ScriptStruct->IsChildOf(FRigVMExecuteContext::StaticStruct()))
+			{
+				for (URigVMPin* ExistingPin : LibraryNode->Pins)
+				{
+					if (ExistingPin->IsExecuteContext())
+					{
+						return NAME_None;
+					}
+				}
+
+				InDirection = ERigVMPinDirection::IO;
+			}
+		}
+	}
+
 	FName PinName = GetUniqueName(InPinName, [LibraryNode](const FName& InName) {
 
 		return LibraryNode->FindPin(InName.ToString()) == nullptr;
