Control Rig: Support for expanding function nodes

diff --git a/Engine/Source/Developer/RigVMDeveloper/Private/RigVMModel/RigVMController.cpp b/Engine/Source/Developer/RigVMDeveloper/Private/RigVMModel/RigVMController.cpp
--- a/Engine/Source/Developer/RigVMDeveloper/Private/RigVMModel/RigVMController.cpp
+++ b/Engine/Source/Developer/RigVMDeveloper/Private/RigVMModel/RigVMController.cpp
@@ -1900,7 +1900,7 @@ TArray<URigVMNode*> URigVMController::ExpandLibraryNode(const FName& InNodeName,
 		return TArray<URigVMNode*>();
 	}
 
-	URigVMLibraryNode* LibNode = Cast<URigVMCollapseNode>(Node);
+	URigVMLibraryNode* LibNode = Cast<URigVMLibraryNode>(Node);
 	if (LibNode == nullptr)
 	{
 		ReportErrorf(TEXT("Node '%s' is not a library node (not collapse nor function)."), *InNodeName.ToString());
@@ -2537,13 +2537,7 @@ TArray<URigVMNode*> URigVMController::ExpandLibraryNode(URigVMLibraryNode* InNod
 
 		for (const FString& EntryTargetLinkPinPath : FromEntryPair.Value)
 		{
-			FString TargetNodeName, TargetPinPath;
-			URigVMPin::SplitPinPathAtStart(EntryTargetLinkPinPath, TargetNodeName, TargetPinPath);
-
-			FString ExpandedNodeName = NodeNameMap.FindChecked(*TargetNodeName).ToString();
-			FString ExpandedPinPath = URigVMPin::JoinPinPath(ExpandedNodeName, TargetPinPath);
-
-			if (URigVMPin* TargetPin = GetGraph()->FindPin(ExpandedPinPath))
+			if (URigVMPin* TargetPin = GetGraph()->FindPin(EntryTargetLinkPinPath))
 			{
 				ApplyPinState(TargetPin, *CollapsedPinState);
 			}
@@ -2747,13 +2741,8 @@ TArray<URigVMNode*> URigVMController::ExpandLibraryNode(URigVMLibraryNode* InNod
 
 		for (const FString& FromEntryNodeTargetPinPath : FromEntryNodePair.Value)
 		{
-			FString TargetPinNodeName, TargetPinPath;
-			URigVMPin::SplitPinPathAtStart(FromEntryNodeTargetPinPath, TargetPinNodeName, TargetPinPath);
-
-			TargetPinNodeName = NodeNameMap.FindChecked(*TargetPinNodeName).ToString();
-
 			URigVMPin* SourcePin = GetGraph()->FindPin(*RemappedSourcePinPath);
-			URigVMPin* TargetPin = GetGraph()->FindPin(URigVMPin::JoinPinPath(TargetPinNodeName, TargetPinPath));
+			URigVMPin* TargetPin = GetGraph()->FindPin(FromEntryNodeTargetPinPath);
 			if (SourcePin && TargetPin)
 			{
 				if (!SourcePin->IsLinkedTo(TargetPin))
