Control Rig: Move default override infrastructure to also support Bound variable paths (and more stuff in the future)
Adapting the AST parser and compiler to deal with variable bindings on encapsulated nodes

diff --git a/Engine/Plugins/Experimental/ControlRig/Source/ControlRigEditor/Private/ControlRigEditorModule.cpp b/Engine/Plugins/Experimental/ControlRig/Source/ControlRigEditor/Private/ControlRigEditorModule.cpp
--- a/Engine/Plugins/Experimental/ControlRig/Source/ControlRigEditor/Private/ControlRigEditorModule.cpp
+++ b/Engine/Plugins/Experimental/ControlRig/Source/ControlRigEditor/Private/ControlRigEditorModule.cpp
@@ -1234,7 +1234,8 @@ void FControlRigEditorModule::GetContextMenuActions(const UControlRigGraphSchema
 						}
 					}
 
-					if(Cast<URigVMUnitNode>(ModelPin->GetNode()))
+					if(Cast<URigVMUnitNode>(ModelPin->GetNode()) != nullptr || 
+						Cast<URigVMLibraryNode>(ModelPin->GetNode()) != nullptr)
 					{
 						if (ModelPin->GetDirection() == ERigVMPinDirection::Input && 
 							!ModelPin->IsExecuteContext() &&
@@ -1282,7 +1283,9 @@ void FControlRigEditorModule::GetContextMenuActions(const UControlRigGraphSchema
 						));
 					}
 
-					if (ModelPin->GetRootPin() == ModelPin && Cast<URigVMUnitNode>(ModelPin->GetNode()) != nullptr)
+					if (ModelPin->GetRootPin() == ModelPin && (
+						Cast<URigVMUnitNode>(ModelPin->GetNode()) != nullptr ||
+						Cast<URigVMLibraryNode>(ModelPin->GetNode()) != nullptr))
 					{
 						if (ModelPin->HasInjectedNodes())
 						{
@@ -1753,12 +1756,24 @@ void FControlRigEditorModule::GetContextMenuActions(const UControlRigGraphSchema
 							LOCTEXT("PromoteToFunction_Tooltip", "Turns the Collapse Node into a Function"),
 							FSlateIcon(),
 							FUIAction(FExecuteAction::CreateLambda([Controller, CollapseNode]() {
-								// Controller->PromoteCollapseNodeToFunction(CollapseNode->GetFName());
-								// todo
+								Controller->PromoteCollapseNodeToFunctionReferenceNode(CollapseNode->GetFName());
 							})
 						));
 					}
 
+					if (URigVMFunctionReferenceNode* FunctionRefNode = Cast<URigVMFunctionReferenceNode>(RigNode->GetModelNode()))
+					{
+						OrganizationSection.AddMenuEntry(
+							"Promote To Collapse Node",
+							LOCTEXT("PromoteToCollapseNode", "Promote To Collapse Node"),
+							LOCTEXT("PromoteToCollapseNode_Tooltip", "Turns the Function Ref Node into a Collapse Node"),
+							FSlateIcon(),
+							FUIAction(FExecuteAction::CreateLambda([Controller, FunctionRefNode]() {
+								Controller->PromoteFunctionReferenceNodeToCollapseNode(FunctionRefNode->GetFName());
+								})
+							));
+					}
+
 					if (URigVMLibraryNode* LibraryNode = Cast<URigVMLibraryNode>(RigNode->GetModelNode()))
 					{
 						OrganizationSection.AddMenuEntry(
@@ -1767,7 +1782,18 @@ void FControlRigEditorModule::GetContextMenuActions(const UControlRigGraphSchema
 							LOCTEXT("ExpandNode_Tooltip", "Expands the contents of the node into this graph"),
 							FSlateIcon(),
 							FUIAction(FExecuteAction::CreateLambda([Controller, LibraryNode]() {
-								Controller->ExpandLibraryNode(LibraryNode->GetFName());
+								Controller->OpenUndoBracket(TEXT("Expand node"));
+								TArray<URigVMNode*> ExpandedNodes = Controller->ExpandLibraryNode(LibraryNode->GetFName());
+								if (ExpandedNodes.Num() > 0)
+								{
+									TArray<FName> ExpandedNodeNames;
+									for (URigVMNode* ExpandedNode : ExpandedNodes)
+									{
+										ExpandedNodeNames.Add(ExpandedNode->GetFName());
+									}
+									Controller->SetNodeSelection(ExpandedNodeNames);
+								}
+								Controller->CloseUndoBracket();
 							})
 						));
 					}
diff --git a/Engine/Plugins/Experimental/ControlRig/Source/ControlRigEditor/Private/ControlRigGraphDetails.cpp b/Engine/Plugins/Experimental/ControlRig/Source/ControlRigEditor/Private/ControlRigGraphDetails.cpp
--- a/Engine/Plugins/Experimental/ControlRig/Source/ControlRigEditor/Private/ControlRigGraphDetails.cpp
+++ b/Engine/Plugins/Experimental/ControlRig/Source/ControlRigEditor/Private/ControlRigGraphDetails.cpp
@@ -472,7 +472,10 @@ void FControlRigArgumentDefaultNode::GenerateChildContent(IDetailChildrenBuilder
 
 void FControlRigArgumentDefaultNode::OnGraphChanged(const FEdGraphEditAction& InAction)
 {
-	OnRebuildChildren.ExecuteIfBound();
+	if (GraphPtr.IsValid() && ControlRigBlueprintPtr.IsValid())
+	{
+		OnRebuildChildren.ExecuteIfBound();
+	}
 }
 
 void FControlRigArgumentDefaultNode::HandleModifiedEvent(ERigVMGraphNotifType InNotifType, URigVMGraph* InGraph, UObject* InSubject)
diff --git a/Engine/Source/Developer/RigVMDeveloper/Private/RigVMCompiler/RigVMAST.cpp b/Engine/Source/Developer/RigVMDeveloper/Private/RigVMCompiler/RigVMAST.cpp
--- a/Engine/Source/Developer/RigVMDeveloper/Private/RigVMCompiler/RigVMAST.cpp
+++ b/Engine/Source/Developer/RigVMDeveloper/Private/RigVMCompiler/RigVMAST.cpp
@@ -559,7 +559,8 @@ bool FRigVMNodeExprAST::IsConstant() const
 		TArray<URigVMPin*> AllPins = CurrentNode->GetAllPinsRecursively();
 		for (URigVMPin* Pin : AllPins)
 		{
-			if (Pin->IsBoundToVariable())
+			URigVMPin::FPinOverride PinOverride(GetProxy(), GetParser()->GetPinOverrides());
+			if (Pin->IsBoundToVariable(PinOverride))
 			{
 				return false;
 			}
@@ -642,7 +643,7 @@ ERigVMPinDirection FRigVMVarExprAST::GetPinDirection() const
 
 FString FRigVMVarExprAST::GetDefaultValue() const
 {
-	return GetPin()->GetDefaultValue(GetParser()->GetPinDefaultOverrides(), GetProxy());
+	return GetPin()->GetDefaultValue(URigVMPin::FPinOverride(GetProxy(), GetParser()->GetPinOverrides()));
 }
 
 bool FRigVMVarExprAST::IsExecuteContext() const
@@ -1121,6 +1122,7 @@ FRigVMExprAST* FRigVMParserAST::TraversePin(const FRigVMASTProxy& InPinProxy, FR
 	ensure(!SubjectToExpression.Contains(InPinProxy));
 
 	URigVMPin* Pin = InPinProxy.GetSubjectChecked<URigVMPin>();
+	URigVMPin::FPinOverride PinOverride(InPinProxy, PinOverrides);
 
 	TArray<FRigVMPinProxyPair> Links = GetSourceLinks(InPinProxy, true);
 	
@@ -1135,21 +1137,22 @@ FRigVMExprAST* FRigVMParserAST::TraversePin(const FRigVMASTProxy& InPinProxy, FR
 
 	struct Local
 	{
-		static void LookForPinsBoundToVariables(URigVMPin* InPin, TArray<URigVMPin*>& SubPinsBoundToVariables)
+		static void LookForPinsBoundToVariables(const FRigVMASTProxy& InPinProxy, URigVMPin* InPin, TArray<URigVMPin*>& SubPinsBoundToVariables, const URigVMPin::FPinOverrideMap& InPinOverrides)
 		{
 			for (URigVMPin* SubPin : InPin->GetSubPins())
 			{
-				if (SubPin->IsBoundToVariable())
+				URigVMPin::FPinOverride PinOverride(InPinProxy, InPinOverrides);
+				if (SubPin->IsBoundToVariable(PinOverride))
 				{
 					SubPinsBoundToVariables.Add(SubPin);
 				}
-				LookForPinsBoundToVariables(SubPin, SubPinsBoundToVariables);
+				LookForPinsBoundToVariables(InPinProxy, SubPin, SubPinsBoundToVariables, InPinOverrides);
 			}
 		}
 	};
 
 	TArray<URigVMPin*> SubPinsBoundToVariables;
-	Local::LookForPinsBoundToVariables(Pin, SubPinsBoundToVariables);
+	Local::LookForPinsBoundToVariables(InPinProxy, Pin, SubPinsBoundToVariables, PinOverrides);
 
 	FRigVMExprAST* PinExpr = nullptr;
 
@@ -1184,7 +1187,7 @@ FRigVMExprAST* FRigVMParserAST::TraversePin(const FRigVMASTProxy& InPinProxy, FR
 			PinCopyExpr->AddParent(PinExpr);
 			PinLiteralExpr->AddParent(PinCopyExpr);
 		}
-		else if (Pin->IsBoundToVariable())
+		else if (Pin->IsBoundToVariable(PinOverride))
 		{
 			PinExpr = MakeExpr<FRigVMExternalVarExprAST>(InPinProxy);
 		}
@@ -1898,7 +1901,9 @@ bool FRigVMParserAST::FoldConstantValuesToLiterals(URigVMGraph* InGraph, URigVMC
 		const TArray<FRigVMASTProxy>& TargetPins = GetTargetPins(PinToComputeProxy);
 		for (const FRigVMASTProxy& TargetPinProxy : TargetPins)
 		{
-			PinDefaultValueOverrides.FindOrAdd(TargetPinProxy) = DefaultValue;
+			URigVMPin::FPinOverrideValue OverrideValue;
+			OverrideValue.DefaultValue = DefaultValue;
+			PinOverrides.FindOrAdd(TargetPinProxy) = OverrideValue;
 		}
 	}
 
@@ -2540,7 +2545,7 @@ void FRigVMParserAST::Inline(URigVMGraph* InGraph, const TArray<FRigVMASTProxy>&
 {
 	struct LocalPinTraversalInfo
 	{
-		URigVMPin::FDefaultValueOverride* PinDefaultValueOverrides;
+		URigVMPin::FPinOverrideMap* PinOverrides;
 		TMap<FRigVMASTProxy, FRigVMASTProxy> SourcePins;
 		TMap<FRigVMASTProxy, TArray<FRigVMASTProxy>>* TargetLinks;
 		TMap<FRigVMASTProxy, TArray<FRigVMASTProxy>>* SourceLinks;
@@ -2754,7 +2759,7 @@ void FRigVMParserAST::Inline(URigVMGraph* InGraph, const TArray<FRigVMASTProxy>&
 					SourceNode->IsA<URigVMLibraryNode>() ||
 					SourceNode->IsA<URigVMFunctionReturnNode>())
 				{
-					OutTraversalInfo.PinDefaultValueOverrides->FindOrAdd(InPinProxy) = SourcePin->GetDefaultValue();
+					OutTraversalInfo.PinOverrides->FindOrAdd(InPinProxy) = URigVMPin::FPinOverrideValue(SourcePin);
 				}
 				else
 				{
@@ -2818,7 +2823,7 @@ void FRigVMParserAST::Inline(URigVMGraph* InGraph, const TArray<FRigVMASTProxy>&
 	// c) flatten links from an entry node / to a return node
 	//    also traverse links along reroutes and flatten them
 	LocalPinTraversalInfo TraversalInfo;
-	TraversalInfo.PinDefaultValueOverrides = &PinDefaultValueOverrides;
+	TraversalInfo.PinOverrides = &PinOverrides;
 	TraversalInfo.TargetLinks = &TargetLinks;
 	TraversalInfo.SourceLinks = &SourceLinks;
 
diff --git a/Engine/Source/Developer/RigVMDeveloper/Private/RigVMCompiler/RigVMCompiler.cpp b/Engine/Source/Developer/RigVMDeveloper/Private/RigVMCompiler/RigVMCompiler.cpp
--- a/Engine/Source/Developer/RigVMDeveloper/Private/RigVMCompiler/RigVMCompiler.cpp
+++ b/Engine/Source/Developer/RigVMDeveloper/Private/RigVMCompiler/RigVMCompiler.cpp
@@ -743,7 +743,7 @@ void URigVMCompiler::TraverseAssign(const FRigVMAssignExprAST* InExpr, FRigVMCom
 		{
 			struct Local
 			{
-				static void SetupRegisterOffset(URigVM* VM, URigVMPin* Pin, FRigVMOperand& Operand, bool bSource)
+				static void SetupRegisterOffset(URigVM* VM, URigVMPin* Pin, FRigVMOperand& Operand, const FRigVMVarExprAST* VarExpr, bool bSource)
 				{
 					URigVMPin* RootPin = Pin->GetRootPin();
 					if (Pin == RootPin)
@@ -751,7 +751,8 @@ void URigVMCompiler::TraverseAssign(const FRigVMAssignExprAST* InExpr, FRigVMCom
 						return;
 					}
 
-					if (bSource && Pin->IsBoundToVariable())
+					URigVMPin::FPinOverride PinOverride(VarExpr->GetProxy(), VarExpr->GetParser()->GetPinOverrides());
+					if (bSource && Pin->IsBoundToVariable(PinOverride))
 					{
 						return;
 					}
@@ -798,8 +799,8 @@ void URigVMCompiler::TraverseAssign(const FRigVMAssignExprAST* InExpr, FRigVMCom
 				}
 			};
 
-			Local::SetupRegisterOffset(WorkData.VM, InExpr->GetSourcePin(), Source, true);
-			Local::SetupRegisterOffset(WorkData.VM, InExpr->GetTargetPin(), Target, false);
+			Local::SetupRegisterOffset(WorkData.VM, InExpr->GetSourcePin(), Source, SourceExpr, true);
+			Local::SetupRegisterOffset(WorkData.VM, InExpr->GetTargetPin(), Target, TargetExpr, false);
 		}
 
 		WorkData.VM->GetByteCode().AddCopyOp(Source, Target);
@@ -1143,7 +1144,8 @@ FString URigVMCompiler::GetPinHash(URigVMPin* InPin, const FRigVMVarExprAST* InV
 	{
 		if (InVarExpr->IsA(FRigVMExprAST::ExternalVar) && !bIsDebugValue)
 		{
-			FString VariablePath = InPin->GetBoundVariablePath();
+			URigVMPin::FPinOverride PinOverride(InVarExpr->GetProxy(), InVarExpr->GetParser()->GetPinOverrides());
+			FString VariablePath = InPin->GetBoundVariablePath(PinOverride);
 			return FString::Printf(TEXT("%sVariable::%s%s"), *Prefix, *VariablePath, *Suffix);
 		}
 
@@ -1313,7 +1315,9 @@ FRigVMOperand URigVMCompiler::FindOrAddRegister(const FRigVMVarExprAST* InVarExp
 		}
 	}
 
-	const URigVMPin::FDefaultValueOverride& DefaultValueOverride = InVarExpr->GetParser()->GetPinDefaultOverrides();
+	const URigVMPin::FPinOverrideMap& PinOverrides = InVarExpr->GetParser()->GetPinOverrides();
+	URigVMPin::FPinOverride PinOverride(InVarExpr->GetProxy(), PinOverrides);
+
 	URigVMPin* Pin = InVarExpr->GetPin();
 	FString BaseCPPType = Pin->IsArray() ? Pin->GetArrayElementCppType() : Pin->GetCPPType();
 	FString Hash = GetPinHash(Pin, InVarExpr, bIsDebugValue);
@@ -1337,9 +1341,9 @@ FRigVMOperand URigVMCompiler::FindOrAddRegister(const FRigVMVarExprAST* InVarExp
 	// check if this is a variable with a segment path
 	if (!Operand.IsValid())
 	{
-		if (Pin->IsBoundToVariable() && !bIsDebugValue)
+		if (Pin->IsBoundToVariable(PinOverride) && !bIsDebugValue)
 		{
-			FString VariablePath = Pin->GetBoundVariablePath();
+			FString VariablePath = Pin->GetBoundVariablePath(PinOverride);
 			FString VariableName = VariablePath, SegmentPath;
 			VariablePath.Split(TEXT("."), &VariableName, &SegmentPath);
 
@@ -1374,14 +1378,14 @@ FRigVMOperand URigVMCompiler::FindOrAddRegister(const FRigVMVarExprAST* InVarExp
 
 				int32 DesiredArraySize = VMStruct->GetArraySize(Pin->GetFName(), WorkData.RigVMUserData);
 
-				DefaultValues = URigVMPin::SplitDefaultValue(Pin->GetDefaultValue(DefaultValueOverride, InVarExpr->GetProxy()));
+				DefaultValues = URigVMPin::SplitDefaultValue(Pin->GetDefaultValue(PinOverride));
 
 				if (DefaultValues.Num() != DesiredArraySize)
 				{
 					FString DefaultValue;
 					if (Pin->GetArraySize() > 0)
 					{
-						DefaultValue = Pin->GetSubPins()[0]->GetDefaultValue(DefaultValueOverride, InVarExpr->GetProxy());
+						DefaultValue = Pin->GetSubPins()[0]->GetDefaultValue(PinOverride);
 					}
 
 					DefaultValues.Reset();
@@ -1393,7 +1397,7 @@ FRigVMOperand URigVMCompiler::FindOrAddRegister(const FRigVMVarExprAST* InVarExp
 			}
 			else
 			{
-				DefaultValues = URigVMPin::SplitDefaultValue(Pin->GetDefaultValue(DefaultValueOverride, InVarExpr->GetProxy()));
+				DefaultValues = URigVMPin::SplitDefaultValue(Pin->GetDefaultValue(PinOverride));
 			}
 
 			while (DefaultValues.Num() < Pin->GetSubPins().Num())
@@ -1403,7 +1407,7 @@ FRigVMOperand URigVMCompiler::FindOrAddRegister(const FRigVMVarExprAST* InVarExp
 		}
 		else if (URigVMEnumNode* EnumNode = Cast<URigVMEnumNode>(Pin->GetNode()))
 		{
-			FString EnumValueStr = EnumNode->GetDefaultValue(DefaultValueOverride, InVarExpr->GetProxy());
+			FString EnumValueStr = EnumNode->GetDefaultValue(PinOverride);
 			if (UEnum* Enum = EnumNode->GetEnum())
 			{
 				DefaultValues.Add(FString::FromInt((int32)Enum->GetValueByNameString(EnumValueStr)));
@@ -1415,7 +1419,7 @@ FRigVMOperand URigVMCompiler::FindOrAddRegister(const FRigVMVarExprAST* InVarExp
 		}
 		else
 		{
-			DefaultValues.Add(Pin->GetDefaultValue(DefaultValueOverride, InVarExpr->GetProxy()));
+			DefaultValues.Add(Pin->GetDefaultValue(PinOverride));
 		}
 
 		UScriptStruct* ScriptStruct = Pin->GetScriptStruct();
diff --git a/Engine/Source/Developer/RigVMDeveloper/Private/RigVMModel/Nodes/RigVMEnumNode.cpp b/Engine/Source/Developer/RigVMDeveloper/Private/RigVMModel/Nodes/RigVMEnumNode.cpp
--- a/Engine/Source/Developer/RigVMDeveloper/Private/RigVMModel/Nodes/RigVMEnumNode.cpp
+++ b/Engine/Source/Developer/RigVMDeveloper/Private/RigVMModel/Nodes/RigVMEnumNode.cpp
@@ -44,12 +44,12 @@ UObject* URigVMEnumNode::GetCPPTypeObject() const
 	return EnumValuePin->GetCPPTypeObject();
 }
 
-FString URigVMEnumNode::GetDefaultValue(const URigVMPin::FDefaultValueOverride& InDefaultValueOverride, const FRigVMASTProxy& InProxy) const
+FString URigVMEnumNode::GetDefaultValue(const URigVMPin::FPinOverride& InOverride) const
 {
 	URigVMPin* EnumValuePin = FindPin(EnumValueName);
 	if (EnumValuePin == nullptr)
 	{
 		return FString();
 	}
-	return EnumValuePin->GetDefaultValue(InDefaultValueOverride, InProxy);
+	return EnumValuePin->GetDefaultValue(InOverride);
 }
diff --git a/Engine/Source/Developer/RigVMDeveloper/Private/RigVMModel/RigVMPin.cpp b/Engine/Source/Developer/RigVMDeveloper/Private/RigVMModel/RigVMPin.cpp
--- a/Engine/Source/Developer/RigVMDeveloper/Private/RigVMModel/RigVMPin.cpp
+++ b/Engine/Source/Developer/RigVMDeveloper/Private/RigVMModel/RigVMPin.cpp
@@ -42,6 +42,9 @@ URigVMPin* URigVMInjectionInfo::GetPin() const
 	return CastChecked<URigVMPin>(GetOuter());
 }
 
+const URigVMPin::FPinOverrideMap URigVMPin::EmptyPinOverrideMap;
+const URigVMPin::FPinOverride URigVMPin::EmptyPinOverride = URigVMPin::FPinOverride(FRigVMASTProxy(), EmptyPinOverrideMap);
+
 bool URigVMPin::SplitPinPathAtStart(const FString& InPinPath, FString& LeftMost, FString& Right)
 {
 	return InPinPath.Split(TEXT("."), &LeftMost, &Right, ESearchCase::IgnoreCase, ESearchDir::FromStart);
@@ -366,14 +369,14 @@ bool URigVMPin::IsExecuteContext() const
 
 FString URigVMPin::GetDefaultValue() const
 {
-	return GetDefaultValue(FDefaultValueOverride(), FRigVMASTProxy());
+	return GetDefaultValue(EmptyPinOverride);
 }
 
-FString URigVMPin::GetDefaultValue(const FDefaultValueOverride& InDefaultValueOverride, const FRigVMASTProxy& InProxy) const
+FString URigVMPin::GetDefaultValue(const URigVMPin::FPinOverride& InOverride) const
 {
-	if (const FString* Override = InDefaultValueOverride.Find(InProxy.GetSibling((URigVMPin*)this)))
+	if (FPinOverrideValue const* OverrideValuePtr = InOverride.Value.Find(InOverride.Key.GetSibling((URigVMPin*)this)))
 	{
-		return *Override;
+		return OverrideValuePtr->DefaultValue;
 	}
 
 	if (IsArray())
@@ -383,7 +386,7 @@ FString URigVMPin::GetDefaultValue(const FDefaultValueOverride& InDefaultValueOv
 			TArray<FString> ElementDefaultValues;
 			for (URigVMPin* SubPin : SubPins)
 			{
-				FString ElementDefaultValue = SubPin->GetDefaultValue(InDefaultValueOverride, InProxy);
+				FString ElementDefaultValue = SubPin->GetDefaultValue(InOverride);
 				if (SubPin->IsStringType())
 				{
 					ElementDefaultValue = TEXT("\"") + ElementDefaultValue + TEXT("\"");
@@ -410,7 +413,7 @@ FString URigVMPin::GetDefaultValue(const FDefaultValueOverride& InDefaultValueOv
 			TArray<FString> MemberDefaultValues;
 			for (URigVMPin* SubPin : SubPins)
 			{
-				FString MemberDefaultValue = SubPin->GetDefaultValue(InDefaultValueOverride, InProxy);
+				FString MemberDefaultValue = SubPin->GetDefaultValue(InOverride);
 				if (SubPin->IsStringType())
 				{
 					MemberDefaultValue = TEXT("\"") + MemberDefaultValue + TEXT("\"");
@@ -566,13 +569,48 @@ UObject* URigVMPin::FindObjectFromCPPTypeObjectPath(const FString& InObjectPath)
 	return FindObject<UObject>(ANY_PACKAGE, *InObjectPath);
 }
 
+// Returns the variable bound to this pin (or NAME_None)
+const FString& URigVMPin::GetBoundVariablePath() const
+{
+	return GetBoundVariablePath(EmptyPinOverride);
+}
+
+// Returns the variable bound to this pin (or NAME_None)
+const FString& URigVMPin::GetBoundVariablePath(const URigVMPin::FPinOverride& InOverride) const
+{
+	if (FPinOverrideValue const* OverrideValuePtr = InOverride.Value.Find(InOverride.Key.GetSibling((URigVMPin*)this)))
+	{
+		return OverrideValuePtr->BoundVariablePath;
+	}
+	return BoundVariablePath;
+}
+
+// Returns the variable bound to this pin (or NAME_None)
 FString URigVMPin::GetBoundVariableName() const
 {
-	FString VariableName = BoundVariablePath;
+	return GetBoundVariableName(EmptyPinOverride);
+}
+
+// Returns the variable bound to this pin (or NAME_None)
+FString URigVMPin::GetBoundVariableName(const URigVMPin::FPinOverride& InOverride) const
+{
+	FString VariableName = GetBoundVariablePath(InOverride);
 	BoundVariablePath.Split(TEXT("."), &VariableName, nullptr);
 	return VariableName;
 }
 
+// Returns true if this pin is bound to a variable
+bool URigVMPin::IsBoundToVariable() const
+{
+	return IsBoundToVariable(EmptyPinOverride);
+}
+
+// Returns true if this pin is bound to a variable
+bool URigVMPin::IsBoundToVariable(const URigVMPin::FPinOverride& InOverride) const
+{
+	return !GetBoundVariablePath(InOverride).IsEmpty();
+}
+
 bool URigVMPin::CanBeBoundToVariable(const FRigVMExternalVariable& InExternalVariable, const FRigVMRegisterOffset& InOffset) const
 {
 	if (!InExternalVariable.IsValid(true))
diff --git a/Engine/Source/Developer/RigVMDeveloper/Public/RigVMCompiler/RigVMAST.h b/Engine/Source/Developer/RigVMDeveloper/Public/RigVMCompiler/RigVMAST.h
--- a/Engine/Source/Developer/RigVMDeveloper/Public/RigVMCompiler/RigVMAST.h
+++ b/Engine/Source/Developer/RigVMDeveloper/Public/RigVMCompiler/RigVMAST.h
@@ -1207,7 +1207,7 @@ public:
 	const FRigVMBlockExprAST* GetObsoleteBlock(bool bCreateIfMissing = true) const;
 
 	// returns the AST's override table for pin defaults
-	const URigVMPin::FDefaultValueOverride& GetPinDefaultOverrides() const { return PinDefaultValueOverrides; }
+	const URigVMPin::FPinOverrideMap& GetPinOverrides() const { return PinOverrides; }
 
 private:
 
@@ -1359,7 +1359,7 @@ private:
 	TMap<FRigVMASTProxy, TArray<FRigVMASTProxy>> TargetLinks;
 	TMap<FRigVMASTProxy, TArray<FRigVMASTProxy>> SourceLinks;
 	static const TArray<FRigVMASTProxy> EmptyProxyArray;
-	URigVMPin::FDefaultValueOverride PinDefaultValueOverrides;
+	URigVMPin::FPinOverrideMap PinOverrides;
 
 	enum ETraverseRelationShip
 	{
diff --git a/Engine/Source/Developer/RigVMDeveloper/Public/RigVMModel/Nodes/RigVMEnumNode.h b/Engine/Source/Developer/RigVMDeveloper/Public/RigVMModel/Nodes/RigVMEnumNode.h
--- a/Engine/Source/Developer/RigVMDeveloper/Public/RigVMModel/Nodes/RigVMEnumNode.h
+++ b/Engine/Source/Developer/RigVMDeveloper/Public/RigVMModel/Nodes/RigVMEnumNode.h
@@ -34,7 +34,7 @@ public:
 	UObject* GetCPPTypeObject() const;
 
 	// Returns the default value of the parameter as a string
-	FString GetDefaultValue(const URigVMPin::FDefaultValueOverride& InDefaultValueOverride = URigVMPin::FDefaultValueOverride(), const FRigVMASTProxy& InProxy = FRigVMASTProxy()) const;
+	FString GetDefaultValue(const URigVMPin::FPinOverride& InOverride = URigVMPin::EmptyPinOverride) const;
 
 	// Override of node title
 	virtual FLinearColor GetNodeColor() const override { return FLinearColor::Blue; }
diff --git a/Engine/Source/Developer/RigVMDeveloper/Public/RigVMModel/RigVMPin.h b/Engine/Source/Developer/RigVMDeveloper/Public/RigVMModel/RigVMPin.h
--- a/Engine/Source/Developer/RigVMDeveloper/Public/RigVMModel/RigVMPin.h
+++ b/Engine/Source/Developer/RigVMDeveloper/Public/RigVMModel/RigVMPin.h
@@ -52,14 +52,6 @@ public:
 	URigVMPin* GetPin() const;
 };
 
-
-/**
- * The Visual Debugging Info is used for visually displaying
- * Data flowing through a pin. Typically this is attached to an input
- * pin causes the pin to inject a node on the link driving it
- */
-
-
 /**
  * The Pin represents a single connector / pin on a node in the RigVM model.
  * Pins can be connected based on rules. Pins also provide access to a 'PinPath',
@@ -78,8 +70,29 @@ class RIGVMDEVELOPER_API URigVMPin : public UObject
 
 public:
 
+	// A struct to store a pin override value
+	struct FPinOverrideValue 
+	{
+		FPinOverrideValue()
+			: DefaultValue()
+			, BoundVariablePath()
+		{}
+
+		FPinOverrideValue(URigVMPin* InPin)
+			: DefaultValue(InPin->GetDefaultValue())
+			, BoundVariablePath(InPin->GetBoundVariablePath())
+		{
+		}
+
+		FString DefaultValue;
+		FString BoundVariablePath;
+	};
+
 	// A map used to override pin default values
-	typedef TMap<FRigVMASTProxy, FString> FDefaultValueOverride;
+	typedef TMap<FRigVMASTProxy, FPinOverrideValue> FPinOverrideMap;
+	typedef TPair<FRigVMASTProxy, const FPinOverrideMap&> FPinOverride;
+	static const URigVMPin::FPinOverrideMap EmptyPinOverrideMap;
+	static const FPinOverride EmptyPinOverride;
 
 	// Splits a PinPath at the start, so for example "Node.Color.R" becomes "Node" and "Color.R"
 	static bool SplitPinPathAtStart(const FString& InPinPath, FString& LeftMost, FString& Right);
@@ -189,7 +202,7 @@ public:
 	FString GetDefaultValue() const;
 
 	// Returns the default value with an additional override ma
-	FString GetDefaultValue(const FDefaultValueOverride& InDefaultValueOverride, const FRigVMASTProxy& InProxy) const;
+	FString GetDefaultValue(const FPinOverride& InOverride) const;
 
 	// Returns true if the default value provided is valid
 	bool IsValidDefaultValue(const FString& InDefaultValue) const;
@@ -300,13 +313,22 @@ public:
 	const TArray<URigVMInjectionInfo*> GetInjectedNodes() const { return InjectionInfos; }
 
 	// Returns the variable bound to this pin (or NAME_None)
-	const FString& GetBoundVariablePath() const { return BoundVariablePath; }
+	const FString& GetBoundVariablePath() const;
+
+	// Returns the variable bound to this pin (or NAME_None)
+	const FString& GetBoundVariablePath(const FPinOverride& InOverride) const;
 
 	// Returns the variable bound to this pin (or NAME_None)
 	FString GetBoundVariableName() const;
 
+	// Returns the variable bound to this pin (or NAME_None)
+	FString GetBoundVariableName(const FPinOverride& InOverride) const;
+
+	// Returns true if this pin is bound to a variable
+	bool IsBoundToVariable() const;
+
 	// Returns true if this pin is bound to a variable
-	bool IsBoundToVariable() const { return !BoundVariablePath.IsEmpty(); }
+	bool IsBoundToVariable(const FPinOverride& InOverride) const;
 
 	// Returns true if the pin can be bound to a given variable
 	bool CanBeBoundToVariable(const FRigVMExternalVariable& InExternalVariable, const FRigVMRegisterOffset& InOffset = FRigVMRegisterOffset()) const;
