Control Rig: Controller support for promote workflows

diff --git a/Engine/Source/Developer/RigVMDeveloper/Private/RigVMModel/RigVMController.cpp b/Engine/Source/Developer/RigVMDeveloper/Private/RigVMModel/RigVMController.cpp
--- a/Engine/Source/Developer/RigVMDeveloper/Private/RigVMModel/RigVMController.cpp
+++ b/Engine/Source/Developer/RigVMDeveloper/Private/RigVMModel/RigVMController.cpp
@@ -2874,6 +2874,168 @@ TArray<URigVMNode*> URigVMController::ExpandLibraryNode(URigVMLibraryNode* InNod
 	return ExpandedNodes;
 }
 
+FName URigVMController::PromoteCollapseNodeToFunctionReferenceNode(const FName& InNodeName, bool bSetupUndoRedo)
+{
+	if (!IsValidGraph())
+	{
+		return NAME_None;
+	}
+
+	URigVMGraph* Graph = GetGraph();
+	check(Graph);
+
+	URigVMNode* Result = PromoteCollapseNodeToFunctionReferenceNode(Cast<URigVMCollapseNode>(Graph->FindNodeByName(InNodeName)), bSetupUndoRedo);
+	if (Result)
+	{
+		return Result->GetFName();
+	}
+	return NAME_None;
+}
+
+FName URigVMController::PromoteFunctionReferenceNodeToCollapseNode(const FName& InNodeName, bool bSetupUndoRedo)
+{
+	if (!IsValidGraph())
+	{
+		return NAME_None;
+	}
+
+	URigVMGraph* Graph = GetGraph();
+	check(Graph);
+
+	URigVMNode* Result = PromoteFunctionReferenceNodeToCollapseNode(Cast<URigVMFunctionReferenceNode>(Graph->FindNodeByName(InNodeName)), bSetupUndoRedo);
+	if (Result)
+	{
+		return Result->GetFName();
+	}
+	return NAME_None;
+}
+
+URigVMFunctionReferenceNode* URigVMController::PromoteCollapseNodeToFunctionReferenceNode(URigVMCollapseNode* InCollapseNode, bool bSetupUndoRedo)
+{
+	if (!IsValidNodeForGraph(InCollapseNode))
+	{
+		return nullptr;
+	}
+
+	URigVMGraph* Graph = GetGraph();
+	check(Graph);
+
+	URigVMFunctionLibrary* FunctionLibrary = Graph->GetDefaultFunctionLibrary();
+	if (FunctionLibrary == nullptr)
+	{
+		return nullptr;
+	}
+
+	if (bSetupUndoRedo)
+	{
+		OpenUndoBracket(TEXT("Promote to Function"));
+	}
+
+	URigVMFunctionReferenceNode* FunctionRefNode = nullptr;
+	URigVMCollapseNode* FunctionDefinition = nullptr;
+	{
+		FRigVMControllerGraphGuard GraphGuard(this, FunctionLibrary, bSetupUndoRedo);
+		FString FunctionName = GetValidNodeName(InCollapseNode->GetName());
+		FunctionDefinition = DuplicateObject<URigVMCollapseNode>(InCollapseNode, FunctionLibrary, *FunctionName);
+
+		if (FunctionDefinition)
+		{
+			FunctionLibrary->Nodes.Add(FunctionDefinition);
+			Notify(ERigVMGraphNotifType::NodeAdded, FunctionDefinition);
+		}
+	}
+
+	if (FunctionDefinition)
+	{
+		FString NodeName = InCollapseNode->GetName();
+		FVector2D NodePosition = InCollapseNode->GetPosition();
+		TMap<FString, FPinState> PinStates = GetPinStates(InCollapseNode);
+
+		TArray<URigVMLink*> Links = InCollapseNode->GetLinks();
+		TArray< TPair< FString, FString > > LinkPaths;
+		for (URigVMLink* Link : Links)
+		{
+			LinkPaths.Add(TPair< FString, FString >(Link->GetSourcePin()->GetPinPath(), Link->GetTargetPin()->GetPinPath()));
+		}
+
+		RemoveNode(InCollapseNode, bSetupUndoRedo);
+
+		FunctionRefNode = AddFunctionReferenceNode(FunctionDefinition, NodePosition, NodeName, bSetupUndoRedo);
+
+		if (FunctionRefNode)
+		{
+			ApplyPinStates(FunctionRefNode, PinStates);
+			for (const TPair<FString, FString>& LinkPath : LinkPaths)
+			{
+				AddLink(LinkPath.Key, LinkPath.Value, bSetupUndoRedo);
+			}
+		}
+	}
+
+	if (bSetupUndoRedo)
+	{
+		CloseUndoBracket();
+	}
+
+	return FunctionRefNode;
+}
+
+URigVMCollapseNode* URigVMController::PromoteFunctionReferenceNodeToCollapseNode(URigVMFunctionReferenceNode* InFunctionRefNode, bool bSetupUndoRedo)
+{
+	if (!IsValidNodeForGraph(InFunctionRefNode))
+	{
+		return nullptr;
+	}
+
+	URigVMGraph* Graph = GetGraph();
+	check(Graph);
+
+	URigVMCollapseNode* FunctionDefinition = Cast<URigVMCollapseNode>(InFunctionRefNode->GetReferencedNode());
+	if (FunctionDefinition == nullptr)
+	{
+		return nullptr;
+	}
+
+	if (bSetupUndoRedo)
+	{
+		OpenUndoBracket(TEXT("Promote to Collapse Node"));
+	}
+
+	FString NodeName = InFunctionRefNode->GetName();
+	FVector2D NodePosition = InFunctionRefNode->GetPosition();
+	TMap<FString, FPinState> PinStates = GetPinStates(InFunctionRefNode);
+
+	TArray<URigVMLink*> Links = InFunctionRefNode->GetLinks();
+	TArray< TPair< FString, FString > > LinkPaths;
+	for (URigVMLink* Link : Links)
+	{
+		LinkPaths.Add(TPair< FString, FString >(Link->GetSourcePin()->GetPinPath(), Link->GetTargetPin()->GetPinPath()));
+	}
+
+	RemoveNode(InFunctionRefNode, bSetupUndoRedo);
+
+	URigVMCollapseNode* CollapseNode = DuplicateObject<URigVMCollapseNode>(FunctionDefinition, Graph, *NodeName);
+	if(CollapseNode)
+	{
+		CollapseNode->Position = NodePosition;
+		Graph->Nodes.Add(CollapseNode);
+		Notify(ERigVMGraphNotifType::NodeAdded, CollapseNode);
+
+		ApplyPinStates(CollapseNode, PinStates);
+		for (const TPair<FString, FString>& LinkPath : LinkPaths)
+		{
+			AddLink(LinkPath.Key, LinkPath.Value, bSetupUndoRedo);
+		}
+	}
+
+	if (bSetupUndoRedo)
+	{
+		CloseUndoBracket();
+	}
+
+	return CollapseNode;
+}
+
 void URigVMController::RefreshFunctionPins(URigVMNode* InNode, bool bNotify)
 {
 	if (InNode == nullptr)
diff --git a/Engine/Source/Developer/RigVMDeveloper/Public/RigVMModel/RigVMController.h b/Engine/Source/Developer/RigVMDeveloper/Public/RigVMModel/RigVMController.h
--- a/Engine/Source/Developer/RigVMDeveloper/Public/RigVMModel/RigVMController.h
+++ b/Engine/Source/Developer/RigVMDeveloper/Public/RigVMModel/RigVMController.h
@@ -287,6 +287,14 @@ public:
 	UFUNCTION(BlueprintCallable, Category = RigVMController)
 	TArray<URigVMNode*> ExpandLibraryNode(const FName& InNodeName, bool bSetupUndoRedo = true);
 
+	// Turns a collapse node into a function node
+	UFUNCTION(BlueprintCallable, Category = RigVMController)
+	FName PromoteCollapseNodeToFunctionReferenceNode(const FName& InNodeName, bool bSetupUndoRedo = true);
+
+	// Turns a collapse node into a function node
+	UFUNCTION(BlueprintCallable, Category = RigVMController)
+	FName PromoteFunctionReferenceNodeToCollapseNode(const FName& InNodeName, bool bSetupUndoRedo = true);
+
 #endif
 
 	// Removes a node from the graph
@@ -578,6 +586,9 @@ private:
 	static void ForEveryPinRecursively(URigVMNode* InNode, TFunction<void(URigVMPin*)> OnEachPinFunction);
 	URigVMCollapseNode* CollapseNodes(const TArray<URigVMNode*>& InNodes, const FString& InCollapseNodeName, bool bSetupUndoRedo);
 	TArray<URigVMNode*> ExpandLibraryNode(URigVMLibraryNode* InNode, bool bSetupUndoRedo);
+	URigVMFunctionReferenceNode* PromoteCollapseNodeToFunctionReferenceNode(URigVMCollapseNode* InCollapseNode, bool bSetupUndoRedo);
+	URigVMCollapseNode* PromoteFunctionReferenceNodeToCollapseNode(URigVMFunctionReferenceNode* InFunctionRefNode, bool bSetupUndoRedo);
+
 	void RefreshFunctionPins(URigVMNode* InNode, bool bNotify = true);
 
 	struct FPinState
