Control Rig: Added series of unit tests for Functions / Encapsulation

diff --git a/Engine/Source/Developer/RigVMDeveloper/Private/RigVMModel/RigVMController.cpp b/Engine/Source/Developer/RigVMDeveloper/Private/RigVMModel/RigVMController.cpp
--- a/Engine/Source/Developer/RigVMDeveloper/Private/RigVMModel/RigVMController.cpp
+++ b/Engine/Source/Developer/RigVMDeveloper/Private/RigVMModel/RigVMController.cpp
@@ -1773,6 +1773,17 @@ TArray<FName> URigVMController::ImportNodesFromText(const FString& InText, bool
 			}
 		}
 
+		if (URigVMFunctionReferenceNode* FunctionRefNode = Cast<URigVMFunctionReferenceNode>(CreatedNode))
+		{
+			if (URigVMFunctionLibrary* FunctionLibrary = FunctionRefNode->GetLibrary())
+			{
+				if (URigVMLibraryNode* FunctionDefinition = FunctionRefNode->GetReferencedNode())
+				{
+					FunctionLibrary->FunctionReferences.FindOrAdd(FunctionDefinition).FunctionReferences.Add(FunctionRefNode);
+				}
+			}
+		}
+
 		Notify(ERigVMGraphNotifType::NodeAdded, CreatedNode);
 
 		NodeNames.Add(CreatedNode->GetFName());
@@ -2623,12 +2634,8 @@ TArray<URigVMNode*> URigVMController::ExpandLibraryNode(URigVMLibraryNode* InNod
 				const TArray<FString>& LibraryPinLinks = *LibraryPinLinksPtr;
 				ensure(LibraryPinLinks.Num() == 1);
 
-				FString TargetPinNodeName, TargetPinPathRemaining;
-				URigVMPin::SplitPinPathAtStart(LibraryPinLinks[0], TargetPinNodeName, TargetPinPathRemaining);
-
 				SourcePinPath = LibraryPinPath;
-				TargetPinNodeName = NodeNameMap.FindChecked(*TargetPinNodeName).ToString();
-				TargetPinPath = URigVMPin::JoinPinPath(TargetPinNodeName, TargetPinPathRemaining);
+				TargetPinPath = LibraryPinLinks[0];
 			}
 		}
 
@@ -8039,6 +8046,7 @@ void URigVMController::RefreshFunctionReferences(URigVMLibraryNode* InFunctionDe
 					TArray<URigVMLink*> Links = FunctionReferencePtr->GetLinks();
 					DetachLinksFromPinObjects(&Links, true);
 					RepopulatePinsOnNode(FunctionReferencePtr.Get(), false, true);
+					TGuardValue<bool> ReportGuard(bReportWarningsAndErrors, false);
 					ReattachLinksToPinObjects(false, &Links, true);
 				}
 			}
diff --git a/Engine/Source/Developer/RigVMDeveloper/Private/RigVMModel/RigVMFunctionLibrary.cpp b/Engine/Source/Developer/RigVMDeveloper/Private/RigVMModel/RigVMFunctionLibrary.cpp
--- a/Engine/Source/Developer/RigVMDeveloper/Private/RigVMModel/RigVMFunctionLibrary.cpp
+++ b/Engine/Source/Developer/RigVMDeveloper/Private/RigVMModel/RigVMFunctionLibrary.cpp
@@ -36,3 +36,42 @@ URigVMLibraryNode* URigVMFunctionLibrary::FindFunction(const FName& InFunctionNa
 	return Cast<URigVMLibraryNode>(FindNodeByName(*FunctionNameStr));
 }
 
+TArray< TSoftObjectPtr<URigVMFunctionReferenceNode> > URigVMFunctionLibrary::GetReferencesForFunction(const FName& InFunctionName)
+{
+	TArray< TSoftObjectPtr<URigVMFunctionReferenceNode> > Result;
+
+	if (URigVMLibraryNode* Function = FindFunction(InFunctionName))
+	{
+		FRigVMFunctionReferenceArray* ReferencesEntry = FunctionReferences.Find(Function);
+		if (ReferencesEntry)
+		{
+			for (int32 ReferenceIndex = 0; ReferenceIndex < ReferencesEntry->Num(); ReferenceIndex++)
+			{
+				const TSoftObjectPtr<URigVMFunctionReferenceNode>& Reference = ReferencesEntry->operator [](ReferenceIndex);
+				Result.Add(TSoftObjectPtr<URigVMFunctionReferenceNode>(Reference.GetUniqueID()));
+			}
+		}
+	}
+
+	return Result;
+}
+
+TArray< FString > URigVMFunctionLibrary::GetReferencePathsForFunction(const FName& InFunctionName)
+{
+	TArray< FString > Result;
+
+	if (URigVMLibraryNode* Function = FindFunction(InFunctionName))
+	{
+		FRigVMFunctionReferenceArray* ReferencesEntry = FunctionReferences.Find(Function);
+		if (ReferencesEntry)
+		{
+			for (int32 ReferenceIndex = 0; ReferenceIndex < ReferencesEntry->Num(); ReferenceIndex++)
+			{
+				const TSoftObjectPtr<URigVMFunctionReferenceNode>& Reference = ReferencesEntry->operator [](ReferenceIndex);
+				Result.Add(Reference.ToString());
+			}
+		}
+	}
+
+	return Result;
+}
diff --git a/Engine/Source/Developer/RigVMDeveloper/Public/RigVMModel/RigVMFunctionLibrary.h b/Engine/Source/Developer/RigVMDeveloper/Public/RigVMModel/RigVMFunctionLibrary.h
--- a/Engine/Source/Developer/RigVMDeveloper/Public/RigVMModel/RigVMFunctionLibrary.h
+++ b/Engine/Source/Developer/RigVMDeveloper/Public/RigVMModel/RigVMFunctionLibrary.h
@@ -54,6 +54,14 @@ public:
 	UFUNCTION(BlueprintCallable, Category = RigVMGraph)
 	URigVMLibraryNode* FindFunction(const FName& InFunctionName) const;
 
+	// Returns all references for a given function name
+	UFUNCTION(BlueprintCallable, Category = RigVMGraph)
+	TArray< TSoftObjectPtr<URigVMFunctionReferenceNode> > GetReferencesForFunction(const FName& InFunctionName);
+
+	// Returns all references for a given function name
+	UFUNCTION(BlueprintCallable, Category = RigVMGraph)
+	TArray< FString > GetReferencePathsForFunction(const FName& InFunctionName);
+
 private:
 
 	UPROPERTY()
