Control Rig: Make sure to initialize struct memory for pin default value validation

diff --git a/Engine/Source/Developer/RigVMDeveloper/Private/RigVMModel/RigVMPin.cpp b/Engine/Source/Developer/RigVMDeveloper/Private/RigVMModel/RigVMPin.cpp
--- a/Engine/Source/Developer/RigVMDeveloper/Private/RigVMModel/RigVMPin.cpp
+++ b/Engine/Source/Developer/RigVMDeveloper/Private/RigVMModel/RigVMPin.cpp
@@ -460,6 +460,7 @@ bool URigVMPin::IsValidDefaultValue(const FString& InDefaultValue) const
 			FRigVMPinDefaultValueImportErrorContext ErrorPipe;
 			TArray<uint8> TempStructBuffer;
 			TempStructBuffer.AddUninitialized(ScriptStruct->GetStructureSize());
+			ScriptStruct->InitializeDefaultValue(TempStructBuffer.GetData());
 
 			{
 				// force logging to the error pipe for error detection
@@ -467,6 +468,8 @@ bool URigVMPin::IsValidDefaultValue(const FString& InDefaultValue) const
 				ScriptStruct->ImportText(*Value, TempStructBuffer.GetData(), nullptr, PPF_None, &ErrorPipe, ScriptStruct->GetName()); 
 			}
 
+			ScriptStruct->DestroyStruct(TempStructBuffer.GetData());
+
 			if (ErrorPipe.NumErrors > 0)
 			{
 				return false;
