Control Rig: Fix links when expanding / collapsing nodes

diff --git a/Engine/Source/Developer/RigVMDeveloper/Private/RigVMModel/RigVMController.cpp b/Engine/Source/Developer/RigVMDeveloper/Private/RigVMModel/RigVMController.cpp
--- a/Engine/Source/Developer/RigVMDeveloper/Private/RigVMModel/RigVMController.cpp
+++ b/Engine/Source/Developer/RigVMDeveloper/Private/RigVMModel/RigVMController.cpp
@@ -2604,7 +2604,8 @@ TArray<URigVMNode*> URigVMController::ExpandLibraryNode(URigVMLibraryNode* InNod
 	}
 
 	// g) remap all output / source pins and create a final list of links to create
-	TMap<FString, FString> RemappedSourcePins;;
+	TMap<FString, FString> RemappedSourcePinsForInputs;
+	TMap<FString, FString> RemappedSourcePinsForOutputs;
 	TArray<URigVMPin*> LibraryPins = InNode->GetAllPinsRecursively();
 	for (URigVMPin* LibraryPin : LibraryPins)
 	{
@@ -2612,8 +2613,27 @@ TArray<URigVMNode*> URigVMController::ExpandLibraryNode(URigVMLibraryNode* InNod
 		FString LibraryNodeName;
 		URigVMPin::SplitPinPathAtStart(LibraryPinPath, LibraryNodeName, LibraryPinPath);
 
-		FString SourcePinPath;
-		FString TargetPinPath;
+
+		struct Local
+		{
+			static void UpdateRemappedSourcePins(FString SourcePinPath, FString TargetPinPath, TMap<FString, FString>& RemappedSourcePins)
+			{
+				while (!SourcePinPath.IsEmpty() && !TargetPinPath.IsEmpty())
+				{
+					RemappedSourcePins.FindOrAdd(SourcePinPath) = TargetPinPath;
+
+					FString SourceLastSegment, TargetLastSegment;
+					if (!URigVMPin::SplitPinPathAtEnd(SourcePinPath, SourcePinPath, SourceLastSegment))
+					{
+						break;
+					}
+					if (!URigVMPin::SplitPinPathAtEnd(TargetPinPath, TargetPinPath, TargetLastSegment))
+					{
+						break;
+					}
+				}
+			}
+		};
 
 		if (LibraryPin->GetDirection() == ERigVMPinDirection::Input ||
 			LibraryPin->GetDirection() == ERigVMPinDirection::IO)
@@ -2623,34 +2643,18 @@ TArray<URigVMNode*> URigVMController::ExpandLibraryNode(URigVMLibraryNode* InNod
 				const TArray<FString>& LibraryPinLinks = *LibraryPinLinksPtr;
 				ensure(LibraryPinLinks.Num() == 1);
 
-				SourcePinPath = LibraryPinPath;
-				TargetPinPath = LibraryPinLinks[0];
+				Local::UpdateRemappedSourcePins(LibraryPinPath, LibraryPinLinks[0], RemappedSourcePinsForInputs);
 			}
 		}
-		else if (LibraryPin->GetDirection() == ERigVMPinDirection::Output)
+		if (LibraryPin->GetDirection() == ERigVMPinDirection::Output ||
+			LibraryPin->GetDirection() == ERigVMPinDirection::IO)
 		{
 			if (const TArray<FString>* LibraryPinLinksPtr = ToReturnNode.Find(LibraryPinPath))
 			{
 				const TArray<FString>& LibraryPinLinks = *LibraryPinLinksPtr;
 				ensure(LibraryPinLinks.Num() == 1);
 
-				SourcePinPath = LibraryPinPath;
-				TargetPinPath = LibraryPinLinks[0];
-			}
-		}
-
-		while (!SourcePinPath.IsEmpty() && !TargetPinPath.IsEmpty())
-		{
-			RemappedSourcePins.FindOrAdd(SourcePinPath) = TargetPinPath;
-
-			FString SourceLastSegment, TargetLastSegment;
-			if (!URigVMPin::SplitPinPathAtEnd(SourcePinPath, SourcePinPath, SourceLastSegment))
-			{
-				break;
-			}
-			if (!URigVMPin::SplitPinPathAtEnd(TargetPinPath, TargetPinPath, TargetLastSegment))
-			{
-				break;
+				Local::UpdateRemappedSourcePins(LibraryPinPath, LibraryPinLinks[0], RemappedSourcePinsForOutputs);
 			}
 		}
 	}
@@ -2693,7 +2697,7 @@ TArray<URigVMNode*> URigVMController::ExpandLibraryNode(URigVMLibraryNode* InNod
 		FString EntryPinPath = FromEntryNodePair.Key;
 		FString EntryPinPathSuffix;
 
-		const FString* RemappedSourcePin = RemappedSourcePins.Find(EntryPinPath);
+		const FString* RemappedSourcePin = RemappedSourcePinsForInputs.Find(EntryPinPath);
 		while (RemappedSourcePin == nullptr)
 		{
 			FString LastSegment;
@@ -2711,7 +2715,7 @@ TArray<URigVMNode*> URigVMController::ExpandLibraryNode(URigVMLibraryNode* InNod
 				EntryPinPathSuffix = URigVMPin::JoinPinPath(LastSegment, EntryPinPathSuffix);
 			}
 
-			RemappedSourcePin = RemappedSourcePins.Find(EntryPinPath);
+			RemappedSourcePin = RemappedSourcePinsForInputs.Find(EntryPinPath);
 		}
 
 		if (RemappedSourcePin == nullptr)
@@ -2792,7 +2796,7 @@ TArray<URigVMNode*> URigVMController::ExpandLibraryNode(URigVMLibraryNode* InNod
 		FString FromLibraryNodePinPath = FromLibraryNodePair.Key;
 		FString FromLibraryNodePinPathSuffix;
 
-		const FString* RemappedSourcePin = RemappedSourcePins.Find(FromLibraryNodePinPath);
+		const FString* RemappedSourcePin = RemappedSourcePinsForOutputs.Find(FromLibraryNodePinPath);
 		while (RemappedSourcePin == nullptr)
 		{
 			FString LastSegment;
@@ -2810,7 +2814,7 @@ TArray<URigVMNode*> URigVMController::ExpandLibraryNode(URigVMLibraryNode* InNod
 				FromLibraryNodePinPathSuffix = URigVMPin::JoinPinPath(LastSegment, FromLibraryNodePinPathSuffix);
 			}
 
-			RemappedSourcePin = RemappedSourcePins.Find(FromLibraryNodePinPath);
+			RemappedSourcePin = RemappedSourcePinsForOutputs.Find(FromLibraryNodePinPath);
 		}
 
 		if (RemappedSourcePin == nullptr)
@@ -2841,11 +2845,8 @@ TArray<URigVMNode*> URigVMController::ExpandLibraryNode(URigVMLibraryNode* InNod
 
 		for (const FString& FromLibraryNodeTargetPinPath : FromLibraryNodePair.Value)
 		{
-			FString TargetPinNodeName, TargetPinPath;
-			URigVMPin::SplitPinPathAtStart(FromLibraryNodeTargetPinPath, TargetPinNodeName, TargetPinPath);
-
 			URigVMPin* SourcePin = GetGraph()->FindPin(*RemappedSourcePinPath);
-			URigVMPin* TargetPin = GetGraph()->FindPin(URigVMPin::JoinPinPath(TargetPinNodeName, TargetPinPath));
+			URigVMPin* TargetPin = GetGraph()->FindPin(FromLibraryNodeTargetPinPath);
 			if (SourcePin && TargetPin)
 			{
 				if (!SourcePin->IsLinkedTo(TargetPin))
