Control Rig / Rigvm AST: Recurse along child pins from inside of library nodes to the outside pin for defaults

diff --git a/Engine/Source/Developer/RigVMDeveloper/Private/RigVMCompiler/RigVMAST.cpp b/Engine/Source/Developer/RigVMDeveloper/Private/RigVMCompiler/RigVMAST.cpp
--- a/Engine/Source/Developer/RigVMDeveloper/Private/RigVMCompiler/RigVMAST.cpp
+++ b/Engine/Source/Developer/RigVMDeveloper/Private/RigVMCompiler/RigVMAST.cpp
@@ -2656,7 +2656,8 @@ void FRigVMParserAST::Inline(URigVMGraph* InGraph, const TArray<FRigVMASTProxy>&
 			}
 
 			if (Pin->GetDirection() != ERigVMPinDirection::Input &&
-				Pin->GetDirection() != ERigVMPinDirection::IO)
+				Pin->GetDirection() != ERigVMPinDirection::IO &&
+				Pin->GetDirection() != ERigVMPinDirection::Output)
 			{
 				return FRigVMASTProxy();
 			}
@@ -2686,6 +2687,23 @@ void FRigVMParserAST::Inline(URigVMGraph* InGraph, const TArray<FRigVMASTProxy>&
 			URigVMPin* ChildPin = Pin;
 			while (ChildPin != nullptr)
 			{
+				if (ChildPin->GetDirection() == ERigVMPinDirection::Output && ChildPin->GetParentPin() == nullptr)
+				{
+					if (URigVMFunctionEntryNode* EntryNode = Cast<URigVMFunctionEntryNode>(Pin->GetNode()))
+					{
+						if (URigVMCollapseNode* OuterNode = Cast<URigVMCollapseNode>(EntryNode->GetGraph()->GetOuter()))
+						{
+							if (URigVMPin* OuterPin = OuterNode->FindPin(ChildPin->GetName()))
+							{
+								FRigVMASTProxy OuterPinProxy = InPinProxy.GetParent().GetSibling(OuterPin);
+								SourcePinProxy = FindSourcePin(OuterPinProxy, OutTraversalInfo);
+								SourcePinProxy = SourcePinProxy.IsValid() ? SourcePinProxy : OuterPinProxy;
+								break;
+							}
+						}
+					}
+				}
+
 				TArray<URigVMLink*> SourceLinks = ChildPin->GetSourceLinks(false /* recursive */);
 				if (SourceLinks.Num() > 0)
 				{
