Control Rig: Filter out entry / return nodes when collapsing node

diff --git a/Engine/Source/Developer/RigVMDeveloper/Private/RigVMModel/RigVMController.cpp b/Engine/Source/Developer/RigVMDeveloper/Private/RigVMModel/RigVMController.cpp
--- a/Engine/Source/Developer/RigVMDeveloper/Private/RigVMModel/RigVMController.cpp
+++ b/Engine/Source/Developer/RigVMDeveloper/Private/RigVMModel/RigVMController.cpp
@@ -1942,6 +1942,12 @@ URigVMCollapseNode* URigVMController::CollapseNodes(const TArray<URigVMNode*>& I
 			continue;
 		}
 
+		if (Node->IsA<URigVMFunctionEntryNode>() ||
+			Node->IsA<URigVMFunctionReturnNode>())
+		{
+			continue;
+		}
+
 		Nodes.Add(Node);
 	}
 
