Control Rig: Make sure to hook up an execute context pin within the collapse node

diff --git a/Engine/Source/Developer/RigVMDeveloper/Private/RigVMModel/RigVMController.cpp b/Engine/Source/Developer/RigVMDeveloper/Private/RigVMModel/RigVMController.cpp
--- a/Engine/Source/Developer/RigVMDeveloper/Private/RigVMModel/RigVMController.cpp
+++ b/Engine/Source/Developer/RigVMDeveloper/Private/RigVMModel/RigVMController.cpp
@@ -2344,6 +2344,72 @@ URigVMCollapseNode* URigVMController::CollapseNodes(const TArray<URigVMNode*>& I
 		RewiredLinks.Add(LinkToRewire);
 	}
 
+	if (ReturnNode)
+	{
+		struct Local
+		{
+			static bool IsLinkedToEntryNode(URigVMNode* InNode, TMap<URigVMNode*, bool>& CachedMap)
+			{
+				if (InNode->IsA<URigVMFunctionEntryNode>())
+				{
+					return true;
+				}
+
+				if (!CachedMap.Contains(InNode))
+				{
+					CachedMap.Add(InNode, false);
+
+					if (URigVMPin* ExecuteContextPin = InNode->FindPin(FRigVMStruct::ExecuteContextName.ToString()))
+					{
+						TArray<URigVMPin*> SourcePins = ExecuteContextPin->GetLinkedSourcePins();
+						for (URigVMPin* SourcePin : SourcePins)
+						{
+							if (IsLinkedToEntryNode(SourcePin->GetNode(), CachedMap))
+							{
+								CachedMap.FindOrAdd(InNode) = true;
+								break;
+							}
+						}
+					}
+				}
+
+				return CachedMap.FindChecked(InNode);
+			}
+		};
+
+		// check if there is a last node on the top level block what we need to hook up
+		TMap<URigVMNode*, bool> IsContainedNodeLinkedToEntryNode;
+		for (URigVMNode* ContainedNode : CollapseNode->GetContainedNodes())
+		{
+			if (!ContainedNode->IsMutable())
+			{
+				continue;
+			}
+
+			// make sure the node doesn't have any mutable nodes connected to its executecontext
+			if (URigVMPin* ExecuteContextPin = ContainedNode->FindPin(FRigVMStruct::ExecuteContextName.ToString()))
+			{
+				if (ExecuteContextPin->GetDirection() != ERigVMPinDirection::IO)
+				{
+					continue;
+				}
+
+				if (ExecuteContextPin->GetTargetLinks().Num() > 0)
+				{
+					continue;
+				}
+
+				if (!Local::IsLinkedToEntryNode(ContainedNode, IsContainedNodeLinkedToEntryNode))
+				{
+					continue;
+				}
+
+				FRigVMControllerGraphGuard GraphGuard(this, CollapseNode->GetContainedGraph(), false);
+				AddLink(ExecuteContextPin, ReturnNode->FindPin(FRigVMStruct::ExecuteContextName.ToString()), false);
+				break;
+			}
+		}
+	}
 
 	for (const FName& NodeToRemove : NodeNames)
 	{
