Control Rig: Entry / return nodes should always have the execute pins at the top

diff --git a/Engine/Source/Developer/RigVMDeveloper/Private/RigVMModel/RigVMController.cpp b/Engine/Source/Developer/RigVMDeveloper/Private/RigVMModel/RigVMController.cpp
--- a/Engine/Source/Developer/RigVMDeveloper/Private/RigVMModel/RigVMController.cpp
+++ b/Engine/Source/Developer/RigVMDeveloper/Private/RigVMModel/RigVMController.cpp
@@ -7403,7 +7403,24 @@ void URigVMController::RepopulatePinsOnNode(URigVMNode* InNode, bool bFollowCore
 			InNode->Pins.Reset();
 			Pins.Reset();
 
+			TArray<URigVMPin*> SortedLibraryPins;
+
+			// add execute pins first
+			for (URigVMPin* LibraryPin : LibraryNode->GetPins())
+			{
+				if (LibraryPin->IsExecuteContext())
+				{
+					SortedLibraryPins.Add(LibraryPin);
+				}
+			}
+
+			// add remaining pins
 			for (URigVMPin* LibraryPin : LibraryNode->GetPins())
+			{
+				SortedLibraryPins.AddUnique(LibraryPin);
+			}
+
+			for (URigVMPin* LibraryPin : SortedLibraryPins)
 			{
 				if (LibraryPin->GetDirection() == ERigVMPinDirection::IO && !LibraryPin->IsExecuteContext())
 				{
