Control Rig: Optimize usage of diagnostics AST and improve interactive performance for link validation

diff --git a/Engine/Source/Developer/RigVMDeveloper/Private/RigVMCompiler/RigVMAST.cpp b/Engine/Source/Developer/RigVMDeveloper/Private/RigVMCompiler/RigVMAST.cpp
--- a/Engine/Source/Developer/RigVMDeveloper/Private/RigVMCompiler/RigVMAST.cpp
+++ b/Engine/Source/Developer/RigVMDeveloper/Private/RigVMCompiler/RigVMAST.cpp
@@ -2567,6 +2567,8 @@ void FRigVMParserAST::Inline(URigVMGraph* InGraph)
 
 void FRigVMParserAST::Inline(URigVMGraph* InGraph, const TArray<FRigVMASTProxy>& InNodeProxies)
 {
+	DECLARE_SCOPE_HIERARCHICAL_COUNTER_FUNC()
+
 	struct LocalPinTraversalInfo
 	{
 		URigVMPin::FPinOverrideMap* PinOverrides;
diff --git a/Engine/Source/Developer/RigVMDeveloper/Private/RigVMModel/RigVMGraph.cpp b/Engine/Source/Developer/RigVMDeveloper/Private/RigVMModel/RigVMGraph.cpp
--- a/Engine/Source/Developer/RigVMDeveloper/Private/RigVMModel/RigVMGraph.cpp
+++ b/Engine/Source/Developer/RigVMDeveloper/Private/RigVMModel/RigVMGraph.cpp
@@ -231,6 +231,27 @@ void URigVMGraph::Notify(ERigVMGraphNotifType InNotifType, UObject* InSubject)
 
 TSharedPtr<FRigVMParserAST> URigVMGraph::GetDiagnosticsAST(bool bForceRefresh, TArray<URigVMLink*> InLinksToSkip)
 {
+	// only refresh the diagnostics AST if have a different set of links to skip
+	if (!bForceRefresh && DiagnosticsAST.IsValid())
+	{
+		const TArray<URigVMLink*> PreviousLinksToSkip = DiagnosticsAST->GetSettings().LinksToSkip;
+		if (PreviousLinksToSkip.Num() < InLinksToSkip.Num())
+		{
+			bForceRefresh = true;
+		}
+		else
+		{
+			for (int32 LinkIndex = 0; LinkIndex < InLinksToSkip.Num(); LinkIndex++)
+			{
+				if (PreviousLinksToSkip[LinkIndex] != InLinksToSkip[LinkIndex])
+				{
+					bForceRefresh = true;
+					break;
+				}
+			}
+		}
+	}
+
 	if (DiagnosticsAST == nullptr || bForceRefresh)
 	{
 		FRigVMParserASTSettings Settings = FRigVMParserASTSettings::Fast();
@@ -280,7 +301,8 @@ void URigVMGraph::PrepareCycleChecking(URigVMPin* InPin, bool bAsInput)
 	{
 		LinksToSkip = InPin->GetLinks();
 	}
-	GetDiagnosticsAST(LinksToSkip.Num() > 0, LinksToSkip)->PrepareCycleChecking(InPin);
+
+	GetDiagnosticsAST(false, LinksToSkip)->PrepareCycleChecking(InPin);
 }
 
 bool URigVMGraph::CanLink(URigVMPin* InSourcePin, URigVMPin* InTargetPin, FString* OutFailureReason)
diff --git a/Engine/Source/Developer/RigVMDeveloper/Public/RigVMCompiler/RigVMAST.h b/Engine/Source/Developer/RigVMDeveloper/Public/RigVMCompiler/RigVMAST.h
--- a/Engine/Source/Developer/RigVMDeveloper/Public/RigVMCompiler/RigVMAST.h
+++ b/Engine/Source/Developer/RigVMDeveloper/Public/RigVMCompiler/RigVMAST.h
@@ -1221,6 +1221,9 @@ public:
 	// returns the AST's override table for pin defaults
 	const URigVMPin::FPinOverrideMap& GetPinOverrides() const { return PinOverrides; }
 
+	// returns the settings used for this parser
+	const FRigVMParserASTSettings& GetSettings() const { return Settings; }
+
 private:
 
 	// private constructor for a partial build
