Fixed string usage issues found by UE_NODISCARD

#rnx

diff --git a/Engine/Plugins/Experimental/PythonScriptPlugin/Source/PythonScriptPlugin/Private/PyWrapperTypeRegistry.cpp b/Engine/Plugins/Experimental/PythonScriptPlugin/Source/PythonScriptPlugin/Private/PyWrapperTypeRegistry.cpp
--- a/Engine/Plugins/Experimental/PythonScriptPlugin/Source/PythonScriptPlugin/Private/PyWrapperTypeRegistry.cpp
+++ b/Engine/Plugins/Experimental/PythonScriptPlugin/Source/PythonScriptPlugin/Private/PyWrapperTypeRegistry.cpp
@@ -2299,7 +2299,7 @@ void FPyWrapperTypeRegistry::GenerateStubCodeForWrappedTypes(const EPyOnlineDocs
 		FString ProjectTopDir;
 		if (FPaths::IsProjectFilePathSet())
 		{
-			ProjectTopDir / FPaths::GetCleanFilename(FPaths::ProjectDir());
+			ProjectTopDir = FPaths::GetCleanFilename(FPaths::ProjectDir());
 		}
 
 		for (const auto& WrappedDataPair : WrappedData)
diff --git a/Engine/Plugins/Runtime/HTTPChunkInstaller/Source/Private/LocalTitleFile.cpp b/Engine/Plugins/Runtime/HTTPChunkInstaller/Source/Private/LocalTitleFile.cpp
--- a/Engine/Plugins/Runtime/HTTPChunkInstaller/Source/Private/LocalTitleFile.cpp
+++ b/Engine/Plugins/Runtime/HTTPChunkInstaller/Source/Private/LocalTitleFile.cpp
@@ -8,7 +8,7 @@
 FLocalTitleFile::FLocalTitleFile(const FString& InRootDirectory)
 : RootDirectory(InRootDirectory)
 {
-	RootDirectory.Replace(TEXT("\\"), TEXT("/"));
+	RootDirectory.ReplaceInline(TEXT("\\"), TEXT("/"));
 	if (!RootDirectory.EndsWith(TEXT("/")))
 	{
 		RootDirectory.Append(TEXT("/"));
diff --git a/Engine/Plugins/VirtualProduction/DMX/DMXEngine/Source/DMXRuntime/Private/DMXRuntimeUtils.cpp b/Engine/Plugins/VirtualProduction/DMX/DMXEngine/Source/DMXRuntime/Private/DMXRuntimeUtils.cpp
--- a/Engine/Plugins/VirtualProduction/DMX/DMXEngine/Source/DMXRuntime/Private/DMXRuntimeUtils.cpp
+++ b/Engine/Plugins/VirtualProduction/DMX/DMXEngine/Source/DMXRuntime/Private/DMXRuntimeUtils.cpp
@@ -30,7 +30,7 @@ bool FDMXRuntimeUtils::GetNameAndIndexFromString(const FString& InString, FStrin
 	}
 
 	// Remove separator characters at the end of the string
-	OutName.TrimEnd();
+	OutName.TrimEndInline();
 	DigitIndex = OutName.Len(); // reuse this variable for the separator index
 
 	while (DigitIndex > 0
diff --git a/Engine/Source/Developer/RigVMDeveloper/Private/RigVMModel/RigVMController.cpp b/Engine/Source/Developer/RigVMDeveloper/Private/RigVMModel/RigVMController.cpp
--- a/Engine/Source/Developer/RigVMDeveloper/Private/RigVMModel/RigVMController.cpp
+++ b/Engine/Source/Developer/RigVMDeveloper/Private/RigVMModel/RigVMController.cpp
@@ -7628,17 +7628,17 @@ FLinearColor URigVMController::GetColorFromMetadata(const FString& InMetadata)
 	FLinearColor Color = FLinearColor::Black;
 
 	FString Metadata = InMetadata;
-	Metadata.TrimStartAndEnd();
+	Metadata.TrimStartAndEndInline();
 	FString SplitString(TEXT(" "));
 	FString Red, Green, Blue, GreenAndBlue;
 	if (Metadata.Split(SplitString, &Red, &GreenAndBlue))
 	{
-		Red.TrimEnd();
-		GreenAndBlue.TrimStart();
+		Red.TrimEndInline();
+		GreenAndBlue.TrimStartInline();
 		if (GreenAndBlue.Split(SplitString, &Green, &Blue))
 		{
-			Green.TrimEnd();
-			Blue.TrimStart();
+			Green.TrimEndInline();
+			Blue.TrimStartInline();
 
 			float RedValue = FCString::Atof(*Red);
 			float GreenValue = FCString::Atof(*Green);
