Control Rig: Fix CIS issues, added missing header

diff --git a/Engine/Source/Developer/RigVMDeveloper/Public/RigVMCompiler/RigVMAST.h b/Engine/Source/Developer/RigVMDeveloper/Public/RigVMCompiler/RigVMAST.h
--- a/Engine/Source/Developer/RigVMDeveloper/Public/RigVMCompiler/RigVMAST.h
+++ b/Engine/Source/Developer/RigVMDeveloper/Public/RigVMCompiler/RigVMAST.h
@@ -8,6 +8,7 @@
 #include "RigVMModel/RigVMLink.h"
 #include "RigVMModel/RigVMNode.h"
 #include "RigVMASTProxy.h"
+#include "Logging/TokenizedMessage.h"
 #include "RigVMAST.generated.h"
 
 DECLARE_DELEGATE_ThreeParams(FRigVMReportDelegate, EMessageSeverity::Type, UObject*, const FString&);
