Control Rig: Improve filtering of block assignments for Expressions

diff --git a/Engine/Source/Developer/RigVMDeveloper/Private/RigVMCompiler/RigVMAST.cpp b/Engine/Source/Developer/RigVMDeveloper/Private/RigVMCompiler/RigVMAST.cpp
--- a/Engine/Source/Developer/RigVMDeveloper/Private/RigVMCompiler/RigVMAST.cpp
+++ b/Engine/Source/Developer/RigVMDeveloper/Private/RigVMCompiler/RigVMAST.cpp
@@ -2208,6 +2208,42 @@ bool FRigVMParserAST::CanLink(URigVMPin* InSourcePin, URigVMPin* InTargetPin, FS
 		return false;
 	}
 
+	if(SourceNode->IsA<URigVMRerouteNode>())
+	{
+		TArray<URigVMNode*> LinkedSourceNodes = SourceNode->GetLinkedSourceNodes();
+		for (int32 LinkedSourceNodeIndex = 0; LinkedSourceNodeIndex<LinkedSourceNodes.Num(); LinkedSourceNodeIndex++)
+		{
+			URigVMNode* LinkedSourceNode = LinkedSourceNodes[LinkedSourceNodeIndex];
+			if(LinkedSourceNode->IsA<URigVMRerouteNode>())
+			{
+				LinkedSourceNodes.Append(LinkedSourceNode->GetLinkedSourceNodes());
+			}
+			else
+			{
+				SourceNode = LinkedSourceNode;
+				break;
+			}
+		}
+	}
+
+	if(TargetNode->IsA<URigVMRerouteNode>())
+	{
+		TArray<URigVMNode*> LinkedTargetNodes = TargetNode->GetLinkedTargetNodes();
+		for (int32 LinkedTargetNodeIndex = 0; LinkedTargetNodeIndex<LinkedTargetNodes.Num(); LinkedTargetNodeIndex++)
+		{
+			URigVMNode* LinkedTargetNode = LinkedTargetNodes[LinkedTargetNodeIndex];
+			if(LinkedTargetNode->IsA<URigVMRerouteNode>())
+			{
+				LinkedTargetNodes.Append(LinkedTargetNode->GetLinkedTargetNodes());
+			}
+			else
+			{
+				TargetNode = LinkedTargetNode;
+				break;
+			}
+		}
+	}
+
 	FRigVMASTProxy SourceNodeProxy = FRigVMASTProxy::MakeFromUObject(SourceNode);
 	FRigVMASTProxy TargetNodeProxy = FRigVMASTProxy::MakeFromUObject(TargetNode);
 
