Control Rig: now linking a pin with a injected node to a reroute node works

diff --git a/Engine/Source/Developer/RigVMDeveloper/Private/RigVMModel/RigVMController.cpp b/Engine/Source/Developer/RigVMDeveloper/Private/RigVMModel/RigVMController.cpp
--- a/Engine/Source/Developer/RigVMDeveloper/Private/RigVMModel/RigVMController.cpp
+++ b/Engine/Source/Developer/RigVMDeveloper/Private/RigVMModel/RigVMController.cpp
@@ -1119,7 +1119,9 @@ URigVMRerouteNode* URigVMController::AddRerouteNodeOnPin(const FString& InPinPat
 		ActionStack->BeginAction(Action);
 	}
 
-	BreakAllLinks(Pin, bAsInput, bSetupUndoRedo);
+	//in case an injected node is present, use its pins for any new links
+	URigVMPin *PinForLink = Pin->GetPinForLink(); 
+	BreakAllLinks(PinForLink, bAsInput, bSetupUndoRedo);
 
 	FString Name = GetValidNodeName(InNodeName.IsEmpty() ? FString(TEXT("RerouteNode")) : InNodeName);
 	URigVMRerouteNode* Node = NewObject<URigVMRerouteNode>(Graph, *Name);
@@ -1161,11 +1163,11 @@ URigVMRerouteNode* URigVMController::AddRerouteNodeOnPin(const FString& InPinPat
 
 	if (bAsInput)
 	{
-		AddLink(ValuePin, Pin, bSetupUndoRedo);
+		AddLink(ValuePin, PinForLink, bSetupUndoRedo);
 	}
 	else
 	{
-		AddLink(Pin, ValuePin, bSetupUndoRedo);
+		AddLink(PinForLink, ValuePin, bSetupUndoRedo);
 	}
 
 	if (bSetupUndoRedo)
