Downgraded a RigVM error to a warning, pending further eval by anim team.

diff --git a/Engine/Source/Developer/RigVMDeveloper/Private/RigVMModel/RigVMController.cpp b/Engine/Source/Developer/RigVMDeveloper/Private/RigVMModel/RigVMController.cpp
--- a/Engine/Source/Developer/RigVMDeveloper/Private/RigVMModel/RigVMController.cpp
+++ b/Engine/Source/Developer/RigVMDeveloper/Private/RigVMModel/RigVMController.cpp
@@ -6654,7 +6654,7 @@ bool URigVMController::IsValidNodeForGraph(URigVMNode* InNode)
 
 	if (InNode->GetGraph() != GetGraph())
 	{
-		ReportErrorf(TEXT("InNode '%s' is on a different graph."), *InNode->GetNodePath());
+		ReportWarningf(TEXT("InNode '%s' is on a different graph. InNode graph is %s, this graph is %s"), *InNode->GetNodePath(), *GetNameSafe(InNode->GetGraph()), *GetNameSafe(GetGraph()));
 		return false;
 	}
 
