Control Rig: Integration suggested change for  temporarily disabling RigVM notifications

diff --git a/Engine/Source/Developer/RigVMDeveloper/Public/RigVMModel/RigVMController.h b/Engine/Source/Developer/RigVMDeveloper/Public/RigVMModel/RigVMController.h
--- a/Engine/Source/Developer/RigVMDeveloper/Public/RigVMModel/RigVMController.h
+++ b/Engine/Source/Developer/RigVMDeveloper/Public/RigVMModel/RigVMController.h
@@ -597,6 +597,12 @@ public:
 		ReportAndNotifyError(FString::Printf(Fmt, Args...));
 	}
 
+	/**
+	 * Function to override the notification behavior and temporarily
+	 * disable all notifications. Client code is responsible for calling
+	 * SuspendNotifications(true) once all changes have been done.
+	 */
+	void SuspendNotifications(bool bSuspend) { bSuspendNotifications = bSuspend; }
 
 private:
 
