Control Rig: Disallow pasting function references into other graphs

diff --git a/Engine/Plugins/Experimental/ControlRig/Source/ControlRigDeveloper/Private/ControlRigBlueprint.cpp b/Engine/Plugins/Experimental/ControlRig/Source/ControlRigDeveloper/Private/ControlRigBlueprint.cpp
--- a/Engine/Plugins/Experimental/ControlRig/Source/ControlRigDeveloper/Private/ControlRigBlueprint.cpp
+++ b/Engine/Plugins/Experimental/ControlRig/Source/ControlRigDeveloper/Private/ControlRigBlueprint.cpp
@@ -80,6 +80,7 @@ UControlRigBlueprint::UControlRigBlueprint(const FObjectInitializer& ObjectIniti
 	FunctionLibraryEdGraph->Initialize(this);
 
 	Model->SetDefaultFunctionLibrary(FunctionLibrary);
+	FunctionLibrary->SetDefaultFunctionLibrary(FunctionLibrary);
 
 	Validator = ObjectInitializer.CreateDefaultSubobject<UControlRigValidator>(this, TEXT("ControlRigValidator"));
 
diff --git a/Engine/Source/Developer/RigVMDeveloper/Private/RigVMModel/RigVMController.cpp b/Engine/Source/Developer/RigVMDeveloper/Private/RigVMModel/RigVMController.cpp
--- a/Engine/Source/Developer/RigVMDeveloper/Private/RigVMModel/RigVMController.cpp
+++ b/Engine/Source/Developer/RigVMDeveloper/Private/RigVMModel/RigVMController.cpp
@@ -1755,6 +1755,11 @@ TArray<FName> URigVMController::ImportNodesFromText(const FString& InText, bool
 	FRigVMUnitNodeCreatedContext::FScope UnitNodeCreatedScope(UnitNodeCreatedContext, ERigVMNodeCreatedReason::Paste);
 	for (URigVMNode* CreatedNode : Factory.CreatedNodes)
 	{
+		if(!CanAddNode(CreatedNode, true))
+		{
+			continue;
+		}
+
 		Graph->Nodes.Add(CreatedNode);
 
 		if (bSetupUndoRedo)
@@ -5902,15 +5907,9 @@ URigVMFunctionReferenceNode* URigVMController::AddFunctionReferenceNode(URigVMLi
 		return nullptr;
 	}
 
-	URigVMLibraryNode* ParentLibraryNode = Cast<URigVMLibraryNode>(Graph->GetOuter());
-	while (ParentLibraryNode)
+	if(!CanAddFunctionRefForDefinition(InFunctionDefinition, true))
 	{
-		if (ParentLibraryNode == InFunctionDefinition)
-		{
-			ReportAndNotifyError(TEXT("You cannot place functions inside of itself or an indirect recursion."));
-			return nullptr;
-		}
-		ParentLibraryNode = Cast<URigVMLibraryNode>(ParentLibraryNode->GetGraph()->GetOuter());
+		return nullptr;
 	}
 
 	FString NodeName = GetValidNodeName(InNodeName.IsEmpty() ? InFunctionDefinition->GetName() : InNodeName);
@@ -6741,6 +6740,75 @@ bool URigVMController::IsValidLinkForGraph(URigVMLink* InLink)
 	return true;
 }
 
+bool URigVMController::CanAddNode(URigVMNode* InNode, bool bReportErrors)
+{
+	if(!IsValidGraph())
+	{
+		return false;
+	}
+
+	check(InNode);
+
+	URigVMGraph* Graph = GetGraph();
+	check(Graph);
+
+	if (URigVMFunctionReferenceNode* FunctionRefNode = Cast<URigVMFunctionReferenceNode>(InNode))
+	{
+		if (URigVMFunctionLibrary* FunctionLibrary = FunctionRefNode->GetLibrary())
+		{
+			if(Graph->GetDefaultFunctionLibrary() != FunctionLibrary)
+			{
+				// this node references a function that is not part of our function library
+				if(bReportErrors)
+				{
+					ReportError(TEXT("Cannot import function reference node."));
+				}
+				DestroyObject(InNode);
+				return false;
+			}
+			else if (URigVMLibraryNode* FunctionDefinition = FunctionRefNode->GetReferencedNode())
+			{
+				if(!CanAddFunctionRefForDefinition(FunctionDefinition, bReportErrors))
+				{
+					DestroyObject(InNode);
+					return false;
+				}
+			}
+		}			
+	}
+
+	return true;
+}
+
+bool URigVMController::CanAddFunctionRefForDefinition(URigVMLibraryNode* InFunctionDefinition, bool bReportErrors)
+{
+	if(!IsValidGraph())
+	{
+		return false;
+	}
+
+	check(InFunctionDefinition);
+
+	URigVMGraph* Graph = GetGraph();
+	check(Graph);
+
+	URigVMLibraryNode* ParentLibraryNode = Cast<URigVMLibraryNode>(Graph->GetOuter());
+	while (ParentLibraryNode)
+	{
+		if (ParentLibraryNode == InFunctionDefinition)
+		{
+			if(bReportErrors)
+			{
+				ReportAndNotifyError(TEXT("You cannot place functions inside of itself or an indirect recursion."));
+			}
+			return false;
+		}
+		ParentLibraryNode = Cast<URigVMLibraryNode>(ParentLibraryNode->GetGraph()->GetOuter());
+	}
+
+	return true;
+}
+
 void URigVMController::AddPinsForStruct(UStruct* InStruct, URigVMNode* InNode, URigVMPin* InParentPin, ERigVMPinDirection InPinDirection, const FString& InDefaultValue, bool bAutoExpandArrays, bool bNotify)
 {
 	TArray<FString> MemberNameValuePairs = URigVMPin::SplitDefaultValue(InDefaultValue);
diff --git a/Engine/Source/Developer/RigVMDeveloper/Public/RigVMModel/RigVMController.h b/Engine/Source/Developer/RigVMDeveloper/Public/RigVMModel/RigVMController.h
--- a/Engine/Source/Developer/RigVMDeveloper/Public/RigVMModel/RigVMController.h
+++ b/Engine/Source/Developer/RigVMDeveloper/Public/RigVMModel/RigVMController.h
@@ -617,6 +617,8 @@ private:
 	bool IsValidNodeForGraph(URigVMNode* InNode);
 	bool IsValidPinForGraph(URigVMPin* InPin);
 	bool IsValidLinkForGraph(URigVMLink* InLink);
+	bool CanAddNode(URigVMNode* InNode, bool bReportErrors);
+	bool CanAddFunctionRefForDefinition(URigVMLibraryNode* InFunctionDefinition, bool bReportErrors);
 	void AddPinsForStruct(UStruct* InStruct, URigVMNode* InNode, URigVMPin* InParentPin, ERigVMPinDirection InPinDirection, const FString& InDefaultValue, bool bAutoExpandArrays, bool bNotify = false);
 	void AddPinsForArray(FArrayProperty* InArrayProperty, URigVMNode* InNode, URigVMPin* InParentPin, ERigVMPinDirection InPinDirection, const TArray<FString>& InDefaultValues, bool bAutoExpandArrays);
 	void ConfigurePinFromProperty(FProperty* InProperty, URigVMPin* InOutPin, ERigVMPinDirection InPinDirection = ERigVMPinDirection::Invalid);
