Control Rig: Implementing localization functionality to create local copies of functions

diff --git a/Engine/Plugins/Experimental/ControlRig/Source/ControlRigDeveloper/Private/ControlRigBlueprint.cpp b/Engine/Plugins/Experimental/ControlRig/Source/ControlRigDeveloper/Private/ControlRigBlueprint.cpp
--- a/Engine/Plugins/Experimental/ControlRig/Source/ControlRigDeveloper/Private/ControlRigBlueprint.cpp
+++ b/Engine/Plugins/Experimental/ControlRig/Source/ControlRigDeveloper/Private/ControlRigBlueprint.cpp
@@ -188,6 +188,17 @@ bool UControlRigBlueprint::TryImportGraphFromText(const FString& InClipboardText
 
 	if (URigVMController* FunctionLibraryController = GetOrCreateController(GetLocalFunctionLibrary()))
 	{
+		TGuardValue<FRigVMController_RequestLocalizeFunctionDelegate> RequestLocalizeDelegateGuard(
+            FunctionLibraryController->RequestLocalizeFunctionDelegate,
+            FRigVMController_RequestLocalizeFunctionDelegate::CreateLambda([this](URigVMLibraryNode* InFunctionToLocalize)
+            {
+            	BroadcastRequestLocalizeFunctionDialog(InFunctionToLocalize);
+
+                const URigVMLibraryNode* LocalizedFunctionNode = GetLocalFunctionLibrary()->FindPreviouslyLocalizedFunction(InFunctionToLocalize);
+                return LocalizedFunctionNode != nullptr;
+            })
+        );
+		
 		TArray<FName> ImportedNodeNames = FunctionLibraryController->ImportNodesFromText(InClipboardText, true);
 		if (ImportedNodeNames.Num() == 0)
 		{
@@ -885,6 +896,36 @@ URigVMController* UControlRigBlueprint::GetOrCreateController(URigVMGraph* InGra
 
 	});
 
+	Controller->IsFunctionAvailableDelegate.BindLambda([WeakThis](URigVMLibraryNode* InFunction) -> bool
+	{
+		if(InFunction == nullptr)
+		{
+			return false;
+		}
+		
+		if(URigVMFunctionLibrary* Library = Cast<URigVMFunctionLibrary>(InFunction->GetOuter()))
+		{
+			if(UControlRigBlueprint* Blueprint = Cast<UControlRigBlueprint>(Library->GetOuter()))
+			{
+				if(Blueprint->IsFunctionPublic(InFunction->GetFName()))
+				{
+					return true;
+				}
+
+				// if it is private - we still see it as public if we are within the same blueprint 
+				if(WeakThis.IsValid())
+				{
+					if(WeakThis.Get() == Blueprint)
+					{
+						return true;
+					}
+				}
+			}
+		}
+
+		return false;
+	});
+
 #if WITH_EDITOR
 
 	// this sets up three delegates:
@@ -1000,6 +1041,60 @@ UEdGraph* UControlRigBlueprint::GetEdGraph(const FString& InNodePath) const
 	return nullptr;
 }
 
+bool UControlRigBlueprint::IsFunctionPublic(const FName& InFunctionName) const
+{
+	for(const FControlRigPublicFunctionData& PublicFunction : PublicFunctions)
+	{
+		if(PublicFunction.Name == InFunctionName)
+		{
+			return true;
+		}
+	}
+	
+	return false;
+}
+
+void UControlRigBlueprint::MarkFunctionPublic(const FName& InFunctionName, bool bIsPublic)
+{
+	if(IsFunctionPublic(InFunctionName) == bIsPublic)
+	{
+		return;
+	}
+	
+	Modify();
+
+	if(bIsPublic)
+	{
+		if(URigVMLibraryNode* FunctionNode = GetLocalFunctionLibrary()->FindFunction(InFunctionName))
+		{
+			if(UControlRigGraph* RigGraph = Cast<UControlRigGraph>(GetEdGraph(FunctionNode->GetContainedGraph())))
+			{
+				const FControlRigPublicFunctionData NewFunctionData = RigGraph->GetPublicFunctionData();
+				for(FControlRigPublicFunctionData& ExistingFunctionData : PublicFunctions)
+				{
+					if(ExistingFunctionData.Name == NewFunctionData.Name)
+					{
+						ExistingFunctionData = NewFunctionData;
+						return;
+					}
+				}
+				PublicFunctions.Add(NewFunctionData);
+			}
+		}
+	}
+	else
+	{
+		for(int32 Index = 0; Index < PublicFunctions.Num(); Index++)
+		{
+			if(PublicFunctions[Index].Name == InFunctionName)
+			{
+				PublicFunctions.RemoveAt(Index);
+				return;
+			}
+		}
+	}
+}
+
 void UControlRigBlueprint::GetTypeActions(FBlueprintActionDatabaseRegistrar& ActionRegistrar) const
 {
 	DECLARE_SCOPE_HIERARCHICAL_COUNTER_FUNC()
@@ -2671,6 +2766,11 @@ void UControlRigBlueprint::BroadcastPostEditChangeChainProperty(FPropertyChanged
 	PostEditChangeChainPropertyEvent.Broadcast(PropertyChangedChainEvent);
 }
 
+void UControlRigBlueprint::BroadcastRequestLocalizeFunctionDialog(URigVMLibraryNode* InFunction, bool bForce)
+{
+	RequestLocalizeFunctionDialog.Broadcast(InFunction, this, bForce);
+}
+
 #endif
 
 void UControlRigBlueprint::CreateEdGraphForCollapseNodeIfNeeded(URigVMCollapseNode* InNode, bool bForce)
diff --git a/Engine/Plugins/Experimental/ControlRig/Source/ControlRigDeveloper/Private/Graph/ControlRigGraphSchema.cpp b/Engine/Plugins/Experimental/ControlRig/Source/ControlRigDeveloper/Private/Graph/ControlRigGraphSchema.cpp
--- a/Engine/Plugins/Experimental/ControlRig/Source/ControlRigDeveloper/Private/Graph/ControlRigGraphSchema.cpp
+++ b/Engine/Plugins/Experimental/ControlRig/Source/ControlRigDeveloper/Private/Graph/ControlRigGraphSchema.cpp
@@ -24,6 +24,10 @@
 #include "RigVMModel/Nodes/RigVMFunctionEntryNode.h"
 #include "RigVMModel/Nodes/RigVMFunctionReturnNode.h"
 
+#if WITH_EDITOR
+#include "ControlRigEditor/Private/Editor/SControlRigFunctionLocalizationWidget.h"
+#endif
+
 #define LOCTEXT_NAMESPACE "ControlRigGraphSchema"
 
 const FName UControlRigGraphSchema::GraphName_ControlRig(TEXT("Rig Graph"));
@@ -40,7 +44,23 @@ FReply FControlRigFunctionDragDropAction::DroppedOnPanel(const TSharedRef< class
 				{
 					if(URigVMController* TargetController = TargetRigBlueprint->GetController(TargetRigGraph))
 					{
-						TargetController->AddFunctionReferenceNode(FunctionDefinitionNode, GraphPosition);
+						if(URigVMFunctionLibrary* FunctionLibrary = Cast<URigVMFunctionLibrary>(FunctionDefinitionNode->GetOuter()))
+						{
+							if(UControlRigBlueprint* FunctionRigBlueprint = Cast<UControlRigBlueprint>(FunctionLibrary->GetOuter()))
+							{
+#if WITH_EDITOR
+								if(FunctionRigBlueprint != TargetRigBlueprint)
+								{
+									if(!FunctionRigBlueprint->IsFunctionPublic(FunctionDefinitionNode->GetFName()))
+									{
+										TargetRigBlueprint->BroadcastRequestLocalizeFunctionDialog(FunctionDefinitionNode);
+										FunctionDefinitionNode = TargetRigBlueprint->GetLocalFunctionLibrary()->FindPreviouslyLocalizedFunction(FunctionDefinitionNode);
+									}
+								}
+#endif
+								TargetController->AddFunctionReferenceNode(FunctionDefinitionNode, GraphPosition);
+							}
+						}
 					}
 				}
 			}
diff --git a/Engine/Plugins/Experimental/ControlRig/Source/ControlRigDeveloper/Public/ControlRigBlueprint.h b/Engine/Plugins/Experimental/ControlRig/Source/ControlRigDeveloper/Public/ControlRigBlueprint.h
--- a/Engine/Plugins/Experimental/ControlRig/Source/ControlRigDeveloper/Public/ControlRigBlueprint.h
+++ b/Engine/Plugins/Experimental/ControlRig/Source/ControlRigDeveloper/Public/ControlRigBlueprint.h
@@ -37,7 +37,7 @@ DECLARE_EVENT_OneParam(UControlRigBlueprint, FOnExternalVariablesChanged, const
 DECLARE_EVENT_TwoParams(UControlRigBlueprint, FOnNodeDoubleClicked, UControlRigBlueprint*, URigVMNode*);
 DECLARE_EVENT_OneParam(UControlRigBlueprint, FOnGraphImported, UEdGraph*);
 DECLARE_EVENT_OneParam(UControlRigBlueprint, FOnPostEditChangeChainProperty, FPropertyChangedChainEvent&);
-
+DECLARE_EVENT_ThreeParams(UControlRigBlueprint, FOnLocalizeFunctionDialogRequested, URigVMLibraryNode*, UControlRigBlueprint*, bool);
 
 USTRUCT()
 struct CONTROLRIGDEVELOPER_API FControlRigPublicFunctionArg
@@ -203,6 +203,9 @@ public:
 	UControlRigGraph* FunctionLibraryEdGraph;
 #endif
 
+	bool IsFunctionPublic(const FName& InFunctionName) const;
+	void MarkFunctionPublic(const FName& InFunctionName, bool bIsPublic = true);
+
 	UPROPERTY(EditAnywhere, BlueprintReadWrite, Category = "VM")
 	FRigVMCompileSettings VMCompileSettings;
 
@@ -441,6 +444,9 @@ public:
 	FOnPostEditChangeChainProperty& OnPostEditChangeChainProperty() { return PostEditChangeChainPropertyEvent; }
 	void BroadcastPostEditChangeChainProperty(FPropertyChangedChainEvent& PropertyChangedChainEvent);
 
+	FOnLocalizeFunctionDialogRequested& OnRequestLocalizeFunctionDialog() { return RequestLocalizeFunctionDialog; }
+	void BroadcastRequestLocalizeFunctionDialog(URigVMLibraryNode* InFunction, bool bForce = false);
+
 private:
 
 	FOnExternalVariablesChanged ExternalVariablesChangedEvent;
@@ -449,6 +455,7 @@ private:
 	FOnNodeDoubleClicked NodeDoubleClickedEvent;
 	FOnGraphImported GraphImportedEvent;
 	FOnPostEditChangeChainProperty PostEditChangeChainPropertyEvent;
+	FOnLocalizeFunctionDialogRequested RequestLocalizeFunctionDialog;
 
 #endif
 
diff --git a/Engine/Plugins/Experimental/ControlRig/Source/ControlRigEditor/Private/ControlRigEditorModule.cpp b/Engine/Plugins/Experimental/ControlRig/Source/ControlRigEditor/Private/ControlRigEditorModule.cpp
--- a/Engine/Plugins/Experimental/ControlRig/Source/ControlRigEditor/Private/ControlRigEditorModule.cpp
+++ b/Engine/Plugins/Experimental/ControlRig/Source/ControlRigEditor/Private/ControlRigEditorModule.cpp
@@ -1912,6 +1912,22 @@ void FControlRigEditorModule::GetContextMenuActions(const UControlRigGraphSchema
 						));
 					}
 
+					if (URigVMFunctionReferenceNode* FunctionReferenceNode = Cast<URigVMFunctionReferenceNode>(RigNode->GetModelNode()))
+					{
+						if(FunctionReferenceNode->GetLibrary() != RigBlueprint->GetLocalFunctionLibrary())
+						{
+							OrganizationSection.AddMenuEntry(
+                                "Localize Function",
+                                LOCTEXT("LocalizeFunction", "Localize Function"),
+                                LOCTEXT("LocalizeFunction_Tooltip", "Creates a local copy of the function backing the node."),
+                                FSlateIcon(),
+                                FUIAction(FExecuteAction::CreateLambda([RigBlueprint, FunctionReferenceNode]() {
+                                    RigBlueprint->BroadcastRequestLocalizeFunctionDialog(FunctionReferenceNode->GetReferencedNode(), true);
+                                })
+                            ));
+						}
+					}
+
 					OrganizationSection.AddSubMenu("Alignment", LOCTEXT("AlignmentHeader", "Alignment"), FText(), FNewToolMenuDelegate::CreateLambda([](UToolMenu* AlignmentMenu)
 					{
 						{
diff --git a/Engine/Plugins/Experimental/ControlRig/Source/ControlRigEditor/Private/ControlRigGraphDetails.cpp b/Engine/Plugins/Experimental/ControlRig/Source/ControlRigEditor/Private/ControlRigGraphDetails.cpp
--- a/Engine/Plugins/Experimental/ControlRig/Source/ControlRigEditor/Private/ControlRigGraphDetails.cpp
+++ b/Engine/Plugins/Experimental/ControlRig/Source/ControlRigEditor/Private/ControlRigGraphDetails.cpp
@@ -940,15 +940,10 @@ FText FControlRigGraphDetails::GetCurrentAccessSpecifierName() const
 		UControlRigBlueprint* ControlRigBlueprint = ControlRigBlueprintPtr.Get();
 
 		const FControlRigPublicFunctionData ExpectedFunctionData = Graph->GetPublicFunctionData();
-		for(const FControlRigPublicFunctionData& PublicFunction : ControlRigBlueprint->PublicFunctions)
+		if(ControlRigBlueprint->IsFunctionPublic(ExpectedFunctionData.Name))
 		{
-			if(PublicFunction.Name == ExpectedFunctionData.Name)
-			{
-				return FText::FromString(*AccessSpecifierStrings[0].Get()); // public
-			}
+			return FText::FromString(*AccessSpecifierStrings[0].Get()); // public
 		}
-
-		return FText::FromString(*AccessSpecifierStrings[1].Get()); // private
 	}
 
 	return FText::FromString(*AccessSpecifierStrings[1].Get()); // private
@@ -960,31 +955,15 @@ void FControlRigGraphDetails::OnAccessSpecifierSelected( TSharedPtr<FString> Spe
 	{
 		UControlRigGraph* Graph = GraphPtr.Get();
 		UControlRigBlueprint* ControlRigBlueprint = ControlRigBlueprintPtr.Get();
-		ControlRigBlueprint->Modify();
-
+		const FControlRigPublicFunctionData ExpectedFunctionData = Graph->GetPublicFunctionData();
+		
 		if(SpecifierName->Equals(TEXT("Private")))
 		{
-			for(int32 Index = 0; Index < ControlRigBlueprint->PublicFunctions.Num(); Index++)
-			{
-				if(ControlRigBlueprint->PublicFunctions[Index].Name == Graph->GetFName())
-				{
-					ControlRigBlueprint->PublicFunctions.RemoveAt(Index);
-					return;
-				}
-			}
+			ControlRigBlueprint->MarkFunctionPublic(ExpectedFunctionData.Name, false);
 		}
 		else
 		{
-			const FControlRigPublicFunctionData NewFunctionData = Graph->GetPublicFunctionData();
-			for(FControlRigPublicFunctionData& ExistingFunctionData : ControlRigBlueprint->PublicFunctions)
-			{
-				if(ExistingFunctionData.Name == NewFunctionData.Name)
-				{
-					ExistingFunctionData = NewFunctionData;
-					return;
-				}
-			}
-			ControlRigBlueprint->PublicFunctions.Add(NewFunctionData);
+			ControlRigBlueprint->MarkFunctionPublic(ExpectedFunctionData.Name, true);
 		}
 	}
 }
diff --git a/Engine/Plugins/Experimental/ControlRig/Source/ControlRigEditor/Private/Editor/ControlRigEditor.cpp b/Engine/Plugins/Experimental/ControlRig/Source/ControlRigEditor/Private/Editor/ControlRigEditor.cpp
--- a/Engine/Plugins/Experimental/ControlRig/Source/ControlRigEditor/Private/Editor/ControlRigEditor.cpp
+++ b/Engine/Plugins/Experimental/ControlRig/Source/ControlRigEditor/Private/Editor/ControlRigEditor.cpp
@@ -78,6 +78,7 @@
 #include "RigVMModel/Nodes/RigVMFunctionEntryNode.h"
 #include "RigVMModel/Nodes/RigVMFunctionReturnNode.h"
 #include "BlueprintActionDatabaseRegistrar.h"
+#include "SControlRigFunctionLocalizationWidget.h"
 
 #define LOCTEXT_NAMESPACE "ControlRigEditor"
 
@@ -126,6 +127,7 @@ FControlRigEditor::~FControlRigEditor()
 		RigBlueprint->OnNodeDoubleClicked().RemoveAll(this);
 		RigBlueprint->OnGraphImported().RemoveAll(this);
 		RigBlueprint->OnPostEditChangeChainProperty().RemoveAll(this);
+		RigBlueprint->OnRequestLocalizeFunctionDialog().RemoveAll(this);
 	}
 
 	if (NodeDetailBuffer.Num() > 0 && NodeDetailStruct != nullptr)
@@ -352,6 +354,7 @@ void FControlRigEditor::InitControlRigEditor(const EToolkitMode::Type Mode, cons
 		InControlRigBlueprint->OnNodeDoubleClicked().AddSP(this, &FControlRigEditor::OnNodeDoubleClicked);
 		InControlRigBlueprint->OnGraphImported().AddSP(this, &FControlRigEditor::OnGraphImported);
 		InControlRigBlueprint->OnPostEditChangeChainProperty().AddSP(this, &FControlRigEditor::OnBlueprintPropertyChainEvent);
+		InControlRigBlueprint->OnRequestLocalizeFunctionDialog().AddSP(this, &FControlRigEditor::OnRequestLocalizeFunctionDialog);
 	}
 
 	UpdateStaleWatchedPins();
@@ -1579,6 +1582,18 @@ void FControlRigEditor::PasteNodes()
 
 	FString TextToImport;
 	FPlatformApplicationMisc::ClipboardPaste(TextToImport);
+
+	TGuardValue<FRigVMController_RequestLocalizeFunctionDelegate> RequestLocalizeDelegateGuard(
+		GetFocusedController()->RequestLocalizeFunctionDelegate,
+		FRigVMController_RequestLocalizeFunctionDelegate::CreateLambda([this](URigVMLibraryNode* InFunctionToLocalize)
+		{
+			OnRequestLocalizeFunctionDialog(InFunctionToLocalize, GetControlRigBlueprint(), true);
+
+			const URigVMLibraryNode* LocalizedFunctionNode = GetControlRigBlueprint()->GetLocalFunctionLibrary()->FindPreviouslyLocalizedFunction(InFunctionToLocalize);
+			return LocalizedFunctionNode != nullptr;
+		})
+	);
+	
 	TArray<FName> NodeNames = GetFocusedController()->ImportNodesFromText(TextToImport);
 
 	if (NodeNames.Num() > 0)
@@ -3216,6 +3231,41 @@ void FControlRigEditor::OnBlueprintPropertyChainEvent(FPropertyChangedChainEvent
 	}
 }
 
+void FControlRigEditor::OnRequestLocalizeFunctionDialog(URigVMLibraryNode* InFunction, UControlRigBlueprint* InTargetBlueprint, bool bForce)
+{
+	check(InFunction);
+	check(InTargetBlueprint);
+
+	if(InTargetBlueprint != GetControlRigBlueprint())
+	{
+		return;
+	}
+	
+	if(URigVMController* TargetController = InTargetBlueprint->GetController(InTargetBlueprint->GetModel()))
+	{
+		if(URigVMFunctionLibrary* FunctionLibrary = Cast<URigVMFunctionLibrary>(InFunction->GetOuter()))
+		{
+			if(UControlRigBlueprint* FunctionRigBlueprint = Cast<UControlRigBlueprint>(FunctionLibrary->GetOuter()))
+			{
+				if(FunctionRigBlueprint != InTargetBlueprint)
+				{
+					if(bForce || !FunctionRigBlueprint->IsFunctionPublic(InFunction->GetFName()))
+					{
+                        TSharedRef<SControlRigFunctionLocalizationDialog> LocalizationDialog = SNew(SControlRigFunctionLocalizationDialog)
+                        .Function(InFunction)
+                        .TargetBlueprint(InTargetBlueprint);
+
+						if (LocalizationDialog->ShowModal() != EAppReturnType::Cancel)
+						{
+							TargetController->LocalizeFunctions(LocalizationDialog->GetFunctionsToLocalize());
+						}
+					}
+				}
+			}
+		}
+	}
+}
+
 void FControlRigEditor::OnCreateComment()
 {
 	TSharedPtr<SGraphEditor> GraphEditor = FocusedGraphEdPtr.Pin();
diff --git a/Engine/Plugins/Experimental/ControlRig/Source/ControlRigEditor/Private/Editor/ControlRigEditor.h b/Engine/Plugins/Experimental/ControlRig/Source/ControlRigEditor/Private/Editor/ControlRigEditor.h
--- a/Engine/Plugins/Experimental/ControlRig/Source/ControlRigEditor/Private/Editor/ControlRigEditor.h
+++ b/Engine/Plugins/Experimental/ControlRig/Source/ControlRigEditor/Private/Editor/ControlRigEditor.h
@@ -12,6 +12,7 @@
 #include "IPersonaViewport.h"
 #include "Framework/MultiBox/MultiBoxBuilder.h"
 #include "RigVMModel/RigVMGraph.h"
+#include "RigVMModel/Nodes/RigVMLibraryNode.h"
 #include "RigVMCore/RigVM.h"
 #include "Graph/SControlRigGraphPinNameListValueWidget.h"
 #include "Styling/SlateTypes.h"
@@ -354,6 +355,7 @@ protected:
 	/** delegate for changing property */
 	virtual void OnFinishedChangingProperties(const FPropertyChangedEvent& PropertyChangedEvent) override;
 	void OnBlueprintPropertyChainEvent(FPropertyChangedChainEvent& PropertyChangedChainEvent);
+	void OnRequestLocalizeFunctionDialog(URigVMLibraryNode* InFunction, UControlRigBlueprint* InTargetBlueprint, bool bForce);
 
 	URigVMGraph* GetFocusedModel() const;
 	URigVMController* GetFocusedController() const;
diff --git a/Engine/Plugins/Experimental/ControlRig/Source/ControlRigEditor/Private/Editor/SControlRigFunctionLocalizationWidget.cpp b/Engine/Plugins/Experimental/ControlRig/Source/ControlRigEditor/Private/Editor/SControlRigFunctionLocalizationWidget.cpp
new file mode 100644
--- /dev/null
+++ b/Engine/Plugins/Experimental/ControlRig/Source/ControlRigEditor/Private/Editor/SControlRigFunctionLocalizationWidget.cpp
@@ -0,0 +1,315 @@
+// Copyright Epic Games, Inc. All Rights Reserved.
+
+#include "SControlRigFunctionLocalizationWidget.h"
+#include "Widgets/Layout/SBorder.h"
+#include "Widgets/SBoxPanel.h"
+#include "Widgets/Layout/SScrollBar.h"
+#include "Widgets/Layout/SSpacer.h"
+#include "Widgets/Input/SCheckBox.h"
+#include "Widgets/Layout/SUniformGridPanel.h"
+#include "DetailLayoutBuilder.h"
+#include "ControlRigEditorStyle.h"
+#include "EditorStyleSet.h"
+#include "SlateOptMacros.h"
+#include "HAL/ConsoleManager.h"
+#include "Framework/Application/SlateApplication.h"
+#include "RigVMModel/Nodes/RigVMCollapseNode.h"
+#include "RigVMModel/Nodes/RigVMFunctionReferenceNode.h"
+
+#define LOCTEXT_NAMESPACE "SControlRigFunctionLocalizationWidget"
+
+//////////////////////////////////////////////////////////////
+/// SControlRigFunctionLocalizationItem
+///////////////////////////////////////////////////////////
+
+SControlRigFunctionLocalizationItem::SControlRigFunctionLocalizationItem(URigVMLibraryNode* InFunction)
+    : Function(InFunction)
+{
+	FString OuterName;
+	bool bIsPublic = false;
+
+	if(URigVMFunctionLibrary* Library = Cast<URigVMFunctionLibrary>(InFunction->GetOuter()))
+	{
+		if(UControlRigBlueprint* Blueprint = Cast<UControlRigBlueprint>(Library->GetOuter()))
+		{
+			OuterName = Blueprint->GetName();
+			bIsPublic = Blueprint->IsFunctionPublic(InFunction->GetFName());
+		}
+	}
+	DisplayText = FText::FromString(FString::Printf(TEXT("%s :: %s"), *OuterName, *InFunction->GetName()));
+	ToolTipText = bIsPublic ?
+		FText::FromString(FString::Printf(TEXT("%s :: %s is public. Check this to create a local copy."), *OuterName, *InFunction->GetName())) :
+		FText::FromString(FString::Printf(TEXT("%s :: %s is private. A local copy has to be created. To avoid this you can turn it public in the source Control Rig."), *OuterName, *InFunction->GetName()));
+}
+
+//////////////////////////////////////////////////////////////
+/// SControlRigFunctionLocalizationTableRow
+///////////////////////////////////////////////////////////
+
+void SControlRigFunctionLocalizationTableRow::Construct(const FArguments& InArgs, const TSharedRef<STableViewBase>& OwnerTable, SControlRigFunctionLocalizationWidget* InLocalizationWidget, TSharedRef<SControlRigFunctionLocalizationItem> InFunctionItem)
+{
+	STableRow<TSharedPtr<SControlRigFunctionLocalizationItem>>::Construct(
+		STableRow<TSharedPtr<SControlRigFunctionLocalizationItem>>::FArguments()
+		.Content()
+		[
+			SNew(SVerticalBox)
+
+			+ SVerticalBox::Slot()
+			.Padding(4.0f, 4.0f, 0.f, 0.f)
+			.AutoHeight()
+			[
+				SNew(SHorizontalBox)
+
+				+ SHorizontalBox::Slot()
+				.AutoWidth()
+				.HAlign(HAlign_Left)
+				.VAlign(VAlign_Center)
+				.Padding(0.0f, 0.0f, 12.f, 0.f)
+				[
+					SNew(STextBlock)
+					.Text(InFunctionItem->DisplayText)
+					.ToolTipText(InFunctionItem->ToolTipText)
+				]
+
+				+ SHorizontalBox::Slot()
+				.AutoWidth()
+				.HAlign(HAlign_Left)
+				.VAlign(VAlign_Center)
+				[
+					SNew(SCheckBox)
+					.IsChecked(InLocalizationWidget, &SControlRigFunctionLocalizationWidget::IsFunctionEnabled, InFunctionItem->Function)
+					.OnCheckStateChanged(InLocalizationWidget, &SControlRigFunctionLocalizationWidget::SetFunctionEnabled, InFunctionItem->Function)
+					.IsEnabled(InLocalizationWidget, &SControlRigFunctionLocalizationWidget::IsFunctionPublic, InFunctionItem->Function)
+					.ToolTipText(InFunctionItem->ToolTipText)
+				]
+			]
+		]
+		, OwnerTable
+	);
+}
+
+//////////////////////////////////////////////////////////////
+/// SControlRigFunctionLocalizationWidget
+///////////////////////////////////////////////////////////
+
+void SControlRigFunctionLocalizationWidget::Construct(const FArguments& InArgs, URigVMLibraryNode* InFunctionToLocalize, UControlRigBlueprint* InTargetBlueprint)
+{
+	FunctionsToLocalize.Add(InFunctionToLocalize);
+	FunctionItems.Reset();
+
+	TArray<URigVMLibraryNode*> NodesToVisit;
+	TArray<URigVMLibraryNode*> FunctionsForTable;
+
+	FunctionsForTable.Add(InFunctionToLocalize);
+	NodesToVisit.Add(InFunctionToLocalize);
+
+	for(int32 NodeToVisitIndex=0; NodeToVisitIndex<NodesToVisit.Num(); NodeToVisitIndex++)
+	{
+		URigVMLibraryNode* NodeToVisit = NodesToVisit[NodeToVisitIndex];
+		const TArray<URigVMNode*>& ContainedNodes = NodeToVisit->GetContainedNodes();
+		for(URigVMNode* ContainedNode : ContainedNodes)
+		{
+			if(URigVMLibraryNode* ContainedLibraryNode = Cast<URigVMLibraryNode>(ContainedNode))
+			{
+				NodesToVisit.AddUnique(ContainedLibraryNode);
+			}
+		}
+
+		if(URigVMFunctionReferenceNode* FunctionReferenceNode = Cast<URigVMFunctionReferenceNode>(NodeToVisit))
+		{
+			if(URigVMLibraryNode* ReferencedNode = FunctionReferenceNode->GetReferencedNode())
+			{
+				if(URigVMFunctionLibrary* ReferencedLibrary = ReferencedNode->GetLibrary())
+				{
+					if(UControlRigBlueprint* ReferencedBlueprint = Cast<UControlRigBlueprint>(ReferencedLibrary->GetOuter()))
+					{
+						if(ReferencedBlueprint != InTargetBlueprint)
+						{
+							FunctionsForTable.AddUnique(ReferencedNode);
+							if(!ReferencedBlueprint->IsFunctionPublic(ReferencedNode->GetFName()))
+							{
+								FunctionsToLocalize.AddUnique(ReferencedNode);
+							}
+						}
+					}
+				}
+
+				NodesToVisit.AddUnique(ReferencedNode);
+			}
+		}
+	}
+
+	// sort the functions to localize based on their nesting
+	Algo::Sort(FunctionsForTable, [](URigVMLibraryNode* A, URigVMLibraryNode* B) -> bool
+    {
+        check(A);
+        check(B);
+        return A->Contains(B);
+    });
+
+	for (URigVMLibraryNode* FunctionForTable : FunctionsForTable)
+	{
+		FunctionItems.Add(MakeShared<SControlRigFunctionLocalizationItem>(FunctionForTable));
+	}
+
+	ChildSlot
+		[
+			SNew(SVerticalBox)
+
+			+ SVerticalBox::Slot()
+			.Padding(4.0f, 4.0f, 0.f, 0.f)
+			.AutoHeight()
+			[
+				SNew(SBorder)
+				.Visibility(EVisibility::Visible)
+				.BorderImage(FEditorStyle::GetBrush("Menu.Background"))
+				[
+					SNew(SListView<TSharedPtr<SControlRigFunctionLocalizationItem>>)
+					.ListItemsSource(&FunctionItems)
+					.OnGenerateRow(this, &SControlRigFunctionLocalizationWidget::GenerateFunctionListRow)
+					.SelectionMode(ESelectionMode::None)
+				]
+			]
+		];
+}
+
+TSharedRef<ITableRow> SControlRigFunctionLocalizationWidget::GenerateFunctionListRow(TSharedPtr<SControlRigFunctionLocalizationItem> InItem, const TSharedRef<STableViewBase>& InOwningTable)
+{
+	TSharedRef<SControlRigFunctionLocalizationTableRow> TableRow = SNew(SControlRigFunctionLocalizationTableRow, InOwningTable, this, InItem.ToSharedRef());
+	TableRows.Add(InItem->Function, TableRow);
+	return TableRow;
+}
+
+ECheckBoxState SControlRigFunctionLocalizationWidget::IsFunctionEnabled(URigVMLibraryNode* InFunction) const
+{
+	return FunctionsToLocalize.Contains(InFunction) ? ECheckBoxState::Checked : ECheckBoxState::Unchecked;
+}
+
+void SControlRigFunctionLocalizationWidget::SetFunctionEnabled(ECheckBoxState NewState, URigVMLibraryNode* InFunction)
+{
+	if(NewState == ECheckBoxState::Checked)
+	{
+		FunctionsToLocalize.AddUnique(InFunction);
+	}
+	else
+	{
+		FunctionsToLocalize.Remove(InFunction);
+	}
+}
+
+bool SControlRigFunctionLocalizationWidget::IsFunctionPublic(URigVMLibraryNode* InFunction) const
+{
+	if(URigVMFunctionLibrary* FunctionLibrary = Cast<URigVMFunctionLibrary>(InFunction->GetOuter()))
+	{
+		if(UControlRigBlueprint* Blueprint = Cast<UControlRigBlueprint>(FunctionLibrary->GetOuter()))
+		{
+			return Blueprint->IsFunctionPublic(InFunction->GetFName());
+		}
+	}
+	return true;	
+}
+
+void SControlRigFunctionLocalizationDialog::Construct(const FArguments& InArgs)
+{
+	UserResponse = EAppReturnType::Cancel;
+
+	static const FString DialogText = TEXT(
+        "This action requires a function to be localized.\n\n" 
+        "Check the functions you want to create copies of below.\n\n"
+        "Note that private functions have to be copied,\n"
+        "while public functions can be referenced.\n\n"
+        "Checking all functions removes the dependencies completely."
+    );
+	
+	SWindow::Construct(SWindow::FArguments()
+        .Title(LOCTEXT("SControlRigFunctionLocalizationDialog_Title", "Select function(s) to localize"))
+        .SupportsMinimize(false)
+        .SupportsMaximize(false)
+        .SizingRule( ESizingRule::Autosized )
+        .ClientSize(FVector2D(450, 450))
+        [
+	        SNew(SHorizontalBox)
+
+	        +SHorizontalBox::Slot()
+	        .AutoWidth()
+	        .HAlign(HAlign_Left)
+	        .VAlign(VAlign_Top)
+	        .Padding(2)
+	        [
+				SNew(SVerticalBox)
+
+                +SVerticalBox::Slot()
+                .AutoHeight()
+                .HAlign(HAlign_Left)
+                .Padding(8)
+                [
+
+                    SNew(STextBlock)
+                    .Text(FText::FromString(DialogText))
+                    .Font(IDetailLayoutBuilder::GetDetailFont())
+                ]
+				
+                +SVerticalBox::Slot()
+                .AutoHeight()
+                .HAlign(HAlign_Left)
+                .Padding(2)
+                [
+                    SAssignNew(FunctionsWidget, SControlRigFunctionLocalizationWidget, InArgs._Function, InArgs._TargetBlueprint)
+                ]
+
+                +SVerticalBox::Slot()
+                .AutoHeight()
+                .HAlign(HAlign_Right)
+                .Padding(5)
+                [
+                    SNew(SUniformGridPanel)
+                    .SlotPadding(FEditorStyle::GetMargin("StandardDialog.SlotPadding"))
+                    .MinDesiredSlotWidth(FEditorStyle::GetFloat("StandardDialog.MinDesiredSlotWidth"))
+                    .MinDesiredSlotHeight(FEditorStyle::GetFloat("StandardDialog.MinDesiredSlotHeight"))
+                    +SUniformGridPanel::Slot(0, 0)
+                    [
+                        SNew(SButton)
+                        .HAlign(HAlign_Center)
+                        .ContentPadding(FEditorStyle::GetMargin("StandardDialog.ContentPadding"))
+                        .Text(LOCTEXT("OK", "OK"))
+                        .OnClicked(this, &SControlRigFunctionLocalizationDialog::OnButtonClick, EAppReturnType::Ok)
+                        .IsEnabled(this, &SControlRigFunctionLocalizationDialog::IsOkButtonEnabled)
+                    ]
+                    +SUniformGridPanel::Slot(1, 0)
+                    [
+                        SNew(SButton)
+                        .HAlign(HAlign_Center)
+                        .ContentPadding(FEditorStyle::GetMargin("StandardDialog.ContentPadding"))
+                        .Text(LOCTEXT("Cancel", "Cancel"))
+                        .OnClicked(this, &SControlRigFunctionLocalizationDialog::OnButtonClick, EAppReturnType::Cancel)
+                    ]
+                ]
+			]
+		]);
+}
+
+FReply SControlRigFunctionLocalizationDialog::OnButtonClick(EAppReturnType::Type ButtonID)
+{
+	UserResponse = ButtonID;
+	RequestDestroyWindow();
+
+	return FReply::Handled();
+}
+
+bool SControlRigFunctionLocalizationDialog::IsOkButtonEnabled() const
+{
+	return FunctionsWidget->FunctionsToLocalize.Num() > 0;
+}
+
+EAppReturnType::Type SControlRigFunctionLocalizationDialog::ShowModal()
+{
+	GEditor->EditorAddModalWindow(SharedThis(this));
+	return UserResponse;
+}
+
+TArray<URigVMLibraryNode*>& SControlRigFunctionLocalizationDialog::GetFunctionsToLocalize()
+{
+	return FunctionsWidget->FunctionsToLocalize;
+}
+
+#undef LOCTEXT_NAMESPACE
diff --git a/Engine/Plugins/Experimental/ControlRig/Source/ControlRigEditor/Private/Editor/SControlRigFunctionLocalizationWidget.h b/Engine/Plugins/Experimental/ControlRig/Source/ControlRigEditor/Private/Editor/SControlRigFunctionLocalizationWidget.h
new file mode 100644
--- /dev/null
+++ b/Engine/Plugins/Experimental/ControlRig/Source/ControlRigEditor/Private/Editor/SControlRigFunctionLocalizationWidget.h
@@ -0,0 +1,78 @@
+// Copyright Epic Games, Inc. All Rights Reserved.
+
+#pragma once
+
+#include "CoreMinimal.h"
+#include "ControlRigBlueprint.h"
+#include "RigVMModel/Nodes/RigVMLibraryNode.h"
+#include "Widgets/Views/SListView.h"
+#include "Dialogs/Dialogs.h"
+
+class SControlRigFunctionLocalizationWidget;
+
+class SControlRigFunctionLocalizationItem : public TSharedFromThis<SControlRigFunctionLocalizationItem>
+{
+public:
+	SControlRigFunctionLocalizationItem(URigVMLibraryNode* InFunction);
+
+	FText DisplayText;
+	FText ToolTipText;
+	URigVMLibraryNode* Function;
+};
+
+class SControlRigFunctionLocalizationTableRow : public STableRow<TSharedPtr<SControlRigFunctionLocalizationItem>>
+{
+	SLATE_BEGIN_ARGS(SControlRigFunctionLocalizationTableRow) {}
+	SLATE_END_ARGS()
+
+	void Construct(const FArguments& InArgs, const TSharedRef<STableViewBase>& OwnerTable, SControlRigFunctionLocalizationWidget* InLocalizationWidget, TSharedRef<SControlRigFunctionLocalizationItem> InFunctionItem);
+};
+
+class SControlRigFunctionLocalizationWidget : public SCompoundWidget
+{
+public:
+	SLATE_BEGIN_ARGS(SControlRigFunctionLocalizationWidget) {}
+	SLATE_END_ARGS()
+
+	void Construct(const FArguments& InArgs, URigVMLibraryNode* InFunctionToLocalize, UControlRigBlueprint* InTargetBlueprint);
+
+	TSharedRef<ITableRow> GenerateFunctionListRow(TSharedPtr<SControlRigFunctionLocalizationItem> InItem, const TSharedRef<STableViewBase>& InOwningTable);
+	ECheckBoxState IsFunctionEnabled(URigVMLibraryNode* InFunction) const;
+	void SetFunctionEnabled(ECheckBoxState NewState, URigVMLibraryNode* InFunction);
+	bool IsFunctionPublic(URigVMLibraryNode* InFunction) const;
+
+private:
+
+	TArray<URigVMLibraryNode*> FunctionsToLocalize;
+	TArray<TSharedPtr<SControlRigFunctionLocalizationItem>> FunctionItems;
+	TMap<URigVMLibraryNode*, TSharedRef<SControlRigFunctionLocalizationTableRow>> TableRows;
+
+	friend class SControlRigFunctionLocalizationTableRow;
+	friend class SControlRigFunctionLocalizationDialog;
+};
+
+class SControlRigFunctionLocalizationDialog : public SWindow
+{
+public:
+	
+	SLATE_BEGIN_ARGS(SControlRigFunctionLocalizationDialog)
+	{
+	}
+
+	SLATE_ARGUMENT(URigVMLibraryNode*, Function)
+	SLATE_ARGUMENT(UControlRigBlueprint*, TargetBlueprint)
+	SLATE_END_ARGS()
+
+	void Construct(const FArguments& InArgs);
+
+	EAppReturnType::Type ShowModal();
+	TArray<URigVMLibraryNode*>& GetFunctionsToLocalize();
+
+protected:
+
+	TSharedPtr<SControlRigFunctionLocalizationWidget> FunctionsWidget;
+	
+	FReply OnButtonClick(EAppReturnType::Type ButtonID);
+	bool IsOkButtonEnabled() const;
+	EAppReturnType::Type UserResponse;
+};
\ No newline at end of file
diff --git a/Engine/Source/Developer/RigVMDeveloper/Private/RigVMModel/Nodes/RigVMLibraryNode.cpp b/Engine/Source/Developer/RigVMDeveloper/Private/RigVMModel/Nodes/RigVMLibraryNode.cpp
--- a/Engine/Source/Developer/RigVMDeveloper/Private/RigVMModel/Nodes/RigVMLibraryNode.cpp
+++ b/Engine/Source/Developer/RigVMDeveloper/Private/RigVMModel/Nodes/RigVMLibraryNode.cpp
@@ -4,6 +4,8 @@
 #include "RigVMModel/Nodes/RigVMFunctionEntryNode.h"
 #include "RigVMModel/Nodes/RigVMFunctionReturnNode.h"
 #include "RigVMModel/RigVMGraph.h"
+#include "RigVMModel/Nodes/RigVMCollapseNode.h"
+#include "RigVMModel/Nodes/RigVMFunctionReferenceNode.h"
 
 const TArray<URigVMNode*> URigVMLibraryNode::EmptyNodes;
 const TArray<URigVMLink*> URigVMLibraryNode::EmptyLinks;
@@ -65,3 +67,40 @@ URigVMFunctionReturnNode* URigVMLibraryNode::GetReturnNode() const
 	}
 	return nullptr;
 }
+
+bool URigVMLibraryNode::Contains(URigVMLibraryNode* InContainedNode, bool bRecursive) const
+{
+	if(InContainedNode == nullptr)
+	{
+		return false;
+	}
+	
+	for(URigVMNode* ContainedNode : GetContainedNodes())
+	{
+		if(ContainedNode == InContainedNode)
+		{
+			return true;
+		}
+
+		if(URigVMFunctionReferenceNode* FunctionReferenceNode = Cast<URigVMFunctionReferenceNode>(ContainedNode))
+		{
+			if(URigVMLibraryNode* ReferencedNode = FunctionReferenceNode->GetReferencedNode())
+			{
+				if(ReferencedNode == InContainedNode)
+				{
+					return true;
+				}
+			}
+		}
+		if(URigVMLibraryNode* ContainedLibraryNode = Cast<URigVMLibraryNode>(ContainedNode))
+		{
+			if(ContainedLibraryNode->Contains(InContainedNode))
+			{
+				return true;
+			}
+		}
+	}
+
+	return false;
+}
+
diff --git a/Engine/Source/Developer/RigVMDeveloper/Private/RigVMModel/RigVMController.cpp b/Engine/Source/Developer/RigVMDeveloper/Private/RigVMModel/RigVMController.cpp
--- a/Engine/Source/Developer/RigVMDeveloper/Private/RigVMModel/RigVMController.cpp
+++ b/Engine/Source/Developer/RigVMDeveloper/Private/RigVMModel/RigVMController.cpp
@@ -12,6 +12,7 @@
 #include "UObject/PropertyPortFlags.h"
 #include "UObject/Package.h"
 #include "Misc/CoreMisc.h"
+#include "Algo/Sort.h"
 
 #if WITH_EDITOR
 #include "Exporters/Exporter.h"
@@ -1767,26 +1768,41 @@ TArray<FName> URigVMController::ImportNodesFromText(const FString& InText, bool
 			ActionStack->AddAction(FRigVMRemoveNodeAction(CreatedNode, this));
 		}
 
-		if (URigVMUnitNode* UnitNode = Cast<URigVMUnitNode>(CreatedNode))
+		// find all nodes affected by this
+		TArray<URigVMNode*> SubNodes;
+		SubNodes.Add(CreatedNode);
+
+		for(int32 SubNodeIndex=0; SubNodeIndex < SubNodes.Num(); SubNodeIndex++)
 		{
-			if (UnitNodeCreatedContext.IsValid())
+			if(URigVMCollapseNode* CollapseNode = Cast<URigVMCollapseNode>(SubNodes[SubNodeIndex]))
 			{
-				if (TSharedPtr<FStructOnScope> StructScope = UnitNode->ConstructStructInstance())
-				{
-					TGuardValue<FName> NodeNameScope(UnitNodeCreatedContext.NodeName, UnitNode->GetFName());
-					FRigVMStruct* StructInstance = (FRigVMStruct*)StructScope->GetStructMemory();
-					StructInstance->OnUnitNodeCreated(UnitNodeCreatedContext);
-				}
+				SubNodes.Append(CollapseNode->GetContainedNodes());
 			}
 		}
 
-		if (URigVMFunctionReferenceNode* FunctionRefNode = Cast<URigVMFunctionReferenceNode>(CreatedNode))
+		for(URigVMNode* SubNode : SubNodes)
 		{
-			if (URigVMFunctionLibrary* FunctionLibrary = FunctionRefNode->GetLibrary())
+			if (URigVMUnitNode* UnitNode = Cast<URigVMUnitNode>(SubNode))
 			{
-				if (URigVMLibraryNode* FunctionDefinition = FunctionRefNode->GetReferencedNode())
+				if (UnitNodeCreatedContext.IsValid())
 				{
-					FunctionLibrary->FunctionReferences.FindOrAdd(FunctionDefinition).FunctionReferences.Add(FunctionRefNode);
+					if (TSharedPtr<FStructOnScope> StructScope = UnitNode->ConstructStructInstance())
+					{
+						TGuardValue<FName> NodeNameScope(UnitNodeCreatedContext.NodeName, UnitNode->GetFName());
+						FRigVMStruct* StructInstance = (FRigVMStruct*)StructScope->GetStructMemory();
+						StructInstance->OnUnitNodeCreated(UnitNodeCreatedContext);
+					}
+				}
+			}
+
+			if (URigVMFunctionReferenceNode* FunctionRefNode = Cast<URigVMFunctionReferenceNode>(SubNode))
+			{
+				if (URigVMFunctionLibrary* FunctionLibrary = FunctionRefNode->GetLibrary())
+				{
+					if (URigVMLibraryNode* FunctionDefinition = FunctionRefNode->GetReferencedNode())
+					{
+						FunctionLibrary->FunctionReferences.FindOrAdd(FunctionDefinition).FunctionReferences.Add(FunctionRefNode);
+					}
 				}
 			}
 		}
@@ -1857,6 +1873,214 @@ TArray<FName> URigVMController::ImportNodesFromText(const FString& InText, bool
 	return NodeNames;
 }
 
+URigVMLibraryNode* URigVMController::LocalizeFunction(
+	URigVMLibraryNode* InFunctionDefinition,
+	bool bLocalizeDependentPrivateFunctions,
+	bool bSetupUndoRedo)
+{
+	if(!IsValidGraph())
+	{
+		return nullptr;
+	}
+
+	if(InFunctionDefinition == nullptr)
+	{
+		return nullptr;
+	}
+
+	TArray<URigVMLibraryNode*> FunctionsToLocalize;
+	FunctionsToLocalize.Add(InFunctionDefinition);
+
+	TMap<URigVMLibraryNode*, URigVMLibraryNode*> Results = LocalizeFunctions(FunctionsToLocalize, bLocalizeDependentPrivateFunctions, bSetupUndoRedo);
+
+	URigVMLibraryNode** LocalizedFunctionPtr = Results.Find(FunctionsToLocalize[0]);
+	if(LocalizedFunctionPtr)
+	{
+		return *LocalizedFunctionPtr;
+	}
+	return nullptr;
+}
+
+TMap<URigVMLibraryNode*, URigVMLibraryNode*> URigVMController::LocalizeFunctions(
+    TArray<URigVMLibraryNode*> InFunctionDefinitions,
+    bool bLocalizeDependentPrivateFunctions,
+    bool bSetupUndoRedo)
+{
+	TMap<URigVMLibraryNode*, URigVMLibraryNode*> LocalizedFunctions;
+
+	if(!IsValidGraph())
+	{
+		return LocalizedFunctions;
+	}
+
+	URigVMGraph* Graph = GetGraph();
+	check(Graph);
+
+	URigVMFunctionLibrary* ThisLibrary = Graph->GetDefaultFunctionLibrary();
+	if(ThisLibrary == nullptr)
+	{
+		return LocalizedFunctions;
+	}
+
+	TArray<URigVMLibraryNode*> FunctionsToLocalize;
+
+	TArray<URigVMLibraryNode*> NodesToVisit;
+	for(URigVMLibraryNode* FunctionDefinition : InFunctionDefinitions)
+	{
+		NodesToVisit.AddUnique(FunctionDefinition);
+		FunctionsToLocalize.AddUnique(FunctionDefinition);
+	}
+
+	// find all functions to localize
+	for(int32 NodeToVisitIndex=0; NodeToVisitIndex<NodesToVisit.Num(); NodeToVisitIndex++)
+	{
+		URigVMLibraryNode* NodeToVisit = NodesToVisit[NodeToVisitIndex];
+
+		if(URigVMCollapseNode* CollapseNode = Cast<URigVMCollapseNode>(NodeToVisit))
+		{
+			const TArray<URigVMNode*>& ContainedNodes = CollapseNode->GetContainedNodes();
+			for(URigVMNode* ContainedNode : ContainedNodes)
+			{
+				if(URigVMLibraryNode* ContainedLibraryNode = Cast<URigVMLibraryNode>(ContainedNode))
+				{
+					NodesToVisit.AddUnique(ContainedLibraryNode);
+				}
+			}
+
+			if(URigVMFunctionLibrary* OtherLibrary = Cast<URigVMFunctionLibrary>(CollapseNode->GetOuter()))
+			{
+				if(OtherLibrary != ThisLibrary)
+				{
+					bool bIsAvailable = false;
+					if(IsFunctionAvailableDelegate.IsBound())
+					{
+						bIsAvailable = IsFunctionAvailableDelegate.Execute(CollapseNode);
+					}
+
+					if(!bIsAvailable)
+					{
+						if(!bLocalizeDependentPrivateFunctions)
+						{
+							ReportAndNotifyErrorf(TEXT("Cannot localize function - dependency %s is private."), *CollapseNode->GetPathName());
+							return LocalizedFunctions;
+						}
+						
+						FunctionsToLocalize.AddUnique(CollapseNode);
+					}
+				}
+			}
+		}
+
+		else if(URigVMFunctionReferenceNode* FunctionReferencedNode = Cast<URigVMFunctionReferenceNode>(NodeToVisit))
+		{
+			if(FunctionReferencedNode->GetLibrary() != ThisLibrary)
+			{
+				if(URigVMCollapseNode* FunctionDefinition = Cast<URigVMCollapseNode>(FunctionReferencedNode->GetReferencedNode()))
+				{
+					NodesToVisit.AddUnique(FunctionDefinition);
+				}
+			}
+		}
+	}
+	
+	// sort the functions to localize based on their nesting
+	Algo::Sort(FunctionsToLocalize, [](URigVMLibraryNode* A, URigVMLibraryNode* B) -> bool
+	{
+		check(A);
+		check(B);
+		return B->Contains(A);
+	});
+
+	// export all of the content for each node
+	TMap<URigVMLibraryNode*, FString> ExportedTextPerFunction;
+	for(URigVMLibraryNode* FunctionToLocalize : FunctionsToLocalize)
+	{
+		URigVMFunctionLibrary* OtherLibrary = Cast<URigVMFunctionLibrary>(FunctionToLocalize->GetOuter());
+		FRigVMControllerGraphGuard GraphGuard(this, OtherLibrary, false);
+
+		const TArray<FName> NodeNamesToExport = {FunctionToLocalize->GetFName()};
+		const FString ExportedText = ExportNodesToText(NodeNamesToExport);
+		ExportedTextPerFunction.Add(FunctionToLocalize, ExportedText);
+	}
+
+	if (bSetupUndoRedo)
+	{
+		OpenUndoBracket(TEXT("Localize functions"));
+	}
+
+	// import the functions to our local function library
+	{
+		FRigVMControllerGraphGuard GraphGuard(this, ThisLibrary, bSetupUndoRedo);
+
+		// override the availability and check up later
+		TGuardValue<FRigVMController_IsFunctionAvailableDelegate> IsFunctionAvailableGuard(IsFunctionAvailableDelegate,
+			FRigVMController_IsFunctionAvailableDelegate::CreateLambda([](URigVMLibraryNode*)
+			{
+				return true;
+			})
+		);
+
+		for(URigVMLibraryNode* FunctionToLocalize : FunctionsToLocalize)
+		{
+			const FString& ExportedText = ExportedTextPerFunction.FindChecked(FunctionToLocalize);
+			TArray<FName> ImportedNodeNames = ImportNodesFromText(ExportedText);
+			if(ImportedNodeNames.Num() != 1)
+			{
+				ReportErrorf(TEXT("Not possible to localize function %s"), *FunctionToLocalize->GetPathName());
+				continue;
+			}
+
+			URigVMLibraryNode* LocalizedFunction = Cast<URigVMLibraryNode>(GetGraph()->FindNodeByName(ImportedNodeNames[0]));
+			if(LocalizedFunction == nullptr)
+			{
+				ReportErrorf(TEXT("Not possible to localize function %s"), *FunctionToLocalize->GetPathName());
+				continue;
+			}
+
+			LocalizedFunctions.Add(FunctionToLocalize, LocalizedFunction);
+			ThisLibrary->LocalizedFunctions.FindOrAdd(FunctionToLocalize->GetPathName(), LocalizedFunction);
+		}
+	}
+
+	// once we have all local functions available, clean up the references
+	TArray<URigVMGraph*> GraphsToUpdate;
+	GraphsToUpdate.AddUnique(Graph);
+	if(URigVMFunctionLibrary* DefaultFunctionLibrary = Graph->GetDefaultFunctionLibrary())
+	{
+		GraphsToUpdate.AddUnique(DefaultFunctionLibrary);
+	}
+	for(int32 GraphToUpdateIndex=0; GraphToUpdateIndex<GraphsToUpdate.Num(); GraphToUpdateIndex++)
+	{
+		URigVMGraph* GraphToUpdate = GraphsToUpdate[GraphToUpdateIndex];
+		
+		const TArray<URigVMNode*> NodesToUpdate = GraphToUpdate->GetNodes();
+		for(URigVMNode* NodeToUpdate : NodesToUpdate)
+		{
+			if(URigVMCollapseNode* CollapseNode = Cast<URigVMCollapseNode>(NodeToUpdate))
+			{
+				GraphsToUpdate.AddUnique(CollapseNode->GetContainedGraph());
+			}
+			else if(URigVMFunctionReferenceNode* FunctionReferenceNode = Cast<URigVMFunctionReferenceNode>(NodeToUpdate))
+			{
+				URigVMLibraryNode* ReferencedNode = FunctionReferenceNode->GetReferencedNode();
+				URigVMLibraryNode** RemappedNodePtr = LocalizedFunctions.Find(ReferencedNode);
+				if(RemappedNodePtr)
+				{
+					URigVMLibraryNode* RemappedNode = *RemappedNodePtr;
+					SetReferencedFunction(FunctionReferenceNode, RemappedNode, bSetupUndoRedo);
+				}
+			}
+		}
+	}
+
+	if(bSetupUndoRedo)
+	{
+		CloseUndoBracket();
+	}
+
+	return LocalizedFunctions;
+}
+
 FName URigVMController::GetUniqueName(const FName& InName, TFunction<bool(const FName&)> IsNameAvailableFunction)
 {
 	FString NamePrefix = InName.ToString();
@@ -3117,6 +3341,30 @@ URigVMCollapseNode* URigVMController::PromoteFunctionReferenceNodeToCollapseNode
 	return CollapseNode;
 }
 
+void URigVMController::SetReferencedFunction(URigVMFunctionReferenceNode* InFunctionRefNode, URigVMLibraryNode* InNewReferencedNode, bool bSetupUndoRedo)
+{
+	if(!IsValidGraph())
+	{
+		return;
+	}
+	
+	URigVMLibraryNode* ReferencedNode = InFunctionRefNode->GetReferencedNode();
+	
+	if(URigVMFunctionLibrary* OtherLibrary = Cast<URigVMFunctionLibrary>(ReferencedNode->GetOuter()))
+	{
+		if(FRigVMFunctionReferenceArray* OtherReferences = OtherLibrary->FunctionReferences.Find(ReferencedNode))
+		{
+			OtherReferences->FunctionReferences.Remove(InFunctionRefNode);
+		}
+	}
+
+	URigVMFunctionLibrary* NewLibrary = Cast<URigVMFunctionLibrary>(InNewReferencedNode->GetOuter());
+	FRigVMFunctionReferenceArray& NewReferences = NewLibrary->FunctionReferences.FindOrAdd(InNewReferencedNode);
+	NewReferences.FunctionReferences.Add(InFunctionRefNode);
+
+	InFunctionRefNode->SetReferencedNode(InNewReferencedNode);
+}
+
 void URigVMController::RefreshFunctionPins(URigVMNode* InNode, bool bNotify)
 {
 	if (InNode == nullptr)
@@ -3284,6 +3532,15 @@ bool URigVMController::RemoveNode(URigVMNode* InNode, bool bSetupUndoRedo, bool
 				}
 			}
 			FunctionLibrary->FunctionReferences.Remove(LibraryNode);
+
+			for(const TPair<FString, URigVMLibraryNode*>& Pair : FunctionLibrary->LocalizedFunctions)
+			{
+				if(Pair.Value == LibraryNode)
+				{
+					FunctionLibrary->LocalizedFunctions.Remove(Pair.Key);
+					break;
+				}
+			}
 		}
 	}
 
@@ -6756,18 +7013,30 @@ bool URigVMController::CanAddNode(URigVMNode* InNode, bool bReportErrors)
 	{
 		if (URigVMFunctionLibrary* FunctionLibrary = FunctionRefNode->GetLibrary())
 		{
-			if(Graph->GetDefaultFunctionLibrary() != FunctionLibrary)
+			if (URigVMLibraryNode* FunctionDefinition = FunctionRefNode->GetReferencedNode())
 			{
-				// this node references a function that is not part of our function library
-				if(bReportErrors)
+				if(!CanAddFunctionRefForDefinition(FunctionDefinition, false))
 				{
-					ReportError(TEXT("Cannot import function reference node."));
+					URigVMFunctionLibrary* TargetLibrary = Graph->GetDefaultFunctionLibrary();
+					URigVMLibraryNode* NewFunctionDefinition = TargetLibrary->FindPreviouslyLocalizedFunction(FunctionDefinition);
+					
+					if((NewFunctionDefinition == nullptr) && RequestLocalizeFunctionDelegate.IsBound())
+					{
+						if(RequestLocalizeFunctionDelegate.Execute(FunctionDefinition))
+						{
+							NewFunctionDefinition = TargetLibrary->FindPreviouslyLocalizedFunction(FunctionDefinition);
+						}
+					}
+
+					if(NewFunctionDefinition == nullptr)
+					{
+						return false;
+					}
+					
+					SetReferencedFunction(FunctionRefNode, NewFunctionDefinition, false);
+					FunctionDefinition = NewFunctionDefinition;
 				}
-				DestroyObject(InNode);
-				return false;
-			}
-			else if (URigVMLibraryNode* FunctionDefinition = FunctionRefNode->GetReferencedNode())
-			{
+				
 				if(!CanAddFunctionRefForDefinition(FunctionDefinition, bReportErrors))
 				{
 					DestroyObject(InNode);
@@ -6776,6 +7045,19 @@ bool URigVMController::CanAddNode(URigVMNode* InNode, bool bReportErrors)
 			}
 		}			
 	}
+	else if(URigVMCollapseNode* CollapseNode = Cast<URigVMCollapseNode>(InNode))
+	{
+		FRigVMControllerGraphGuard GraphGuard(this, CollapseNode->GetContainedGraph(), false);
+
+		TArray<URigVMNode*> ContainedNodes = CollapseNode->GetContainedNodes();
+		for(URigVMNode* ContainedNode : ContainedNodes)
+		{
+			if(!CanAddNode(ContainedNode, bReportErrors))
+			{
+				return false;
+			}
+		}
+	}
 
 	return true;
 }
@@ -6792,6 +7074,18 @@ bool URigVMController::CanAddFunctionRefForDefinition(URigVMLibraryNode* InFunct
 	URigVMGraph* Graph = GetGraph();
 	check(Graph);
 
+	if(IsFunctionAvailableDelegate.IsBound())
+	{
+		if(!IsFunctionAvailableDelegate.Execute(InFunctionDefinition))
+		{
+			if(bReportErrors)
+			{
+				ReportAndNotifyError(TEXT("Function is not available for placement in another graph host."));
+			}
+			return false;
+		}
+	}
+
 	URigVMLibraryNode* ParentLibraryNode = Cast<URigVMLibraryNode>(Graph->GetOuter());
 	while (ParentLibraryNode)
 	{
diff --git a/Engine/Source/Developer/RigVMDeveloper/Private/RigVMModel/RigVMFunctionLibrary.cpp b/Engine/Source/Developer/RigVMDeveloper/Private/RigVMModel/RigVMFunctionLibrary.cpp
--- a/Engine/Source/Developer/RigVMDeveloper/Private/RigVMModel/RigVMFunctionLibrary.cpp
+++ b/Engine/Source/Developer/RigVMDeveloper/Private/RigVMModel/RigVMFunctionLibrary.cpp
@@ -75,3 +75,41 @@ TArray< FString > URigVMFunctionLibrary::GetReferencePathsForFunction(const FNam
 
 	return Result;
 }
+
+URigVMLibraryNode* URigVMFunctionLibrary::FindPreviouslyLocalizedFunction(URigVMLibraryNode* InFunctionToLocalize)
+{
+	if(InFunctionToLocalize == nullptr)
+	{
+		return nullptr;
+	}
+	
+	const FString PathName = InFunctionToLocalize->GetPathName();
+
+	if(!LocalizedFunctions.Contains((PathName)))
+	{
+		return nullptr;
+	}
+
+	URigVMLibraryNode* LocalizedFunction = LocalizedFunctions.FindChecked(PathName);
+
+	// once we found the function - let's make sure it's notation is right
+	if(LocalizedFunction->GetPins().Num() != InFunctionToLocalize->GetPins().Num())
+	{
+		return nullptr;
+	}
+	for(int32 PinIndex=0; PinIndex < InFunctionToLocalize->GetPins().Num(); PinIndex++)
+	{
+		URigVMPin* PinA = InFunctionToLocalize->GetPins()[PinIndex];
+		URigVMPin* PinB = LocalizedFunction->GetPins()[PinIndex];
+
+		if((PinA->GetFName() != PinB->GetFName()) ||
+			(PinA->GetCPPType() != PinB->GetCPPType()) ||
+			(PinA->GetCPPTypeObject() != PinB->GetCPPTypeObject()) ||
+			(PinA->IsArray() != PinB->IsArray()))
+		{
+			return nullptr;
+		}
+	}
+	
+	return LocalizedFunction;
+}
diff --git a/Engine/Source/Developer/RigVMDeveloper/Public/RigVMModel/Nodes/RigVMLibraryNode.h b/Engine/Source/Developer/RigVMDeveloper/Public/RigVMModel/Nodes/RigVMLibraryNode.h
--- a/Engine/Source/Developer/RigVMDeveloper/Public/RigVMModel/Nodes/RigVMLibraryNode.h
+++ b/Engine/Source/Developer/RigVMDeveloper/Public/RigVMModel/Nodes/RigVMLibraryNode.h
@@ -36,6 +36,7 @@ public:
 	virtual const TArray<URigVMLink*>& GetContainedLinks() const;
 	virtual URigVMFunctionEntryNode* GetEntryNode() const;
 	virtual URigVMFunctionReturnNode* GetReturnNode() const;
+	virtual bool Contains(URigVMLibraryNode* InContainedNode, bool bRecursive = true) const;
 
 protected:
 
diff --git a/Engine/Source/Developer/RigVMDeveloper/Public/RigVMModel/RigVMController.h b/Engine/Source/Developer/RigVMDeveloper/Public/RigVMModel/RigVMController.h
--- a/Engine/Source/Developer/RigVMDeveloper/Public/RigVMModel/RigVMController.h
+++ b/Engine/Source/Developer/RigVMDeveloper/Public/RigVMModel/RigVMController.h
@@ -23,6 +23,8 @@ class URigVMActionStack;
 DECLARE_DELEGATE_RetVal_OneParam(bool, FRigVMController_ShouldStructUnfoldDelegate, const UStruct*)
 DECLARE_DELEGATE_RetVal(TArray<FRigVMExternalVariable>, FRigVMController_GetExternalVariablesDelegate)
 DECLARE_DELEGATE_RetVal(const FRigVMByteCode*, FRigVMController_GetByteCodeDelegate)
+DECLARE_DELEGATE_RetVal_OneParam(bool, FRigVMController_IsFunctionAvailableDelegate, URigVMLibraryNode*)
+DECLARE_DELEGATE_RetVal_OneParam(bool, FRigVMController_RequestLocalizeFunctionDelegate, URigVMLibraryNode*)
 
 /**
  * The Controller is the sole authority to perform changes
@@ -277,6 +279,19 @@ public:
 	UFUNCTION(BlueprintCallable, Category = RigVMController)
 	TArray<FName> ImportNodesFromText(const FString& InText, bool bSetupUndoRedo = true);
 
+	// Copies a function declaration into this graph's local function library
+	UFUNCTION(BlueprintCallable, Category = RigVMController)
+    URigVMLibraryNode* LocalizeFunction(
+    	URigVMLibraryNode* InFunctionDefinition,
+		bool bLocalizeDependentPrivateFunctions = true,
+    	bool bSetupUndoRedo = true);
+
+	// Copies a series of function declaratioms into this graph's local function library
+    TMap<URigVMLibraryNode*, URigVMLibraryNode*> LocalizeFunctions(
+        TArray<URigVMLibraryNode*> InFunctionDefinitions,
+        bool bLocalizeDependentPrivateFunctions = true,
+        bool bSetupUndoRedo = true);
+
 	// Returns a unique name
 	static FName GetUniqueName(const FName& InName, TFunction<bool(const FName&)> IsNameAvailableFunction);
 
@@ -548,6 +563,12 @@ public:
 	// A delegate to retrieve the current bytecode of the graph
 	FRigVMController_GetByteCodeDelegate GetCurrentByteCodeDelegate;
 
+	// A delegate to determine if a function is public
+	FRigVMController_IsFunctionAvailableDelegate IsFunctionAvailableDelegate;
+
+	// A delegate to localize a function on demand
+	FRigVMController_RequestLocalizeFunctionDelegate RequestLocalizeFunctionDelegate;
+
 	int32 DetachLinksFromPinObjects(const TArray<URigVMLink*>* InLinks = nullptr, bool bNotify = false);
 	int32 ReattachLinksToPinObjects(bool bFollowCoreRedirectors = false, const TArray<URigVMLink*>* InLinks = nullptr, bool bNotify = false);
 	void AddPinRedirector(bool bInput, bool bOutput, const FString& OldPinPath, const FString& NewPinPath);
@@ -649,6 +670,7 @@ private:
 	TArray<URigVMNode*> ExpandLibraryNode(URigVMLibraryNode* InNode, bool bSetupUndoRedo);
 	URigVMFunctionReferenceNode* PromoteCollapseNodeToFunctionReferenceNode(URigVMCollapseNode* InCollapseNode, bool bSetupUndoRedo);
 	URigVMCollapseNode* PromoteFunctionReferenceNodeToCollapseNode(URigVMFunctionReferenceNode* InFunctionRefNode, bool bSetupUndoRedo);
+	void SetReferencedFunction(URigVMFunctionReferenceNode* InFunctionRefNode, URigVMLibraryNode* InNewReferencedNode, bool bSetupUndoRedo);
 
 	void RefreshFunctionPins(URigVMNode* InNode, bool bNotify = true);
 
diff --git a/Engine/Source/Developer/RigVMDeveloper/Public/RigVMModel/RigVMFunctionLibrary.h b/Engine/Source/Developer/RigVMDeveloper/Public/RigVMModel/RigVMFunctionLibrary.h
--- a/Engine/Source/Developer/RigVMDeveloper/Public/RigVMModel/RigVMFunctionLibrary.h
+++ b/Engine/Source/Developer/RigVMDeveloper/Public/RigVMModel/RigVMFunctionLibrary.h
@@ -62,11 +62,22 @@ public:
 	UFUNCTION(BlueprintCallable, Category = RigVMGraph)
 	TArray< FString > GetReferencePathsForFunction(const FName& InFunctionName);
 
+	// Returns a function that has been previously localized based on the provided function to localize.
+	// We maintain meta data on what functions have been created locally based on which other ones,
+	// and use this method to avoid redundant localizations.
+	URigVMLibraryNode* FindPreviouslyLocalizedFunction(URigVMLibraryNode* InFunctionToLocalize);
+
 private:
 
 	UPROPERTY()
 	TMap< URigVMLibraryNode*, FRigVMFunctionReferenceArray > FunctionReferences;
 
+	// A map which stores a library node per original pathname.
+	// The source pathname is the full path of the source function that was localized
+	// to the local copy stored in the value of the pair.
+	UPROPERTY()
+	TMap< FString, URigVMLibraryNode* > LocalizedFunctions;
+
 	friend class URigVMController;
 	friend class URigVMCompiler;
 };
