Control Rig: Move default function library setup to virtual function - avoid CIS breakage

diff --git a/Engine/Plugins/Experimental/ControlRig/Source/ControlRigDeveloper/Private/ControlRigBlueprint.cpp b/Engine/Plugins/Experimental/ControlRig/Source/ControlRigDeveloper/Private/ControlRigBlueprint.cpp
--- a/Engine/Plugins/Experimental/ControlRig/Source/ControlRigDeveloper/Private/ControlRigBlueprint.cpp
+++ b/Engine/Plugins/Experimental/ControlRig/Source/ControlRigDeveloper/Private/ControlRigBlueprint.cpp
@@ -80,7 +80,6 @@ UControlRigBlueprint::UControlRigBlueprint(const FObjectInitializer& ObjectIniti
 	FunctionLibraryEdGraph->Initialize(this);
 
 	Model->SetDefaultFunctionLibrary(FunctionLibrary);
-	FunctionLibrary->SetDefaultFunctionLibrary(FunctionLibrary);
 
 	Validator = ObjectInitializer.CreateDefaultSubobject<UControlRigValidator>(this, TEXT("ControlRigValidator"));
 
diff --git a/Engine/Source/Developer/RigVMDeveloper/Private/RigVMModel/RigVMFunctionLibrary.cpp b/Engine/Source/Developer/RigVMDeveloper/Private/RigVMModel/RigVMFunctionLibrary.cpp
--- a/Engine/Source/Developer/RigVMDeveloper/Private/RigVMModel/RigVMFunctionLibrary.cpp
+++ b/Engine/Source/Developer/RigVMDeveloper/Private/RigVMModel/RigVMFunctionLibrary.cpp
@@ -5,6 +5,7 @@
 URigVMFunctionLibrary::URigVMFunctionLibrary()
 : URigVMGraph()
 {
+	
 }
 
 FString URigVMFunctionLibrary::GetNodePath() const
@@ -12,6 +13,15 @@ FString URigVMFunctionLibrary::GetNodePath() const
 	return FString::Printf(TEXT("FunctionLibrary::%s"), *Super::GetNodePath());
 }
 
+URigVMFunctionLibrary* URigVMFunctionLibrary::GetDefaultFunctionLibrary() const
+{
+	if(URigVMFunctionLibrary* DefaultFunctionLibrary = Super::GetDefaultFunctionLibrary())
+	{
+		return DefaultFunctionLibrary;
+	}
+	return (URigVMFunctionLibrary*)this;
+}
+
 TArray<URigVMLibraryNode*> URigVMFunctionLibrary::GetFunctions() const
 {
 	TArray<URigVMLibraryNode*> Functions;
diff --git a/Engine/Source/Developer/RigVMDeveloper/Public/RigVMModel/RigVMFunctionLibrary.h b/Engine/Source/Developer/RigVMDeveloper/Public/RigVMModel/RigVMFunctionLibrary.h
--- a/Engine/Source/Developer/RigVMDeveloper/Public/RigVMModel/RigVMFunctionLibrary.h
+++ b/Engine/Source/Developer/RigVMDeveloper/Public/RigVMModel/RigVMFunctionLibrary.h
@@ -44,6 +44,7 @@ public:
 
 	// URigVMGraph interface
 	virtual FString GetNodePath() const override;
+	virtual URigVMFunctionLibrary* GetDefaultFunctionLibrary() const override;
 	// end URigVMGraph interface
 
 	// Returns all of the stored functions
diff --git a/Engine/Source/Developer/RigVMDeveloper/Public/RigVMModel/RigVMGraph.h b/Engine/Source/Developer/RigVMDeveloper/Public/RigVMModel/RigVMGraph.h
--- a/Engine/Source/Developer/RigVMDeveloper/Public/RigVMModel/RigVMGraph.h
+++ b/Engine/Source/Developer/RigVMDeveloper/Public/RigVMModel/RigVMGraph.h
@@ -100,7 +100,7 @@ public:
 
 	// Returns the locally available function library
 	UFUNCTION(BlueprintCallable, Category = RigVMGraph)
-	URigVMFunctionLibrary* GetDefaultFunctionLibrary() const;
+	virtual URigVMFunctionLibrary* GetDefaultFunctionLibrary() const;
 
 	UFUNCTION(BlueprintCallable, Category = RigVMGraph)
 	void SetDefaultFunctionLibrary(URigVMFunctionLibrary* InFunctionLibrary);
