Control Rig: Services around introspection of run counts

diff --git a/Engine/Source/Developer/RigVMDeveloper/Private/RigVMCompiler/RigVMASTProxy.cpp b/Engine/Source/Developer/RigVMDeveloper/Private/RigVMCompiler/RigVMASTProxy.cpp
--- a/Engine/Source/Developer/RigVMDeveloper/Private/RigVMCompiler/RigVMASTProxy.cpp
+++ b/Engine/Source/Developer/RigVMDeveloper/Private/RigVMCompiler/RigVMASTProxy.cpp
@@ -93,3 +93,62 @@ FRigVMASTProxy FRigVMASTProxy::MakeFromUObject(UObject* InSubject)
 #endif
 	return Proxy;
 }
+
+FRigVMASTProxy FRigVMASTProxy::MakeFromCallPath(const FString& InCallPath, UObject* InRootObject)
+{
+	if(InCallPath.IsEmpty() || InRootObject == nullptr)
+	{
+		return FRigVMASTProxy();
+	}
+
+	UObject* ParentObject = InRootObject;
+
+	FRigVMASTProxy Proxy;
+	Proxy.Callstack.Stack.Reset();
+
+	FString Left;
+	FString Right = InCallPath;
+
+	while(!Right.IsEmpty() && ParentObject != nullptr)
+	{
+		if(!Right.Split(TEXT("|"), &Left, &Right))
+		{
+			Left = Right;
+			Right.Empty();
+		}
+
+		if(URigVMLibraryNode* LibraryNode = Cast<URigVMLibraryNode>(ParentObject))
+		{
+			ParentObject = LibraryNode->GetContainedGraph();
+		}
+
+		if(URigVMGraph* Graph = Cast<URigVMGraph>(ParentObject))
+		{
+			if(URigVMNode* FoundNode = Graph->FindNodeByName(*Left))
+			{
+				Proxy.Callstack.Stack.Push(FoundNode);
+				ParentObject = FoundNode;
+			}
+			else if(URigVMPin* FoundPin = Graph->FindPin(Left))
+			{
+				Proxy.Callstack.Stack.Push(FoundPin);
+				break;
+			}
+			else
+			{
+				return FRigVMASTProxy();
+			}
+		}
+		else
+		{
+			return FRigVMASTProxy();
+		}
+	}
+
+#if UE_BUILD_DEBUG
+	Proxy.DebugName = Proxy.Callstack.GetCallPath();
+#endif
+
+	return Proxy;
+}
+
diff --git a/Engine/Source/Developer/RigVMDeveloper/Private/RigVMCompiler/RigVMCompiler.cpp b/Engine/Source/Developer/RigVMDeveloper/Private/RigVMCompiler/RigVMCompiler.cpp
--- a/Engine/Source/Developer/RigVMDeveloper/Private/RigVMCompiler/RigVMCompiler.cpp
+++ b/Engine/Source/Developer/RigVMDeveloper/Private/RigVMCompiler/RigVMCompiler.cpp
@@ -385,7 +385,8 @@ void URigVMCompiler::TraverseEntry(const FRigVMEntryExprAST* InExpr, FRigVMCompi
 
 		if (Settings.SetupNodeInstructionIndex)
 		{
-			WorkData.VM->GetByteCode().SetSubject(EntryInstructionIndex, UnitNode);
+			const FString CallPath = InExpr->GetProxy().GetCallstack().GetCallPath();
+			WorkData.VM->GetByteCode().SetSubject(EntryInstructionIndex, UnitNode, CallPath);
 		}
 
 		TraverseChildren(InExpr, WorkData);
@@ -468,7 +469,8 @@ int32 URigVMCompiler::TraverseCallExtern(const FRigVMCallExternExprAST* InExpr,
 		InstructionIndex = WorkData.VM->GetByteCode().GetNumInstructions() - 1;
 		if (Settings.SetupNodeInstructionIndex)
 		{
-			WorkData.VM->GetByteCode().SetSubject(InstructionIndex, UnitNode);
+			const FString CallPath = InExpr->GetProxy().GetCallstack().GetCallPath();
+			WorkData.VM->GetByteCode().SetSubject(InstructionIndex, UnitNode, CallPath);
 		}
 
 		ensure(InExpr->NumChildren() == UnitNode->Pins.Num());
@@ -491,6 +493,15 @@ int32 URigVMCompiler::TraverseCallExtern(const FRigVMCallExternExprAST* InExpr,
 				{
 					FRigVMOperand SourceRootOperand(Source.GetMemoryType(), Source.GetRegisterIndex(), INDEX_NONE);
 					WorkData.VM->GetByteCode().AddCopyOp(SourceRootOperand, *Target);
+					const int32 CopyOpInstructionIndex = WorkData.VM->GetByteCode().GetNumInstructions() - 1;
+
+					if (Settings.SetupNodeInstructionIndex)
+					{
+						const FRigVMVarExprAST* VarExpr = PinExpr->To<FRigVMVarExprAST>();
+						const FRigVMASTProxy Proxy = VarExpr->GetProxy().GetSibling(VarExpr->GetPin()->GetNode());
+						const FString CallPath = Proxy.GetCallstack().GetCallPath();
+						WorkData.VM->GetByteCode().SetSubject(CopyOpInstructionIndex, Proxy.GetSubject(), CallPath);
+					}
 				}
 			}
 		}
@@ -508,6 +519,7 @@ void URigVMCompiler::TraverseForLoop(const FRigVMCallExternExprAST* InExpr, FRig
 	}
 
 	URigVMUnitNode* UnitNode = Cast<URigVMUnitNode>(InExpr->GetNode());
+	const FString CallPath = InExpr->GetProxy().GetCallstack().GetCallPath();
 
 	const FRigVMVarExprAST* CompletedExpr = InExpr->FindVarWithPinName(FRigVMStruct::ForLoopCompletedPinName);
 	check(CompletedExpr);
@@ -521,9 +533,17 @@ void URigVMCompiler::TraverseForLoop(const FRigVMCallExternExprAST* InExpr, FRig
 	check(IndexExpr);
 	FRigVMOperand IndexOperand = WorkData.ExprToOperand.FindChecked(IndexExpr);
 	WorkData.VM->GetByteCode().AddZeroOp(IndexOperand);
+	if (Settings.SetupNodeInstructionIndex)
+	{
+		WorkData.VM->GetByteCode().SetSubject(WorkData.VM->GetByteCode().GetNumInstructions() - 1, UnitNode, CallPath);
+	}
 
 	// call the for loop compute
 	int32 ForLoopInstructionIndex = TraverseCallExtern(InExpr, WorkData);
+	if (Settings.SetupNodeInstructionIndex)
+	{
+		WorkData.VM->GetByteCode().SetSubject(ForLoopInstructionIndex, UnitNode, CallPath);
+	}
 
 	// set up the jump forward (jump out of the loop)
 	const FRigVMVarExprAST* ContinueLoopExpr = InExpr->FindVarWithPinName(FRigVMStruct::ForLoopContinuePinName);
@@ -532,12 +552,20 @@ void URigVMCompiler::TraverseForLoop(const FRigVMCallExternExprAST* InExpr, FRig
 
 	uint64 JumpToEndByte = WorkData.VM->GetByteCode().AddJumpIfOp(ERigVMOpCode::JumpForwardIf, 0, ContinueLoopOperand, false);
 	int32 JumpToEndInstruction = WorkData.VM->GetByteCode().GetNumInstructions() - 1;
+	if (Settings.SetupNodeInstructionIndex)
+	{
+		WorkData.VM->GetByteCode().SetSubject(WorkData.VM->GetByteCode().GetNumInstructions() - 1, UnitNode, CallPath);
+	}
 
 	// begin the loop's block
 	const FRigVMVarExprAST* CountExpr = InExpr->FindVarWithPinName(FRigVMStruct::ForLoopCountPinName);
 	check(CountExpr);
 	FRigVMOperand CountOperand = WorkData.ExprToOperand.FindChecked(CountExpr);
 	WorkData.VM->GetByteCode().AddBeginBlockOp(CountOperand, IndexOperand);
+	if (Settings.SetupNodeInstructionIndex)
+	{
+		WorkData.VM->GetByteCode().SetSubject(WorkData.VM->GetByteCode().GetNumInstructions() - 1, UnitNode, CallPath);
+	}
 
 	// traverse the body of the loop
 	WorkData.ExprToSkip.Remove(ExecuteExpr);
@@ -545,13 +573,25 @@ void URigVMCompiler::TraverseForLoop(const FRigVMCallExternExprAST* InExpr, FRig
 
 	// end the loop's block
 	WorkData.VM->GetByteCode().AddEndBlockOp();
+	if (Settings.SetupNodeInstructionIndex)
+	{
+		WorkData.VM->GetByteCode().SetSubject(WorkData.VM->GetByteCode().GetNumInstructions() - 1, UnitNode, CallPath);
+	}
 
 	// increment the index
 	WorkData.VM->GetByteCode().AddIncrementOp(IndexOperand);
+	if (Settings.SetupNodeInstructionIndex)
+	{
+		WorkData.VM->GetByteCode().SetSubject(WorkData.VM->GetByteCode().GetNumInstructions() - 1, UnitNode, CallPath);
+	}
 
 	// jump to the beginning of the loop
 	int32 JumpToStartInstruction = WorkData.VM->GetByteCode().GetNumInstructions();
 	WorkData.VM->GetByteCode().AddJumpOp(ERigVMOpCode::JumpBackward, JumpToStartInstruction - ForLoopInstructionIndex);
+	if (Settings.SetupNodeInstructionIndex)
+	{
+		WorkData.VM->GetByteCode().SetSubject(WorkData.VM->GetByteCode().GetNumInstructions() - 1, UnitNode, CallPath);
+	}
 
 	// update the jump operator with the right address
 	int32 InstructionsToEnd = WorkData.VM->GetByteCode().GetNumInstructions() - JumpToEndInstruction;
@@ -722,7 +762,8 @@ void URigVMCompiler::TraverseAssign(const FRigVMAssignExprAST* InExpr, FRigVMCom
 				{
 					if (URigVMVariableNode* VariableNode = Cast<URigVMVariableNode>(SourcePin->GetNode()))
 					{
-						WorkData.VM->GetByteCode().SetSubject(InstructionIndex, VariableNode);
+						const FString CallPath = SourceExpr->GetProxy().GetSibling(VariableNode).GetCallstack().GetCallPath();
+						WorkData.VM->GetByteCode().SetSubject(InstructionIndex, VariableNode, CallPath);
 					}
 				}
 			}
@@ -732,7 +773,8 @@ void URigVMCompiler::TraverseAssign(const FRigVMAssignExprAST* InExpr, FRigVMCom
 				{
 					if (URigVMVariableNode* VariableNode = Cast<URigVMVariableNode>(TargetPin->GetNode()))
 					{
-						WorkData.VM->GetByteCode().SetSubject(InstructionIndex, VariableNode);
+						const FString CallPath = TargetExpr->GetProxy().GetSibling(VariableNode).GetCallstack().GetCallPath();
+						WorkData.VM->GetByteCode().SetSubject(InstructionIndex, VariableNode, CallPath);
 					}
 				}
 			}
@@ -747,6 +789,14 @@ void URigVMCompiler::TraverseAssign(const FRigVMAssignExprAST* InExpr, FRigVMCom
 				{
 					FRigVMOperand SourceRootOperand(Source.GetMemoryType(), Source.GetRegisterIndex(), INDEX_NONE);
 					WorkData.VM->GetByteCode().AddCopyOp(SourceRootOperand, WatchOperand);
+					int32 CopyInstructionIndex = WorkData.VM->GetByteCode().GetNumInstructions() - 1;
+
+					if (Settings.SetupNodeInstructionIndex)
+					{
+						URigVMNode* Node = SourceExpr->GetPin()->GetNode();
+						const FString CallPath = SourceExpr->GetProxy().GetSibling(Node).GetCallstack().GetCallPath();
+						WorkData.VM->GetByteCode().SetSubject(CopyInstructionIndex, Node, CallPath);
+					}
 				}
 			}
 		}
@@ -760,6 +810,14 @@ void URigVMCompiler::TraverseAssign(const FRigVMAssignExprAST* InExpr, FRigVMCom
 				{
 					FRigVMOperand TargetRootOperand(Target.GetMemoryType(), Target.GetRegisterIndex(), INDEX_NONE);
 					WorkData.VM->GetByteCode().AddCopyOp(TargetRootOperand, WatchOperand);
+					int32 CopyInstructionIndex = WorkData.VM->GetByteCode().GetNumInstructions() - 1;
+
+					if (Settings.SetupNodeInstructionIndex)
+					{
+						URigVMNode* Node = TargetExpr->GetPin()->GetNode();
+						const FString CallPath = TargetExpr->GetProxy().GetSibling(Node).GetCallstack().GetCallPath();
+						WorkData.VM->GetByteCode().SetSubject(CopyInstructionIndex, Node, CallPath);
+					}
 				}
 			}
 		}
@@ -796,6 +854,7 @@ void URigVMCompiler::TraverseBranch(const FRigVMBranchExprAST* InExpr, FRigVMCom
 	}
 
 	URigVMBranchNode* BranchNode = Cast<URigVMBranchNode>(InExpr->GetNode());
+	const FString CallPath = InExpr->GetProxy().GetCallstack().GetCallPath();
 
 	const FRigVMVarExprAST* ExecuteContextExpr = InExpr->ChildAt<FRigVMVarExprAST>(0);
 	const FRigVMVarExprAST* ConditionExpr = InExpr->ChildAt<FRigVMVarExprAST>(1);
@@ -817,7 +876,7 @@ void URigVMCompiler::TraverseBranch(const FRigVMBranchExprAST* InExpr, FRigVMCom
 	int32 JumpToFalseInstruction = WorkData.VM->GetByteCode().GetNumInstructions() - 1;
 	if (Settings.SetupNodeInstructionIndex)
 	{
-		WorkData.VM->GetByteCode().SetSubject(JumpToFalseInstruction, BranchNode);
+		WorkData.VM->GetByteCode().SetSubject(WorkData.VM->GetByteCode().GetNumInstructions() - 1, BranchNode, CallPath);
 	}
 
 	// traverse the true case
@@ -825,6 +884,10 @@ void URigVMCompiler::TraverseBranch(const FRigVMBranchExprAST* InExpr, FRigVMCom
 
 	uint64 JumpToEndByte = WorkData.VM->GetByteCode().AddJumpOp(ERigVMOpCode::JumpForward, 1);
 	int32 JumpToEndInstruction = WorkData.VM->GetByteCode().GetNumInstructions() - 1;
+	if (Settings.SetupNodeInstructionIndex)
+	{
+		WorkData.VM->GetByteCode().SetSubject(WorkData.VM->GetByteCode().GetNumInstructions() - 1, BranchNode, CallPath);
+	}
 
 	// correct the jump to false instruction index
 	int32 NumInstructionsInTrueCase = WorkData.VM->GetByteCode().GetNumInstructions() - JumpToFalseInstruction;
@@ -849,6 +912,7 @@ void URigVMCompiler::TraverseIf(const FRigVMIfExprAST* InExpr, FRigVMCompilerWor
 	}
 
 	URigVMIfNode* IfNode = Cast<URigVMIfNode>(InExpr->GetNode());
+	const FString CallPath = InExpr->GetProxy().GetCallstack().GetCallPath();
 
 	const FRigVMVarExprAST* ConditionExpr = InExpr->ChildAt<FRigVMVarExprAST>(0);
 	const FRigVMVarExprAST* TrueExpr = InExpr->ChildAt<FRigVMVarExprAST>(1);
@@ -871,7 +935,7 @@ void URigVMCompiler::TraverseIf(const FRigVMIfExprAST* InExpr, FRigVMCompilerWor
 	int32 JumpToFalseInstruction = WorkData.VM->GetByteCode().GetNumInstructions() - 1;
 	if (Settings.SetupNodeInstructionIndex)
 	{
-		WorkData.VM->GetByteCode().SetSubject(JumpToFalseInstruction, IfNode);
+		WorkData.VM->GetByteCode().SetSubject(WorkData.VM->GetByteCode().GetNumInstructions() - 1, IfNode, CallPath);
 	}
 
 	// traverse the true case
@@ -885,9 +949,17 @@ void URigVMCompiler::TraverseIf(const FRigVMIfExprAST* InExpr, FRigVMCompilerWor
 	FRigVMOperand& TrueOperand = WorkData.ExprToOperand.FindChecked(TrueExpr);
 
 	WorkData.VM->GetByteCode().AddCopyOp(TrueOperand, ResultOperand);
+	if (Settings.SetupNodeInstructionIndex)
+	{
+		WorkData.VM->GetByteCode().SetSubject(WorkData.VM->GetByteCode().GetNumInstructions() - 1, IfNode, CallPath);
+	}
 
 	uint64 JumpToEndByte = WorkData.VM->GetByteCode().AddJumpOp(ERigVMOpCode::JumpForward, 1);
 	int32 JumpToEndInstruction = WorkData.VM->GetByteCode().GetNumInstructions() - 1;
+	if (Settings.SetupNodeInstructionIndex)
+	{
+		WorkData.VM->GetByteCode().SetSubject(WorkData.VM->GetByteCode().GetNumInstructions() - 1, IfNode, CallPath);
+	}
 
 	// correct the jump to false instruction index
 	int32 NumInstructionsInTrueCase = WorkData.VM->GetByteCode().GetNumInstructions() - JumpToFalseInstruction;
@@ -904,6 +976,10 @@ void URigVMCompiler::TraverseIf(const FRigVMIfExprAST* InExpr, FRigVMCompilerWor
 	FRigVMOperand& FalseOperand = WorkData.ExprToOperand.FindChecked(FalseExpr);
 
 	WorkData.VM->GetByteCode().AddCopyOp(FalseOperand, ResultOperand);
+	if (Settings.SetupNodeInstructionIndex)
+	{
+		WorkData.VM->GetByteCode().SetSubject(WorkData.VM->GetByteCode().GetNumInstructions() - 1, IfNode, CallPath);
+	}
 
 	// correct the jump to end instruction index
 	int32 NumInstructionsInFalseCase = WorkData.VM->GetByteCode().GetNumInstructions() - JumpToEndInstruction;
@@ -913,6 +989,8 @@ void URigVMCompiler::TraverseIf(const FRigVMIfExprAST* InExpr, FRigVMCompilerWor
 void URigVMCompiler::TraverseSelect(const FRigVMSelectExprAST* InExpr, FRigVMCompilerWorkData& WorkData)
 {
 	URigVMSelectNode* SelectNode = Cast<URigVMSelectNode>(InExpr->GetNode());
+	const FString CallPath = InExpr->GetProxy().GetCallstack().GetCallPath();
+
 	int32 NumCases = SelectNode->FindPin(URigVMSelectNode::ValueName)->GetArraySize();
 
 	if (WorkData.bSetupMemory)
@@ -987,13 +1065,18 @@ void URigVMCompiler::TraverseSelect(const FRigVMSelectExprAST* InExpr, FRigVMCom
 	{
 		// compare and jump eventually
 		WorkData.VM->GetByteCode().AddEqualsOp(IndexOperand, WorkData.IntegerLiterals.FindChecked(CaseIndex), WorkData.ComparisonOperand);
-		if (CaseIndex == 1 && Settings.SetupNodeInstructionIndex)
+		if (Settings.SetupNodeInstructionIndex)
 		{
-			WorkData.VM->GetByteCode().SetSubject(WorkData.VM->GetByteCode().GetNumInstructions() - 1, SelectNode);
+			WorkData.VM->GetByteCode().SetSubject(WorkData.VM->GetByteCode().GetNumInstructions() - 1, SelectNode, CallPath);
 		}
 
 		uint64 JumpByte = WorkData.VM->GetByteCode().AddJumpIfOp(ERigVMOpCode::JumpForwardIf, 1, WorkData.ComparisonOperand, true);
 		int32 JumpInstruction = WorkData.VM->GetByteCode().GetNumInstructions() - 1;
+		if (Settings.SetupNodeInstructionIndex)
+		{
+			WorkData.VM->GetByteCode().SetSubject(WorkData.VM->GetByteCode().GetNumInstructions() - 1, SelectNode, CallPath);
+		}
+
 		JumpToCaseBytes.Add(JumpByte);
 		JumpToCaseInstructions.Add(JumpInstruction);
 	}
@@ -1014,11 +1097,20 @@ void URigVMCompiler::TraverseSelect(const FRigVMSelectExprAST* InExpr, FRigVMCom
 		// add copy op to copy the result
 		FRigVMOperand& CaseOperand = WorkData.ExprToOperand.FindChecked(CaseExpressions[CaseIndex]);
 		WorkData.VM->GetByteCode().AddCopyOp(CaseOperand, ResultOperand);
+		if (Settings.SetupNodeInstructionIndex)
+		{
+			WorkData.VM->GetByteCode().SetSubject(WorkData.VM->GetByteCode().GetNumInstructions() - 1, SelectNode, CallPath);
+		}
 		
 		if (CaseIndex < NumCases - 1)
 		{
 			uint64 JumpByte = WorkData.VM->GetByteCode().AddJumpOp(ERigVMOpCode::JumpForward, 1);
 			int32 JumpInstruction = WorkData.VM->GetByteCode().GetNumInstructions() - 1;
+			if (Settings.SetupNodeInstructionIndex)
+			{
+				WorkData.VM->GetByteCode().SetSubject(WorkData.VM->GetByteCode().GetNumInstructions() - 1, SelectNode, CallPath);
+			}
+
 			JumpToEndBytes.Add(JumpByte);
 			JumpToEndInstructions.Add(JumpInstruction);
 		}
diff --git a/Engine/Source/Developer/RigVMDeveloper/Private/RigVMModel/Nodes/RigVMLibraryNode.cpp b/Engine/Source/Developer/RigVMDeveloper/Private/RigVMModel/Nodes/RigVMLibraryNode.cpp
--- a/Engine/Source/Developer/RigVMDeveloper/Private/RigVMModel/Nodes/RigVMLibraryNode.cpp
+++ b/Engine/Source/Developer/RigVMDeveloper/Private/RigVMModel/Nodes/RigVMLibraryNode.cpp
@@ -31,6 +31,28 @@ bool URigVMLibraryNode::IsDefinedAsVarying() const
 	return false;
 }
 
+TArray<int32> URigVMLibraryNode::GetInstructionsForVM(URigVM* InVM, const FRigVMASTProxy& InProxy) const
+{
+	TArray<int32> Instructions;
+
+#if WITH_EDITOR
+
+	if(InVM == nullptr)
+	{
+		return Instructions;
+	}
+
+	const FRigVMASTProxy Proxy = InProxy.GetChild((UObject*)this);
+	for(URigVMNode* ContainedNode : GetContainedNodes())
+	{
+		Instructions.Append(ContainedNode->GetInstructionsForVM(InVM, Proxy));
+	}
+
+#endif
+	
+	return Instructions;
+}
+
 const TArray<URigVMNode*>& URigVMLibraryNode::GetContainedNodes() const
 {
 	if(URigVMGraph* Graph = GetContainedGraph())
diff --git a/Engine/Source/Developer/RigVMDeveloper/Private/RigVMModel/RigVMGraph.cpp b/Engine/Source/Developer/RigVMDeveloper/Private/RigVMModel/RigVMGraph.cpp
--- a/Engine/Source/Developer/RigVMDeveloper/Private/RigVMModel/RigVMGraph.cpp
+++ b/Engine/Source/Developer/RigVMDeveloper/Private/RigVMModel/RigVMGraph.cpp
@@ -41,6 +41,24 @@ TArray<URigVMGraph*> URigVMGraph::GetContainedGraphs(bool bRecursive) const
 	return Graphs;
 }
 
+URigVMGraph* URigVMGraph::GetParentGraph() const
+{
+	if(URigVMCollapseNode* CollapseNode = Cast<URigVMCollapseNode>(GetOuter()))
+	{
+		return CollapseNode->GetGraph();
+	}
+	return nullptr;
+}
+
+URigVMGraph* URigVMGraph::GetRootGraph() const
+{
+	if(URigVMGraph* ParentGraph = GetParentGraph())
+	{
+		return ParentGraph->GetRootGraph();
+	}
+	return (URigVMGraph*)this;
+}
+
 URigVMFunctionEntryNode* URigVMGraph::GetEntryNode() const
 {
 	for (URigVMNode* Node : Nodes)
diff --git a/Engine/Source/Developer/RigVMDeveloper/Private/RigVMModel/RigVMNode.cpp b/Engine/Source/Developer/RigVMDeveloper/Private/RigVMModel/RigVMNode.cpp
--- a/Engine/Source/Developer/RigVMDeveloper/Private/RigVMModel/RigVMNode.cpp
+++ b/Engine/Source/Developer/RigVMDeveloper/Private/RigVMModel/RigVMNode.cpp
@@ -161,6 +161,15 @@ URigVMGraph* URigVMNode::GetGraph() const
 	return nullptr;
 }
 
+URigVMGraph* URigVMNode::GetRootGraph() const
+{
+	if (URigVMGraph* Graph = GetGraph())
+	{
+		return Graph->GetRootGraph();
+	}
+	return nullptr;
+}
+
 URigVMInjectionInfo* URigVMNode::GetInjectionInfo() const
 {
 	return Cast<URigVMInjectionInfo>(GetOuter());
@@ -454,3 +463,50 @@ int32 URigVMNode::GetNumSlicesForContext(const FName& InContextName, const FRigV
 
 	return MaxSlices;
 }
+
+TArray<int32> URigVMNode::GetInstructionsForVM(URigVM* InVM, const FRigVMASTProxy& InProxy) const
+{
+	TArray<int32> Instructions;
+
+#if WITH_EDITOR
+
+	if(InVM == nullptr)
+	{
+		return Instructions;
+	}
+	
+	if(InProxy.IsValid())
+	{
+		const FRigVMASTProxy Proxy = InProxy.GetChild((UObject*)this);
+		const FString CallPath = Proxy.GetCallstack().GetCallPath();
+		return InVM->GetByteCode().GetAllInstructionIndicesForCallPath(CallPath, false, true);
+	}
+	else
+	{
+		return InVM->GetByteCode().GetAllInstructionIndicesForSubject((URigVMNode*)this);
+	}
+	
+#endif
+
+	return Instructions;
+}
+
+int32 URigVMNode::GetInstructionVisitedCount(URigVM* InVM, const FRigVMASTProxy& InProxy) const
+{
+#if WITH_EDITOR
+	if(InVM)
+	{
+		TArray<int32> Instructions = GetInstructionsForVM(InVM, InProxy);
+		if(Instructions.Num() > 0)
+		{
+			int32 Count = 0;
+			for(int32 Instruction : Instructions)
+			{
+				Count += InVM->GetInstructionVisitedCount(Instruction);
+			}
+			return Count;
+		}
+	}
+#endif
+	return 0;
+}
diff --git a/Engine/Source/Developer/RigVMDeveloper/Public/RigVMCompiler/RigVMASTProxy.h b/Engine/Source/Developer/RigVMDeveloper/Public/RigVMCompiler/RigVMASTProxy.h
--- a/Engine/Source/Developer/RigVMDeveloper/Public/RigVMCompiler/RigVMASTProxy.h
+++ b/Engine/Source/Developer/RigVMDeveloper/Public/RigVMCompiler/RigVMASTProxy.h
@@ -140,6 +140,7 @@ public:
 	}
 
 	static FRigVMASTProxy MakeFromUObject(UObject* InSubject);
+	static FRigVMASTProxy MakeFromCallPath(const FString& InCallPath, UObject* InRootObject);
 
 	FRigVMASTProxy GetSibling(UObject* InSubject) const
 	{
@@ -166,6 +167,12 @@ public:
 	FRigVMASTProxy GetChild(UObject* InSubject) const
 	{
 		FRigVMASTProxy Child = *this;
+
+		if(!Child.IsValid())
+		{
+			Child.Callstack.Stack.Reset();
+		}
+		
 		Child.Callstack.Stack.Push(InSubject);
 #if UE_BUILD_DEBUG
 		Child.DebugName = Child.Callstack.GetCallPath();
diff --git a/Engine/Source/Developer/RigVMDeveloper/Public/RigVMModel/Nodes/RigVMLibraryNode.h b/Engine/Source/Developer/RigVMDeveloper/Public/RigVMModel/Nodes/RigVMLibraryNode.h
--- a/Engine/Source/Developer/RigVMDeveloper/Public/RigVMModel/Nodes/RigVMLibraryNode.h
+++ b/Engine/Source/Developer/RigVMDeveloper/Public/RigVMModel/Nodes/RigVMLibraryNode.h
@@ -26,6 +26,7 @@ public:
 	// Override node functions
 	virtual bool IsDefinedAsConstant() const override;
 	virtual bool IsDefinedAsVarying() const override;
+	virtual TArray<int32> GetInstructionsForVM(URigVM* InVM, const FRigVMASTProxy& InProxy = FRigVMASTProxy()) const override; 
 
 	// Library node interface
 	virtual FString GetNodeCategory() const { return FString(); }
diff --git a/Engine/Source/Developer/RigVMDeveloper/Public/RigVMModel/RigVMGraph.h b/Engine/Source/Developer/RigVMDeveloper/Public/RigVMModel/RigVMGraph.h
--- a/Engine/Source/Developer/RigVMDeveloper/Public/RigVMModel/RigVMGraph.h
+++ b/Engine/Source/Developer/RigVMDeveloper/Public/RigVMModel/RigVMGraph.h
@@ -42,6 +42,14 @@ public:
 	UFUNCTION(BlueprintCallable, Category = RigVMGraph)
 	TArray<URigVMGraph*> GetContainedGraphs(bool bRecursive = false) const;
 
+	// Returns the parent graph of this graph
+	UFUNCTION(BlueprintCallable, Category = RigVMGraph)
+    URigVMGraph* GetParentGraph() const;
+
+	// Returns the root / top level parent graph of this graph (or this if it is the root graph)
+	UFUNCTION(BlueprintCallable, Category = RigVMGraph)
+    URigVMGraph* GetRootGraph() const;
+
 	// Returns the entry node of this graph
 	UFUNCTION(BlueprintCallable, Category = RigVMGraph)
 	URigVMFunctionEntryNode* GetEntryNode() const;
diff --git a/Engine/Source/Developer/RigVMDeveloper/Public/RigVMModel/RigVMNode.h b/Engine/Source/Developer/RigVMDeveloper/Public/RigVMModel/RigVMNode.h
--- a/Engine/Source/Developer/RigVMDeveloper/Public/RigVMModel/RigVMNode.h
+++ b/Engine/Source/Developer/RigVMDeveloper/Public/RigVMModel/RigVMNode.h
@@ -3,6 +3,7 @@
 #pragma once
 
 #include "RigVMPin.h"
+#include "RigVMCore/RigVM.h"
 #include "RigVMNode.generated.h"
 
 class URigVMGraph;
@@ -71,6 +72,10 @@ public:
 	UFUNCTION(BlueprintCallable, Category = RigVMNode)
 	URigVMGraph* GetGraph() const;
 
+	// Returns the top level / root Graph of this Node
+	UFUNCTION(BlueprintCallable, Category = RigVMNode)
+    URigVMGraph* GetRootGraph() const;
+
 	// Returns the injection info of this Node (or nullptr)
 	UFUNCTION(BlueprintCallable, Category = RigVMNode)
 	URigVMInjectionInfo* GetInjectionInfo() const;
@@ -184,6 +189,12 @@ public:
 	UFUNCTION(BlueprintCallable, Category = RigVMNode)
 	FName GetPreviousFName() const { return PreviousName; }
 
+	// Returns the indices of associated instructions for this node
+	virtual TArray<int32> GetInstructionsForVM(URigVM* InVM, const FRigVMASTProxy& InProxy = FRigVMASTProxy()) const; 
+
+	// Returns the number of visited / run instructions for this node
+	virtual int32 GetInstructionVisitedCount(URigVM* InVM, const FRigVMASTProxy& InProxy = FRigVMASTProxy()) const; 
+
 private:
 
 	static const FString NodeColorName;
diff --git a/Engine/Source/Runtime/RigVM/Private/RigVMCore/RigVMByteCode.cpp b/Engine/Source/Runtime/RigVM/Private/RigVMCore/RigVMByteCode.cpp
--- a/Engine/Source/Runtime/RigVM/Private/RigVMCore/RigVMByteCode.cpp
+++ b/Engine/Source/Runtime/RigVM/Private/RigVMCore/RigVMByteCode.cpp
@@ -534,6 +534,8 @@ void FRigVMByteCode::Reset()
 #if WITH_EDITORONLY_DATA
 	SubjectPerInstruction.Reset();
 	SubjectToInstructions.Reset();
+	CallPathPerInstruction.Reset();
+	CallPathToInstructions.Reset();
 #endif
 }
 
@@ -547,6 +549,8 @@ void FRigVMByteCode::Empty()
 #if WITH_EDITORONLY_DATA
 	SubjectPerInstruction.Empty();
 	SubjectToInstructions.Empty();
+	CallPathPerInstruction.Empty();
+	CallPathToInstructions.Empty();
 #endif
 }
 
@@ -1265,7 +1269,65 @@ const TArray<int32>& FRigVMByteCode::GetAllInstructionIndicesForSubject(UObject*
 	return EmptyInstructionIndices;
 }
 
-void FRigVMByteCode::SetSubject(int32 InInstructionIndex, UObject* InSubject)
+FString FRigVMByteCode::GetCallPathForInstruction(int32 InInstructionIndex) const
+{
+	if (CallPathPerInstruction.IsValidIndex(InInstructionIndex))
+	{
+		return CallPathPerInstruction[InInstructionIndex];
+	}
+	return FString();
+}
+
+int32 FRigVMByteCode::GetFirstInstructionIndexForCallPath(const FString& InCallPath, bool bStartsWith, bool bEndsWith) const
+{
+	const TArray<int32> InstructionIndices = GetAllInstructionIndicesForCallPath(InCallPath, bStartsWith, bEndsWith);
+	if (InstructionIndices.Num() > 0)
+	{
+		return InstructionIndices[0];
+	}
+	return INDEX_NONE;
+}
+
+TArray<int32> FRigVMByteCode::GetAllInstructionIndicesForCallPath(const FString& InCallPath, bool bStartsWith, bool bEndsWith) const
+{
+	if(InCallPath.IsEmpty())
+	{
+		return EmptyInstructionIndices;
+	}
+
+	TArray<int32> MatchedInstructions;
+	if (const TArray<int32>* InstructionIndices = CallPathToInstructions.Find(InCallPath))
+	{
+		MatchedInstructions.Append(*InstructionIndices);
+	}
+	
+	if(bStartsWith || bEndsWith)
+	{
+		const FString CallPathStart = FString::Printf(TEXT("%s|"), *InCallPath);
+		const FString CallPathEnd = FString::Printf(TEXT("|%s"), *InCallPath);
+		for(int32 InstructionIndex = 0; InstructionIndex < CallPathPerInstruction.Num(); InstructionIndex++)
+		{
+			if(bStartsWith)
+			{
+				if(CallPathPerInstruction[InstructionIndex].StartsWith(CallPathStart))
+				{
+					MatchedInstructions.Add(InstructionIndex);
+				}
+			}
+			else if(bEndsWith)
+			{
+				if(CallPathPerInstruction[InstructionIndex].EndsWith(CallPathEnd))
+				{
+					MatchedInstructions.Add(InstructionIndex);
+				}
+			}
+		}
+	}
+
+	return MatchedInstructions;
+}
+
+void FRigVMByteCode::SetSubject(int32 InInstructionIndex, UObject* InSubject, const FString& InCallPath)
 {
 	if (SubjectPerInstruction.Num() <= InInstructionIndex)
 	{
@@ -1273,6 +1335,13 @@ void FRigVMByteCode::SetSubject(int32 InInstructionIndex, UObject* InSubject)
 	}
 	SubjectPerInstruction[InInstructionIndex] = InSubject;
 	SubjectToInstructions.FindOrAdd(InSubject).Add(InInstructionIndex);
+
+	if (CallPathPerInstruction.Num() <= InInstructionIndex)
+	{
+		CallPathPerInstruction.AddZeroed(1 + InInstructionIndex - CallPathPerInstruction.Num());
+	}
+	CallPathPerInstruction[InInstructionIndex] = InCallPath;
+	CallPathToInstructions.FindOrAdd(InCallPath).Add(InInstructionIndex);
 }
 
 #endif
diff --git a/Engine/Source/Runtime/RigVM/Public/RigVMCore/RigVM.h b/Engine/Source/Runtime/RigVM/Public/RigVMCore/RigVM.h
--- a/Engine/Source/Runtime/RigVM/Public/RigVMCore/RigVM.h
+++ b/Engine/Source/Runtime/RigVM/Public/RigVMCore/RigVM.h
@@ -189,14 +189,21 @@ public:
 	TArray<FName> GetEntryNames() const;
 
 #if WITH_EDITOR
+	
 	// Returns true if the given instruction has been visited during the last run
 	FORCEINLINE bool WasInstructionVisitedDuringLastRun(int32 InIndex) const
+	{
+		return GetInstructionVisitedCount(InIndex) > 0;
+	}
+
+	// Returns the number of times an instruction has been hit
+	FORCEINLINE int32 GetInstructionVisitedCount(int32 InIndex) const
 	{
 		if (InstructionVisitedDuringLastRun.IsValidIndex(InIndex))
 		{
-			return InstructionVisitedDuringLastRun[InIndex] > 0;
+			return InstructionVisitedDuringLastRun[InIndex];
 		}
-		return false;
+		return 0;
 	}
 
 	// Returns the order of all instructions during the last run
diff --git a/Engine/Source/Runtime/RigVM/Public/RigVMCore/RigVMByteCode.h b/Engine/Source/Runtime/RigVM/Public/RigVMCore/RigVMByteCode.h
--- a/Engine/Source/Runtime/RigVM/Public/RigVMCore/RigVMByteCode.h
+++ b/Engine/Source/Runtime/RigVM/Public/RigVMCore/RigVMByteCode.h
@@ -5,6 +5,7 @@
 #include "CoreMinimal.h"
 #include "RigVMRegistry.h"
 #include "RigVMStatistics.h"
+
 #include "RigVMByteCode.generated.h"
 
 struct FRigVMByteCode;
@@ -673,6 +674,15 @@ public:
 	// returns all found instruction indices for a given subject
 	const TArray<int32>& GetAllInstructionIndicesForSubject(UObject* InSubject) const;
 
+	// returns the callpath which was used to inject a given instruction
+	FString GetCallPathForInstruction(int32 InInstructionIndex) const;
+
+	// returns the first hit instruction index for a given callpath (or INDEX_NONE)
+	int32 GetFirstInstructionIndexForCallPath(const FString& InCallPath, bool bStartsWith = false, bool bEndsWith = false) const;
+
+	// returns all found instruction indices for a given callpath
+	TArray<int32> GetAllInstructionIndicesForCallPath(const FString& InCallPath, bool bStartsWith = false, bool bEndsWith = false) const;
+
 #endif
 
 private:
@@ -701,12 +711,14 @@ private:
 
 	TArray<UObject*> SubjectPerInstruction;
 	TMap<UObject*, TArray<int32>> SubjectToInstructions;
+	TArray<FString> CallPathPerInstruction;
+	TMap<FString, TArray<int32>> CallPathToInstructions;
 
 #endif
 
 #if WITH_EDITOR
 
-	void SetSubject(int32 InInstructionIndex, UObject* InSubject);
+	void SetSubject(int32 InInstructionIndex, UObject* InSubject, const FString& InCallPath);
 
 #endif
 
