Control Rig: Wire execute pins to non-loop iterations by default

diff --git a/Engine/Source/Developer/RigVMDeveloper/Private/RigVMModel/RigVMController.cpp b/Engine/Source/Developer/RigVMDeveloper/Private/RigVMModel/RigVMController.cpp
--- a/Engine/Source/Developer/RigVMDeveloper/Private/RigVMModel/RigVMController.cpp
+++ b/Engine/Source/Developer/RigVMDeveloper/Private/RigVMModel/RigVMController.cpp
@@ -2614,17 +2614,52 @@ URigVMCollapseNode* URigVMController::CollapseNodes(const TArray<URigVMNode*>& I
 
 		// check if there is a last node on the top level block what we need to hook up
 		TMap<URigVMNode*, bool> IsContainedNodeLinkedToEntryNode;
-		for (URigVMNode* ContainedNode : CollapseNode->GetContainedNodes())
+
+		TArray<URigVMNode*> NodesForExecutePin;
+		NodesForExecutePin.Add(EntryNode);
+		for (int32 NodeForExecutePinIndex = 0; NodeForExecutePinIndex < NodesForExecutePin.Num(); NodeForExecutePinIndex++)
 		{
-			if (!ContainedNode->IsMutable())
+			URigVMNode* NodeForExecutePin = NodesForExecutePin[NodeForExecutePinIndex];
+			if (!NodeForExecutePin->IsMutable())
 			{
 				continue;
 			}
 
+			TArray<URigVMNode*> TargetNodes = NodeForExecutePin->GetLinkedTargetNodes();
+			for(URigVMNode* TargetNode : TargetNodes)
+			{
+				NodesForExecutePin.AddUnique(TargetNode);
+			}
+
 			// make sure the node doesn't have any mutable nodes connected to its executecontext
-			if (URigVMPin* ExecuteContextPin = ContainedNode->FindPin(FRigVMStruct::ExecuteContextName.ToString()))
+			URigVMPin* ExecuteContextPin = nullptr;
+			if(URigVMUnitNode* UnitNode = Cast<URigVMUnitNode>(NodeForExecutePin))
 			{
-				if (ExecuteContextPin->GetDirection() != ERigVMPinDirection::IO)
+				TSharedPtr<FStructOnScope> UnitScope = UnitNode->ConstructStructInstance();
+				if(UnitScope.IsValid())
+				{
+					FRigVMStruct* Unit = (FRigVMStruct*)UnitScope->GetStructMemory();
+					if(Unit->IsForLoop())
+					{
+						ExecuteContextPin = NodeForExecutePin->FindPin(FRigVMStruct::ForLoopCompletedPinName.ToString());
+					}
+				}
+			}
+
+			if(ExecuteContextPin == nullptr)
+			{
+				ExecuteContextPin = NodeForExecutePin->FindPin(FRigVMStruct::ExecuteContextName.ToString());
+			}
+
+			if(ExecuteContextPin)
+			{
+				if(!ExecuteContextPin->IsExecuteContext())
+				{
+					continue;
+				}
+
+				if (ExecuteContextPin->GetDirection() != ERigVMPinDirection::IO &&
+					ExecuteContextPin->GetDirection() != ERigVMPinDirection::Output)
 				{
 					continue;
 				}
@@ -2634,7 +2669,7 @@ URigVMCollapseNode* URigVMController::CollapseNodes(const TArray<URigVMNode*>& I
 					continue;
 				}
 
-				if (!Local::IsLinkedToEntryNode(ContainedNode, IsContainedNodeLinkedToEntryNode))
+				if (!Local::IsLinkedToEntryNode(NodeForExecutePin, IsContainedNodeLinkedToEntryNode))
 				{
 					continue;
 				}
