Control Rig: Improve count consolidatoin for RigVMNode / RigVMLibrary node

diff --git a/Engine/Source/Developer/RigVMDeveloper/Private/RigVMModel/Nodes/RigVMLibraryNode.cpp b/Engine/Source/Developer/RigVMDeveloper/Private/RigVMModel/Nodes/RigVMLibraryNode.cpp
--- a/Engine/Source/Developer/RigVMDeveloper/Private/RigVMModel/Nodes/RigVMLibraryNode.cpp
+++ b/Engine/Source/Developer/RigVMDeveloper/Private/RigVMModel/Nodes/RigVMLibraryNode.cpp
@@ -53,6 +53,29 @@ TArray<int32> URigVMLibraryNode::GetInstructionsForVM(URigVM* InVM, const FRigVM
 	return Instructions;
 }
 
+int32 URigVMLibraryNode::GetInstructionVisitedCount(URigVM* InVM, const FRigVMASTProxy& InProxy, bool bConsolidatePerNode) const
+{
+	int32 Count = 0;
+
+#if WITH_EDITOR
+
+	if(InVM == nullptr)
+	{
+		return Count;
+	}
+
+	const FRigVMASTProxy Proxy = InProxy.GetChild((UObject*)this);
+	for(URigVMNode* ContainedNode : GetContainedNodes())
+	{
+		const int32 CountPerNode = ContainedNode->GetInstructionVisitedCount(InVM, Proxy, bConsolidatePerNode);
+		Count += CountPerNode;
+	}
+
+#endif
+
+	return Count;
+}
+
 const TArray<URigVMNode*>& URigVMLibraryNode::GetContainedNodes() const
 {
 	if(URigVMGraph* Graph = GetContainedGraph())
diff --git a/Engine/Source/Developer/RigVMDeveloper/Private/RigVMModel/RigVMNode.cpp b/Engine/Source/Developer/RigVMDeveloper/Private/RigVMModel/RigVMNode.cpp
--- a/Engine/Source/Developer/RigVMDeveloper/Private/RigVMModel/RigVMNode.cpp
+++ b/Engine/Source/Developer/RigVMDeveloper/Private/RigVMModel/RigVMNode.cpp
@@ -491,7 +491,7 @@ TArray<int32> URigVMNode::GetInstructionsForVM(URigVM* InVM, const FRigVMASTProx
 	return Instructions;
 }
 
-int32 URigVMNode::GetInstructionVisitedCount(URigVM* InVM, const FRigVMASTProxy& InProxy) const
+int32 URigVMNode::GetInstructionVisitedCount(URigVM* InVM, const FRigVMASTProxy& InProxy, bool bConsolidatePerNode) const
 {
 #if WITH_EDITOR
 	if(InVM)
@@ -502,7 +502,15 @@ int32 URigVMNode::GetInstructionVisitedCount(URigVM* InVM, const FRigVMASTProxy&
 			int32 Count = 0;
 			for(int32 Instruction : Instructions)
 			{
-				Count += InVM->GetInstructionVisitedCount(Instruction);
+				const int32 CountPerInstruction = InVM->GetInstructionVisitedCount(Instruction);
+				if(bConsolidatePerNode)
+				{
+					Count = FMath::Max<int32>(Count, CountPerInstruction);
+				}
+				else
+				{
+					Count += CountPerInstruction;
+				}
 			}
 			return Count;
 		}
diff --git a/Engine/Source/Developer/RigVMDeveloper/Public/RigVMModel/Nodes/RigVMLibraryNode.h b/Engine/Source/Developer/RigVMDeveloper/Public/RigVMModel/Nodes/RigVMLibraryNode.h
--- a/Engine/Source/Developer/RigVMDeveloper/Public/RigVMModel/Nodes/RigVMLibraryNode.h
+++ b/Engine/Source/Developer/RigVMDeveloper/Public/RigVMModel/Nodes/RigVMLibraryNode.h
@@ -27,6 +27,7 @@ public:
 	virtual bool IsDefinedAsConstant() const override;
 	virtual bool IsDefinedAsVarying() const override;
 	virtual TArray<int32> GetInstructionsForVM(URigVM* InVM, const FRigVMASTProxy& InProxy = FRigVMASTProxy()) const override; 
+	virtual int32 GetInstructionVisitedCount(URigVM* InVM, const FRigVMASTProxy& InProxy = FRigVMASTProxy(), bool bConsolidatePerNode = false) const override; 
 
 	// Library node interface
 	virtual FString GetNodeCategory() const { return FString(); }
diff --git a/Engine/Source/Developer/RigVMDeveloper/Public/RigVMModel/RigVMNode.h b/Engine/Source/Developer/RigVMDeveloper/Public/RigVMModel/RigVMNode.h
--- a/Engine/Source/Developer/RigVMDeveloper/Public/RigVMModel/RigVMNode.h
+++ b/Engine/Source/Developer/RigVMDeveloper/Public/RigVMModel/RigVMNode.h
@@ -193,7 +193,7 @@ public:
 	virtual TArray<int32> GetInstructionsForVM(URigVM* InVM, const FRigVMASTProxy& InProxy = FRigVMASTProxy()) const; 
 
 	// Returns the number of visited / run instructions for this node
-	virtual int32 GetInstructionVisitedCount(URigVM* InVM, const FRigVMASTProxy& InProxy = FRigVMASTProxy()) const; 
+	virtual int32 GetInstructionVisitedCount(URigVM* InVM, const FRigVMASTProxy& InProxy = FRigVMASTProxy(), bool bConsolidatePerNode = false) const; 
 
 private:
 
