ControlRig: Fix watches not working for functions and collapsed nodes.

diff --git a/Engine/Plugins/Experimental/ControlRig/Source/ControlRigDeveloper/Private/ControlRigBlueprint.cpp b/Engine/Plugins/Experimental/ControlRig/Source/ControlRigDeveloper/Private/ControlRigBlueprint.cpp
--- a/Engine/Plugins/Experimental/ControlRig/Source/ControlRigDeveloper/Private/ControlRigBlueprint.cpp
+++ b/Engine/Plugins/Experimental/ControlRig/Source/ControlRigDeveloper/Private/ControlRigBlueprint.cpp
@@ -2062,9 +2062,9 @@ void UControlRigBlueprint::HandleModifiedEvent(ERigVMGraphNotifType InNotifType,
 			{
 				if (URigVMPin* ModelPin = Cast<URigVMPin>(InSubject))
 				{
-					for (int32 i = 0; i < UbergraphPages.Num(); ++i)
-					{
-						if (UControlRigGraph* Graph = Cast<UControlRigGraph>(UbergraphPages[i]))
+					if (UEdGraph* EdGraph = GetEdGraph(InGraph))
+					{							
+						if (UControlRigGraph* Graph = Cast<UControlRigGraph>(EdGraph))
 						{
 							if (UEdGraphNode* EdNode = Graph->FindNodeForModelNodeName(ModelPin->GetNode()->GetFName()))
 							{
@@ -2079,9 +2079,9 @@ void UControlRigBlueprint::HandleModifiedEvent(ERigVMGraphNotifType InNotifType,
 										WatchedPins.Remove(EdPin);
 									}
 									RequestAutoVMRecompilation();
-										MarkPackageDirty();
-									}
+									MarkPackageDirty();
 								}
+							}
 						}
 					}
 				}
diff --git a/Engine/Plugins/Experimental/ControlRig/Source/ControlRigEditor/Private/ControlRigEditorModule.cpp b/Engine/Plugins/Experimental/ControlRig/Source/ControlRigEditor/Private/ControlRigEditorModule.cpp
--- a/Engine/Plugins/Experimental/ControlRig/Source/ControlRigEditor/Private/ControlRigEditorModule.cpp
+++ b/Engine/Plugins/Experimental/ControlRig/Source/ControlRigEditor/Private/ControlRigEditorModule.cpp
@@ -1235,20 +1235,26 @@ void FControlRigEditorModule::GetContextMenuActions(const UControlRigGraphSchema
 		if (UEdGraphPin* InGraphPin = (UEdGraphPin* )Context->Pin)
 		{
 			UEdGraph* Graph = InGraphPin->GetOwningNode()->GetGraph();
-
-			// Add the watch pin / unwatch pin menu items
+			
+			if(UControlRigGraphNode* RigNode = Cast<UControlRigGraphNode>(InGraphPin->GetOwningNode()))
 			{
-				FToolMenuSection& Section = Menu->AddSection("EdGraphSchemaWatches", LOCTEXT("WatchesHeader", "Watches"));
-				UBlueprint* OwnerBlueprint = FBlueprintEditorUtils::FindBlueprintForGraphChecked(Context->Graph);
+				// Check if we can watch this value (not inside any function)
+				bool bIsInsideFunction = RigNode->GetModel()->GetRootGraph()->IsA<URigVMFunctionLibrary>();
+				if(!bIsInsideFunction)
 				{
-					if (FKismetDebugUtilities::IsPinBeingWatched(OwnerBlueprint, InGraphPin))
+					// Add the watch pin / unwatch pin menu items
+					FToolMenuSection& Section = Menu->AddSection("EdGraphSchemaWatches", LOCTEXT("WatchesHeader", "Watches"));
+					UBlueprint* OwnerBlueprint = FBlueprintEditorUtils::FindBlueprintForGraphChecked(Context->Graph);
 					{
-						Section.AddMenuEntry(FGraphEditorCommands::Get().StopWatchingPin);
-					}
-					else
-					{
-						Section.AddMenuEntry(FGraphEditorCommands::Get().StartWatchingPin);
-					}
+						if (FKismetDebugUtilities::IsPinBeingWatched(OwnerBlueprint, InGraphPin))
+						{
+							Section.AddMenuEntry(FGraphEditorCommands::Get().StopWatchingPin);
+						}
+						else
+						{
+							Section.AddMenuEntry(FGraphEditorCommands::Get().StartWatchingPin);
+						}
+					}				
 				}
 			}
 
diff --git a/Engine/Source/Developer/RigVMDeveloper/Private/RigVMCompiler/RigVMAST.cpp b/Engine/Source/Developer/RigVMDeveloper/Private/RigVMCompiler/RigVMAST.cpp
--- a/Engine/Source/Developer/RigVMDeveloper/Private/RigVMCompiler/RigVMAST.cpp
+++ b/Engine/Source/Developer/RigVMDeveloper/Private/RigVMCompiler/RigVMAST.cpp
@@ -1653,7 +1653,7 @@ void FRigVMParserAST::FoldAssignments()
 		// if this node is a variable node and the pin requires a watch... skip this
 		if (Cast<URigVMVariableNode>(AssignExpr->GetSourcePin()->GetNode()))
 		{
-			if(AssignExpr->GetSourcePin()->RequiresWatch())
+			if(AssignExpr->GetSourcePin()->RequiresWatch(true))
 			{
 				continue;
 			}
diff --git a/Engine/Source/Developer/RigVMDeveloper/Private/RigVMCompiler/RigVMCompiler.cpp b/Engine/Source/Developer/RigVMDeveloper/Private/RigVMCompiler/RigVMCompiler.cpp
--- a/Engine/Source/Developer/RigVMDeveloper/Private/RigVMCompiler/RigVMCompiler.cpp
+++ b/Engine/Source/Developer/RigVMDeveloper/Private/RigVMCompiler/RigVMCompiler.cpp
@@ -410,7 +410,7 @@ int32 URigVMCompiler::TraverseCallExtern(const FRigVMCallExternExprAST* InExpr,
 		{
 			for (URigVMPin* Pin : UnitNode->Pins)
 			{
-				if (Pin->RequiresWatch())
+				if (Pin->RequiresWatch(true))
 				{
 					FRigVMASTProxy PinProxy = FRigVMASTProxy::MakeFromUObject(Pin);
 					FRigVMVarExprAST TempVarExpr(FRigVMExprAST::EType::Var, PinProxy);
@@ -479,7 +479,7 @@ int32 URigVMCompiler::TraverseCallExtern(const FRigVMCallExternExprAST* InExpr,
 			URigVMPin* Pin = UnitNode->Pins[PinIndex];
 
 			// ensure to copy the debug values
-			if (Pin->RequiresWatch() && Settings.EnablePinWatches)
+			if (Pin->RequiresWatch(true) && Settings.EnablePinWatches)
 			{
 				const FRigVMExprAST* PinExpr = InExpr->ChildAt(PinIndex);
 				if (PinExpr->IsA(FRigVMExprAST::EType::CachedValue))
@@ -782,7 +782,7 @@ void URigVMCompiler::TraverseAssign(const FRigVMAssignExprAST* InExpr, FRigVMCom
 
 		if (Cast<URigVMVariableNode>(SourceExpr->GetPin()->GetNode()))
 		{
-			if (SourceExpr->GetPin()->RequiresWatch() && Settings.EnablePinWatches)
+			if (SourceExpr->GetPin()->RequiresWatch(true) && Settings.EnablePinWatches)
 			{
 				FRigVMOperand WatchOperand = FindOrAddRegister(SourceExpr, WorkData, true /*debug watch*/);
 				if (WatchOperand.IsValid())
@@ -803,7 +803,7 @@ void URigVMCompiler::TraverseAssign(const FRigVMAssignExprAST* InExpr, FRigVMCom
 
 		if (Cast<URigVMVariableNode>(TargetExpr->GetPin()->GetNode()))
 		{
-			if (TargetExpr->GetPin()->RequiresWatch() && Settings.EnablePinWatches)
+			if (TargetExpr->GetPin()->RequiresWatch(true) && Settings.EnablePinWatches)
 			{
 				FRigVMOperand WatchOperand = FindOrAddRegister(TargetExpr, WorkData, true /*debug watch*/);
 				if (WatchOperand.IsValid())
@@ -1123,7 +1123,7 @@ void URigVMCompiler::TraverseSelect(const FRigVMSelectExprAST* InExpr, FRigVMCom
 	}
 }
 
-FString URigVMCompiler::GetPinHash(URigVMPin* InPin, const FRigVMVarExprAST* InVarExpr, bool bIsDebugValue)
+FString URigVMCompiler::GetPinHash(const URigVMPin* InPin, const FRigVMVarExprAST* InVarExpr, bool bIsDebugValue)
 {
 	FString Prefix = bIsDebugValue ? TEXT("DebugWatch:") : TEXT("");
 	FString Suffix;
@@ -1711,7 +1711,15 @@ FRigVMOperand URigVMCompiler::FindOrAddRegister(const FRigVMVarExprAST* InVarExp
 	}
 	ensure(Operand.IsValid());
 
-	WorkData.PinPathToOperand->Add(Hash, Operand);
+	// Get all possible pins that lead to the same operand
+	TArray<const URigVMPin*> VirtualPins;
+	Pin->GetExposedPinChain(VirtualPins);
+	for (const URigVMPin* VirtualPin : VirtualPins)
+	{
+		FString VirtualPinHash = GetPinHash(VirtualPin, InVarExpr, bIsDebugValue);
+		WorkData.PinPathToOperand->Add(VirtualPinHash, Operand);	
+	}
+	
 	WorkData.ExprToOperand.Add(InVarExpr, Operand);
 
 	return Operand;
diff --git a/Engine/Source/Developer/RigVMDeveloper/Private/RigVMModel/RigVMPin.cpp b/Engine/Source/Developer/RigVMDeveloper/Private/RigVMModel/RigVMPin.cpp
--- a/Engine/Source/Developer/RigVMDeveloper/Private/RigVMModel/RigVMPin.cpp
+++ b/Engine/Source/Developer/RigVMDeveloper/Private/RigVMModel/RigVMPin.cpp
@@ -12,7 +12,10 @@
 #include "Misc/OutputDevice.h"
 #include "Misc/DefaultValueHelper.h"
 #include "Logging/LogScopedVerbosityOverride.h"
-
+#include "RigVMModel/Nodes/RigVMCollapseNode.h"
+#include "RigVMModel/Nodes/RigVMFunctionReferenceNode.h"
+#include "RigVMModel/Nodes/RigVMFunctionEntryNode.h"
+#include "RigVMModel/Nodes/RigVMFunctionReturnNode.h"
 
 class FRigVMPinDefaultValueImportErrorContext : public FOutputDevice
 {
@@ -217,6 +220,69 @@ FString URigVMPin::GetSegmentPath() const
 	return FString();
 }
 
+void URigVMPin::GetExposedPinChain(TArray<const URigVMPin*>& OutExposedPins) const
+{
+	// Find the first pin in the chain (source)
+	for (URigVMLink* Link : GetSourceLinks())
+	{
+		URigVMPin* SourcePin = Link->GetSourcePin();
+		
+		// If the source is on an entry node, add the pin and make a recursive call on the collapse node pin
+		if (URigVMFunctionEntryNode* EntryNode = Cast<URigVMFunctionEntryNode>(SourcePin->GetNode()))
+		{
+			URigVMGraph* Graph = EntryNode->GetGraph();
+			if (URigVMCollapseNode* CollapseNode = Cast<URigVMCollapseNode>(Graph->GetOuter()))
+			{
+				URigVMPin* CollapseNodePin = CollapseNode->FindPin(SourcePin->GetName());
+				CollapseNodePin->GetExposedPinChain(OutExposedPins);				
+			}
+		}
+		else if (URigVMFunctionReturnNode* ReturnNode = Cast<URigVMFunctionReturnNode>(SourcePin->GetNode()))
+		{
+			URigVMGraph* Graph = ReturnNode->GetGraph();
+			if (URigVMCollapseNode* CollapseNode = Cast<URigVMCollapseNode>(Graph->GetOuter()))
+			{
+				URigVMPin* CollapseNodePin = CollapseNode->FindPin(SourcePin->GetName());
+				CollapseNodePin->GetExposedPinChain(OutExposedPins);				
+			}
+		}
+		else
+		{
+			SourcePin->GetExposedPinChain(OutExposedPins);
+		}
+
+		return;
+	}
+
+
+	// Add the pins in the OutExposedPins array in depth-first order
+	TArray<const URigVMPin*> ToProcess;
+	ToProcess.Push(this);
+	while (!ToProcess.IsEmpty())
+	{
+		const URigVMPin* Current = ToProcess.Pop();
+		OutExposedPins.Add(Current);
+
+		// Add target pins connected to the current pin
+		for (URigVMLink* Link : Current->GetTargetLinks())
+		{
+			URigVMPin* TargetPin = Link->GetTargetPin();
+			ToProcess.Push(TargetPin);
+		}
+
+		// If pin is on a collapse node, add entry pin
+		if (URigVMCollapseNode* CollapseNode = Cast<URigVMCollapseNode>(Current->GetNode()))
+		{
+			URigVMFunctionEntryNode* EntryNode = CollapseNode->GetEntryNode();
+			URigVMPin* EntryPin = EntryNode->FindPin(Current->GetName());
+			if (EntryPin)
+			{
+				ToProcess.Push(EntryPin);
+			}
+		}		
+	}		
+}
+
 FName URigVMPin::GetDisplayName() const
 {
 	if (DisplayName == NAME_None)
@@ -262,8 +328,22 @@ bool URigVMPin::IsDefinedAsConstant() const
 	return bIsConstant;
 }
 
-bool URigVMPin::RequiresWatch() const
+bool URigVMPin::RequiresWatch(const bool bCheckExposedPinChain) const
 {
+	if (!bRequiresWatch && bCheckExposedPinChain)
+	{
+		TArray<const URigVMPin*> VirtualPins;
+		GetExposedPinChain(VirtualPins);
+		
+		for (const URigVMPin* VirtualPin : VirtualPins)
+		{
+			if (VirtualPin->bRequiresWatch)
+			{
+				return true;
+			}
+		}		
+	}
+	
 	return bRequiresWatch;
 }
 
diff --git a/Engine/Source/Developer/RigVMDeveloper/Public/RigVMCompiler/RigVMCompiler.h b/Engine/Source/Developer/RigVMDeveloper/Public/RigVMCompiler/RigVMCompiler.h
--- a/Engine/Source/Developer/RigVMDeveloper/Public/RigVMCompiler/RigVMCompiler.h
+++ b/Engine/Source/Developer/RigVMDeveloper/Public/RigVMCompiler/RigVMCompiler.h
@@ -120,7 +120,7 @@ public:
 	bool Compile(URigVMGraph* InGraph, URigVMController* InController, URigVM* OutVM, const TArray<FRigVMExternalVariable>& InExternalVariables, const TArray<FRigVMUserDataArray>& InRigVMUserData, TMap<FString, FRigVMOperand>* OutOperands, TSharedPtr<FRigVMParserAST> InAST = TSharedPtr<FRigVMParserAST>());
 
 	static UScriptStruct* GetScriptStructForCPPType(const FString& InCPPType);
-	static FString GetPinHash(URigVMPin* InPin, const FRigVMVarExprAST* InVarExpr, bool bIsDebugValue = false);
+	static FString GetPinHash(const URigVMPin* InPin, const FRigVMVarExprAST* InVarExpr, bool bIsDebugValue = false);
 
 private:
 
diff --git a/Engine/Source/Developer/RigVMDeveloper/Public/RigVMModel/RigVMPin.h b/Engine/Source/Developer/RigVMDeveloper/Public/RigVMModel/RigVMPin.h
--- a/Engine/Source/Developer/RigVMDeveloper/Public/RigVMModel/RigVMPin.h
+++ b/Engine/Source/Developer/RigVMDeveloper/Public/RigVMModel/RigVMPin.h
@@ -130,6 +130,11 @@ public:
 	UFUNCTION(BlueprintCallable, Category = RigVMPin)
 	FString GetSegmentPath() const;
 
+	// Populates an array of pins which will be reduced to the same operand in the
+	// VM. This includes Source-Target pins in different nodes, pins in collapse and
+	// referenced function nodes, and their corresponding entry and return nodes.
+	void GetExposedPinChain(TArray<const URigVMPin*>& OutExposedPins) const;
+
 	// Returns the display label of the pin
 	UFUNCTION(BlueprintCallable, Category = RigVMPin)
 	FName GetDisplayName() const;
@@ -148,7 +153,7 @@ public:
 
 	// Returns true if the pin should be watched
 	UFUNCTION(BlueprintCallable, Category = RigVMPin)
-	bool RequiresWatch() const;
+	bool RequiresWatch(const bool bCheckExposedPinChain = false) const;
 
 	// Returns true if the data type of the Pin is a struct
 	UFUNCTION(BlueprintCallable, Category = RigVMPin)
