URigVMPin::GetExposedPinChain: Add a visited set to the DFS algorithm to prevent infinite loop due to cycles.
#rnx

diff --git a/Engine/Source/Developer/RigVMDeveloper/Private/RigVMModel/RigVMPin.cpp b/Engine/Source/Developer/RigVMDeveloper/Private/RigVMModel/RigVMPin.cpp
--- a/Engine/Source/Developer/RigVMDeveloper/Private/RigVMModel/RigVMPin.cpp
+++ b/Engine/Source/Developer/RigVMDeveloper/Private/RigVMModel/RigVMPin.cpp
@@ -256,11 +256,17 @@ void URigVMPin::GetExposedPinChain(TArray<const URigVMPin*>& OutExposedPins) con
 
 
 	// Add the pins in the OutExposedPins array in depth-first order
+	TSet<const URigVMPin*> FoundPins;
 	TArray<const URigVMPin*> ToProcess;
 	ToProcess.Push(this);
 	while (!ToProcess.IsEmpty())
 	{
 		const URigVMPin* Current = ToProcess.Pop();
+		if (FoundPins.Contains(Current))
+		{
+			continue;
+		}
+		FoundPins.Add(Current);
 		OutExposedPins.Add(Current);
 
 		// Add target pins connected to the current pin
