ControlRig: Fix Variable nodes to storing the correct value

diff --git a/Engine/Source/Developer/RigVMDeveloper/Private/RigVMModel/RigVMPin.cpp b/Engine/Source/Developer/RigVMDeveloper/Private/RigVMModel/RigVMPin.cpp
--- a/Engine/Source/Developer/RigVMDeveloper/Private/RigVMModel/RigVMPin.cpp
+++ b/Engine/Source/Developer/RigVMDeveloper/Private/RigVMModel/RigVMPin.cpp
@@ -222,6 +222,13 @@ FString URigVMPin::GetSegmentPath() const
 
 void URigVMPin::GetExposedPinChain(TArray<const URigVMPin*>& OutExposedPins) const
 {
+	// Variable nodes do not share the operand with their source link
+	if (GetNode()->IsA<URigVMVariableNode>() && GetDirection() == ERigVMPinDirection::Input)
+	{
+		OutExposedPins.Add(this);
+		return;
+	}
+	
 	// Find the first pin in the chain (source)
 	for (URigVMLink* Link : GetSourceLinks())
 	{
@@ -246,6 +253,11 @@ void URigVMPin::GetExposedPinChain(TArray<const URigVMPin*>& OutExposedPins) con
 				CollapseNodePin->GetExposedPinChain(OutExposedPins);				
 			}
 		}
+		// Variable nodes do not share the operand with their source link
+		else if (SourcePin->GetNode()->IsA<URigVMVariableNode>())
+		{
+			continue;
+		}
 		else
 		{
 			SourcePin->GetExposedPinChain(OutExposedPins);
@@ -254,7 +266,6 @@ void URigVMPin::GetExposedPinChain(TArray<const URigVMPin*>& OutExposedPins) con
 		return;
 	}
 
-
 	// Add the pins in the OutExposedPins array in depth-first order
 	TSet<const URigVMPin*> FoundPins;
 	TArray<const URigVMPin*> ToProcess;
@@ -273,6 +284,12 @@ void URigVMPin::GetExposedPinChain(TArray<const URigVMPin*>& OutExposedPins) con
 		for (URigVMLink* Link : Current->GetTargetLinks())
 		{
 			URigVMPin* TargetPin = Link->GetTargetPin();
+
+			// Variable nodes do not share the operand with their source link
+			if (TargetPin->GetNode()->IsA<URigVMVariableNode>())
+			{
+				continue;
+			}
 			ToProcess.Push(TargetPin);
 		}
 
@@ -285,7 +302,17 @@ void URigVMPin::GetExposedPinChain(TArray<const URigVMPin*>& OutExposedPins) con
 			{
 				ToProcess.Push(EntryPin);
 			}
-		}		
+		}
+		// If pin is on a return node, add parent pin on collapse node
+		else if (URigVMFunctionReturnNode* ReturnNode = Cast<URigVMFunctionReturnNode>(Current->GetNode()))
+		{
+			URigVMGraph* Graph = ReturnNode->GetGraph();
+			if (URigVMCollapseNode* ParentNode = Cast<URigVMCollapseNode>(Graph->GetOuter()))
+			{
+				URigVMPin* CollapseNodePin = ParentNode->FindPin(Current->GetName());
+				ToProcess.Push(CollapseNodePin);
+			}
+		}
 	}		
 }
 
