Control Rig: Fix automation test around RigVM Compiler Array Array Link

diff --git a/Engine/Source/Developer/RigVMDeveloper/Private/RigVMCompiler/RigVMCompiler.cpp b/Engine/Source/Developer/RigVMDeveloper/Private/RigVMCompiler/RigVMCompiler.cpp
--- a/Engine/Source/Developer/RigVMDeveloper/Private/RigVMCompiler/RigVMCompiler.cpp
+++ b/Engine/Source/Developer/RigVMDeveloper/Private/RigVMCompiler/RigVMCompiler.cpp
@@ -1712,14 +1712,21 @@ FRigVMOperand URigVMCompiler::FindOrAddRegister(const FRigVMVarExprAST* InVarExp
 	ensure(Operand.IsValid());
 
 	// Get all possible pins that lead to the same operand
-	TArray<const URigVMPin*> VirtualPins;
-	Pin->GetExposedPinChain(VirtualPins);
-	for (const URigVMPin* VirtualPin : VirtualPins)
+
+	if(Settings.ASTSettings.bFoldAssignments)
+	{
+		TArray<const URigVMPin*> VirtualPins;
+		Pin->GetExposedPinChain(VirtualPins);
+		for (const URigVMPin* VirtualPin : VirtualPins)
+		{
+			FString VirtualPinHash = GetPinHash(VirtualPin, InVarExpr, bIsDebugValue);
+			WorkData.PinPathToOperand->Add(VirtualPinHash, Operand);	
+		}
+	}
+	else
 	{
-		FString VirtualPinHash = GetPinHash(VirtualPin, InVarExpr, bIsDebugValue);
-		WorkData.PinPathToOperand->Add(VirtualPinHash, Operand);	
+		WorkData.PinPathToOperand->Add(Hash, Operand);
 	}
-	
 	WorkData.ExprToOperand.Add(InVarExpr, Operand);
 
 	return Operand;
