ControlRig: Revert changes breaking unit tests.

diff --git a/Engine/Source/Developer/RigVMDeveloper/Private/RigVMCompiler/RigVMCompiler.cpp b/Engine/Source/Developer/RigVMDeveloper/Private/RigVMCompiler/RigVMCompiler.cpp
--- a/Engine/Source/Developer/RigVMDeveloper/Private/RigVMCompiler/RigVMCompiler.cpp
+++ b/Engine/Source/Developer/RigVMDeveloper/Private/RigVMCompiler/RigVMCompiler.cpp
@@ -1713,17 +1713,17 @@ FRigVMOperand URigVMCompiler::FindOrAddRegister(const FRigVMVarExprAST* InVarExp
 
 	// Get all possible pins that lead to the same operand
 
-	if(Settings.ASTSettings.bFoldAssignments)
-	{
-		TArray<const URigVMPin*> VirtualPins;
-		Pin->GetExposedPinChain(VirtualPins);
-		for (const URigVMPin* VirtualPin : VirtualPins)
-		{
-			FString VirtualPinHash = GetPinHash(VirtualPin, InVarExpr, bIsDebugValue);
-			WorkData.PinPathToOperand->Add(VirtualPinHash, Operand);	
-		}
-	}
-	else
+	// if(Settings.ASTSettings.bFoldAssignments)
+	// {
+	// 	TArray<const URigVMPin*> VirtualPins;
+	// 	Pin->GetExposedPinChain(VirtualPins);
+	// 	for (const URigVMPin* VirtualPin : VirtualPins)
+	// 	{
+	// 		FString VirtualPinHash = GetPinHash(VirtualPin, InVarExpr, bIsDebugValue);
+	// 		WorkData.PinPathToOperand->Add(VirtualPinHash, Operand);	
+	// 	}
+	// }
+	// else
 	{
 		WorkData.PinPathToOperand->Add(Hash, Operand);
 	}
