ControlRig: Fix watches not working for functions and collapsed nodes.

diff --git a/Engine/Source/Developer/RigVMDeveloper/Private/RigVMCompiler/RigVMCompiler.cpp b/Engine/Source/Developer/RigVMDeveloper/Private/RigVMCompiler/RigVMCompiler.cpp
--- a/Engine/Source/Developer/RigVMDeveloper/Private/RigVMCompiler/RigVMCompiler.cpp
+++ b/Engine/Source/Developer/RigVMDeveloper/Private/RigVMCompiler/RigVMCompiler.cpp
@@ -1713,17 +1713,17 @@ FRigVMOperand URigVMCompiler::FindOrAddRegister(const FRigVMVarExprAST* InVarExp
 
 	// Get all possible pins that lead to the same operand
 
-	// if(Settings.ASTSettings.bFoldAssignments)
-	// {
-	// 	TArray<const URigVMPin*> VirtualPins;
-	// 	Pin->GetExposedPinChain(VirtualPins);
-	// 	for (const URigVMPin* VirtualPin : VirtualPins)
-	// 	{
-	// 		FString VirtualPinHash = GetPinHash(VirtualPin, InVarExpr, bIsDebugValue);
-	// 		WorkData.PinPathToOperand->Add(VirtualPinHash, Operand);	
-	// 	}
-	// }
-	// else
+	if(Settings.ASTSettings.bFoldAssignments)
+	{
+		TArray<const URigVMPin*> VirtualPins;
+		Pin->GetExposedPinChain(VirtualPins);
+		for (const URigVMPin* VirtualPin : VirtualPins)
+		{
+			FString VirtualPinHash = GetPinHash(VirtualPin, InVarExpr, bIsDebugValue);
+			WorkData.PinPathToOperand->Add(VirtualPinHash, Operand);	
+		}
+	}
+	else
 	{
 		WorkData.PinPathToOperand->Add(Hash, Operand);
 	}
diff --git a/Engine/Source/Developer/RigVMDeveloper/Private/RigVMModel/RigVMPin.cpp b/Engine/Source/Developer/RigVMDeveloper/Private/RigVMModel/RigVMPin.cpp
--- a/Engine/Source/Developer/RigVMDeveloper/Private/RigVMModel/RigVMPin.cpp
+++ b/Engine/Source/Developer/RigVMDeveloper/Private/RigVMModel/RigVMPin.cpp
@@ -223,7 +223,7 @@ FString URigVMPin::GetSegmentPath() const
 void URigVMPin::GetExposedPinChain(TArray<const URigVMPin*>& OutExposedPins) const
 {
 	// Variable nodes do not share the operand with their source link
-	if (GetNode()->IsA<URigVMVariableNode>() && GetDirection() == ERigVMPinDirection::Input)
+	if ((GetNode()->IsA<URigVMVariableNode>() || GetNode()->IsA<URigVMParameterNode>()) && GetDirection() == ERigVMPinDirection::Input)
 	{
 		OutExposedPins.Add(this);
 		return;
@@ -254,7 +254,7 @@ void URigVMPin::GetExposedPinChain(TArray<const URigVMPin*>& OutExposedPins) con
 			}
 		}
 		// Variable nodes do not share the operand with their source link
-		else if (SourcePin->GetNode()->IsA<URigVMVariableNode>())
+		else if (SourcePin->GetNode()->IsA<URigVMVariableNode>() || SourcePin->GetNode()->IsA<URigVMParameterNode>())
 		{
 			continue;
 		}
@@ -286,7 +286,7 @@ void URigVMPin::GetExposedPinChain(TArray<const URigVMPin*>& OutExposedPins) con
 			URigVMPin* TargetPin = Link->GetTargetPin();
 
 			// Variable nodes do not share the operand with their source link
-			if (TargetPin->GetNode()->IsA<URigVMVariableNode>())
+			if (TargetPin->GetNode()->IsA<URigVMVariableNode>() || TargetPin->GetNode()->IsA<URigVMParameterNode>())
 			{
 				continue;
 			}
