Control Rig: Watch variables overriding expresion to operand map

diff --git a/Engine/Source/Developer/RigVMDeveloper/Private/RigVMCompiler/RigVMCompiler.cpp b/Engine/Source/Developer/RigVMDeveloper/Private/RigVMCompiler/RigVMCompiler.cpp
--- a/Engine/Source/Developer/RigVMDeveloper/Private/RigVMCompiler/RigVMCompiler.cpp
+++ b/Engine/Source/Developer/RigVMDeveloper/Private/RigVMCompiler/RigVMCompiler.cpp
@@ -1333,6 +1333,7 @@ FRigVMOperand URigVMCompiler::FindOrAddRegister(const FRigVMVarExprAST* InVarExp
 	{
 		if (!bIsDebugValue)
 		{
+			check(!WorkData.ExprToOperand.Contains(InVarExpr));
 			WorkData.ExprToOperand.Add(InVarExpr, *ExistingOperand);
 		}
 		return *ExistingOperand;
@@ -1727,7 +1728,12 @@ FRigVMOperand URigVMCompiler::FindOrAddRegister(const FRigVMVarExprAST* InVarExp
 	{
 		WorkData.PinPathToOperand->Add(Hash, Operand);
 	}
-	WorkData.ExprToOperand.Add(InVarExpr, Operand);
+	
+	if (!bIsDebugValue)
+	{
+		check(!WorkData.ExprToOperand.Contains(InVarExpr));
+		WorkData.ExprToOperand.Add(InVarExpr, Operand);
+	}
 
 	return Operand;
 }
