Control Rig: Orphan pin functionality for RigVM

diff --git a/Engine/Plugins/Experimental/ControlRig/Source/ControlRig/Private/Units/Math/RigUnit_MathFloat.h b/Engine/Plugins/Experimental/ControlRig/Source/ControlRig/Private/Units/Math/RigUnit_MathFloat.h
--- a/Engine/Plugins/Experimental/ControlRig/Source/ControlRig/Private/Units/Math/RigUnit_MathFloat.h
+++ b/Engine/Plugins/Experimental/ControlRig/Source/ControlRig/Private/Units/Math/RigUnit_MathFloat.h
@@ -2,7 +2,7 @@
 
 #pragma once
 
-#include "Units/Math/RigUnit_MathBase.h"
+#include "RigUnit_MathBase.h"
 #include "RigUnit_MathFloat.generated.h"
 
 USTRUCT(meta=(Abstract, Category="Math|Float", MenuDescSuffix="(Float)"))
diff --git a/Engine/Plugins/Experimental/ControlRig/Source/ControlRigDeveloper/Private/ControlRigBlueprint.cpp b/Engine/Plugins/Experimental/ControlRig/Source/ControlRigDeveloper/Private/ControlRigBlueprint.cpp
--- a/Engine/Plugins/Experimental/ControlRig/Source/ControlRigDeveloper/Private/ControlRigBlueprint.cpp
+++ b/Engine/Plugins/Experimental/ControlRig/Source/ControlRigDeveloper/Private/ControlRigBlueprint.cpp
@@ -69,6 +69,7 @@ UControlRigBlueprint::UControlRigBlueprint(const FObjectInitializer& ObjectIniti
 	bRecompileOnLoad = 0;
 	bAutoRecompileVM = true;
 	bVMRecompilationRequired = false;
+	bIsCompiling = false;
 	VMRecompilationBracket = 0;
 
 	Model = ObjectInitializer.CreateDefaultSubobject<URigVMGraph>(this, TEXT("RigVMModel"));
@@ -93,6 +94,12 @@ UControlRigBlueprint::UControlRigBlueprint(const FObjectInitializer& ObjectIniti
 
 	VMCompileSettings.ASTSettings.ReportDelegate.BindUObject(this, &UControlRigBlueprint::HandleReportFromCompiler);
 
+#if WITH_EDITOR
+	CompileLog.SetSourcePath(GetPathName());
+	CompileLog.bLogDetailedResults = false;
+	CompileLog.EventDisplayThresholdMs = false;
+#endif
+
 	Hierarchy = CreateDefaultSubobject<URigHierarchy>(TEXT("Hierarchy"));
 	HierarchyController = CreateDefaultSubobject<URigHierarchyController>(TEXT("HierarchyController"));
 	HierarchyController->SetHierarchy(Hierarchy);
@@ -496,11 +503,12 @@ void UControlRigBlueprint::PostLoad()
 	{
 		for (URigVMGraph* GraphToDetach : GraphsToDetach)
 		{
-			GetOrCreateController(GraphToDetach)->DetachLinksFromPinObjects();
+			URigVMController* Controller = GetOrCreateController(GraphToDetach);
+			Controller->DetachLinksFromPinObjects();
 			TArray<URigVMNode*> Nodes = GraphToDetach->GetNodes();
 			for (URigVMNode* Node : Nodes)
 			{
-				GetOrCreateController(GraphToDetach)->RepopulatePinsOnNode(Node);
+				Controller->RepopulatePinsOnNode(Node, true, false, true);
 			}
 		}
 		SetupPinRedirectorsForBackwardsCompatibility();
@@ -508,7 +516,13 @@ void UControlRigBlueprint::PostLoad()
 
 	for (URigVMGraph* GraphToDetach : GraphsToDetach)
 	{
-		GetOrCreateController(GraphToDetach)->ReattachLinksToPinObjects(true /* follow redirectors */);
+		URigVMController* Controller = GetOrCreateController(GraphToDetach);
+		Controller->ReattachLinksToPinObjects(true /* follow redirectors */, nullptr, false, true);
+		
+		for(URigVMNode* Node : GraphToDetach->GetNodes())
+		{
+			Controller->RemoveUnusedOrphanedPins(Node, false);
+		}
 	}
 
 	// perform backwards compat value upgrades
@@ -549,6 +563,9 @@ void UControlRigBlueprint::PostLoad()
 		}
 	}	
 
+	CompileLog.Messages.Reset();
+	CompileLog.NumErrors = CompileLog.NumWarnings = 0;
+	
 	RecompileVM();
 	RequestControlRigInit();
 
@@ -567,6 +584,12 @@ void UControlRigBlueprint::PostLoad()
 
 void UControlRigBlueprint::RecompileVM()
 {
+	if(bIsCompiling)
+	{
+		return;
+	}
+	TGuardValue<bool> CompilingGuard(bIsCompiling, true);
+	
 	bErrorsDuringCompilation = false;
 
 	UControlRigBlueprintGeneratedClass* RigClass = GetControlRigBlueprintGeneratedClass();
@@ -605,12 +628,19 @@ void UControlRigBlueprint::RecompileVM()
 		UserData.Add(FRigVMUserDataArray(&InitContextPtr, 1));
 		UserData.Add(FRigVMUserDataArray(&UpdateContextPtr, 1));
 
+		CompileLog.Messages.Reset();
+		CompileLog.NumErrors = CompileLog.NumWarnings = 0;
+
 		URigVMCompiler* Compiler = URigVMCompiler::StaticClass()->GetDefaultObject<URigVMCompiler>();
 		Compiler->Settings = VMCompileSettings;
-		Compiler->Compile(Model, GetController(), CDO->VM, CDO->GetExternalVariablesImpl(false), UserData, &PinToOperandMap);
+		Compiler->Compile(Model, GetOrCreateController(), CDO->VM, CDO->GetExternalVariablesImpl(false), UserData, &PinToOperandMap);
 
 		if (bErrorsDuringCompilation)
 		{
+			if(CDO->VM)
+			{
+				VMCompiledEvent.Broadcast(this, CDO->VM);
+			}
 			return;
 		}
 
@@ -673,48 +703,73 @@ void UControlRigBlueprint::DecrementVMRecompileBracket()
 
 void UControlRigBlueprint::HandleReportFromCompiler(EMessageSeverity::Type InSeverity, UObject* InSubject, const FString& InMessage)
 {
+	UObject* SubjectForMessage = InSubject;
+	if(URigVMNode* ModelNode = Cast<URigVMNode>(SubjectForMessage))
+	{
+		if(UControlRigBlueprint* RigBlueprint = ModelNode->GetTypedOuter<UControlRigBlueprint>())
+		{
+			if(UControlRigGraph* EdGraph = Cast<UControlRigGraph>(RigBlueprint->GetEdGraph(ModelNode->GetGraph())))
+			{
+				if(UEdGraphNode* EdNode = EdGraph->FindNodeForModelNodeName(ModelNode->GetFName()))
+				{
+					SubjectForMessage = EdNode;
+				}
+			}
+		}
+	}
+
+	FCompilerResultsLog* Log = CurrentMessageLog ? CurrentMessageLog : &CompileLog;
 	if (InSeverity == EMessageSeverity::Error)
 	{
 		Status = BS_Error;
 		MarkPackageDirty();
 
-#if WITH_EDITOR
-		FNotificationInfo Info(FText::FromString(InMessage));
-		Info.Image = FCoreStyle::Get().GetBrush(TEXT("MessageLog.Error"));
-		Info.bFireAndForget = true;
-		Info.FadeOutDuration = 5.0f;
-		Info.ExpireDuration = 8.0f;
-		TSharedPtr<SNotificationItem> NotificationPtr = FSlateNotificationManager::Get().AddNotification(Info);
-		NotificationPtr->SetCompletionState(SNotificationItem::CS_Success);
-#endif
+		if(InMessage.Contains(TEXT("@@")))
+		{
+			Log->Error(*InMessage, SubjectForMessage);
+		}
+		else
+		{
+			Log->Error(*InMessage);
+		}
 
+		BroadCastReportCompilerMessage(InSeverity, InSubject, InMessage);
 		FScriptExceptionHandler::Get().HandleException(ELogVerbosity::Error, *InMessage, *FString());
 		bErrorsDuringCompilation = true;
 	}
 	else if (InSeverity == EMessageSeverity::Warning)
 	{
+		if(InMessage.Contains(TEXT("@@")))
+		{
+			Log->Warning(*InMessage, SubjectForMessage);
+		}
+		else
+		{
+			Log->Warning(*InMessage);
+		}
+
+		BroadCastReportCompilerMessage(InSeverity, InSubject, InMessage);
 		FScriptExceptionHandler::Get().HandleException(ELogVerbosity::Warning, *InMessage, *FString());
 	}
 	else
 	{
+		if(InMessage.Contains(TEXT("@@")))
+		{
+			Log->Note(*InMessage, SubjectForMessage);
+		}
+		else
+		{
+			Log->Note(*InMessage);
+		}
+
 		UE_LOG(LogControlRigDeveloper, Display, TEXT("%s"), *InMessage);
 	}
 
-	if (URigVMNode* Node = Cast<URigVMNode>(InSubject))
+	if (UControlRigGraphNode* EdGraphNode = Cast<UControlRigGraphNode>(SubjectForMessage))
 	{
-		if (URigVMGraph* Graph = Node->GetGraph())
-		{
-			if (UControlRigGraph* EdGraph = Cast<UControlRigGraph>(GetEdGraph(Graph)))
-			{
-				if (UControlRigGraphNode* EdGraphNode = Cast<UControlRigGraphNode>(EdGraph->FindNodeForModelNodeName(Node->GetFName())))
-				{
-					EdGraphNode->ErrorType = (int32)InSeverity;
-					EdGraphNode->ErrorMsg = InMessage;
-					EdGraphNode->bHasCompilerMessage = EdGraphNode->ErrorType <= int32(EMessageSeverity::Info);
-
-				}
-			}
-		}
+		EdGraphNode->ErrorType = (int32)InSeverity;
+		EdGraphNode->ErrorMsg = InMessage;
+		EdGraphNode->bHasCompilerMessage = EdGraphNode->ErrorType <= int32(EMessageSeverity::Info);
 	}
 }
 
@@ -2824,6 +2879,11 @@ void UControlRigBlueprint::BroadcastRequestLocalizeFunctionDialog(URigVMLibraryN
 	RequestLocalizeFunctionDialog.Broadcast(InFunction, this, bForce);
 }
 
+void UControlRigBlueprint::BroadCastReportCompilerMessage(EMessageSeverity::Type InSeverity, UObject* InSubject, const FString& InMessage)
+{
+	ReportCompilerMessageEvent.Broadcast(InSeverity, InSubject, InMessage);
+}
+
 #endif
 
 void UControlRigBlueprint::CreateEdGraphForCollapseNodeIfNeeded(URigVMCollapseNode* InNode, bool bForce)
diff --git a/Engine/Plugins/Experimental/ControlRig/Source/ControlRigDeveloper/Private/Graph/ControlRigGraphNode.cpp b/Engine/Plugins/Experimental/ControlRig/Source/ControlRigDeveloper/Private/Graph/ControlRigGraphNode.cpp
--- a/Engine/Plugins/Experimental/ControlRig/Source/ControlRigDeveloper/Private/Graph/ControlRigGraphNode.cpp
+++ b/Engine/Plugins/Experimental/ControlRig/Source/ControlRigDeveloper/Private/Graph/ControlRigGraphNode.cpp
@@ -139,8 +139,6 @@ void UControlRigGraphNode::ReconstructNode_Internal(bool bForce)
 	// Clear previously set messages
 	ErrorMsg.Reset();
 
-	// @TODO: support pin orphaning/conversions for upgrades/deprecations?
-
 	// Move the existing pins to a saved array
 	TArray<UEdGraphPin*> OldPins(Pins);
 	Pins.Reset();
@@ -201,7 +199,6 @@ void UControlRigGraphNode::ReallocatePinsDuringReconstruction(const TArray<UEdGr
 void UControlRigGraphNode::RewireOldPinsToNewPins(TArray<UEdGraphPin*>& InOldPins, TArray<UEdGraphPin*>& InNewPins)
 {
 	DECLARE_SCOPE_HIERARCHICAL_COUNTER_FUNC()
-	// @TODO: we should account for redirectors, orphaning etc. here too!
 
 	for(UEdGraphPin* OldPin : InOldPins)
 	{
@@ -334,32 +331,36 @@ void UControlRigGraphNode::AllocateDefaultPins()
 
 	if (URigVMNode* ModelNode = GetModelNode())
 	{
-		for (URigVMPin* ModelPin : ModelNode->GetPins())
+		for(int32 PinListIndex=0; PinListIndex<2; PinListIndex++)
 		{
-			if (ModelPin->ShowInDetailsPanelOnly())
-			{
-				continue;
-			}
-			if (ModelPin->GetDirection() == ERigVMPinDirection::IO)
+			const TArray<URigVMPin*>& ModelPins = PinListIndex == 0 ? ModelNode->GetPins() : ModelNode->GetOrphanedPins();
+			for (URigVMPin* ModelPin : ModelPins)
 			{
-				if (ModelPin->IsStruct())
+				if (ModelPin->ShowInDetailsPanelOnly())
 				{
-					if (ModelPin->GetScriptStruct()->IsChildOf(FRigVMExecuteContext::StaticStruct()))
+					continue;
+				}
+				if (ModelPin->GetDirection() == ERigVMPinDirection::IO)
+				{
+					if (ModelPin->IsStruct())
 					{
-						ExecutePins.Add(ModelPin);
-						continue;
+						if (ModelPin->GetScriptStruct()->IsChildOf(FRigVMExecuteContext::StaticStruct()))
+						{
+							ExecutePins.Add(ModelPin);
+							continue;
+						}
 					}
+					InputOutputPins.Add(ModelPin);
+				}
+				else if (ModelPin->GetDirection() == ERigVMPinDirection::Input || 
+                    ModelPin->GetDirection() == ERigVMPinDirection::Visible)
+				{
+					InputPins.Add(ModelPin);
+				}
+				else if (ModelPin->GetDirection() == ERigVMPinDirection::Output)
+				{
+					OutputPins.Add(ModelPin);
 				}
-				InputOutputPins.Add(ModelPin);
-			}
-			else if (ModelPin->GetDirection() == ERigVMPinDirection::Input || 
-				ModelPin->GetDirection() == ERigVMPinDirection::Visible)
-			{
-				InputPins.Add(ModelPin);
-			}
-			else if (ModelPin->GetDirection() == ERigVMPinDirection::Output)
-			{
-				OutputPins.Add(ModelPin);
 			}
 		}
 	}
@@ -413,6 +414,7 @@ void UControlRigGraphNode::CreateInputPins(URigVMPin* InParentPin)
 			{
 				Pair.InputPin->PinFriendlyName = FText::FromName(ModelPin->GetDisplayName());
 				Pair.InputPin->bNotConnectable = ModelPin->GetDirection() != ERigVMPinDirection::Input;
+				Pair.InputPin->bOrphanedPin = ModelPin->IsOrphanPin() ? 1 : 0; 
 
 				SetupPinDefaultsFromModel(Pair.InputPin);
 
@@ -444,6 +446,7 @@ void UControlRigGraphNode::CreateInputOutputPins(URigVMPin* InParentPin, bool bH
 				Pair.InputPin->bHidden = bHidden;
 				Pair.InputPin->PinFriendlyName = FText::FromName(ModelPin->GetDisplayName());
 				Pair.InputPin->bNotConnectable = ModelPin->GetDirection() != ERigVMPinDirection::IO;
+				Pair.InputPin->bOrphanedPin = ModelPin->IsOrphanPin() ? 1 : 0; 
 
 				SetupPinDefaultsFromModel(Pair.InputPin);
 
@@ -463,6 +466,7 @@ void UControlRigGraphNode::CreateInputOutputPins(URigVMPin* InParentPin, bool bH
 				Pair.OutputPin->bHidden = bHidden;
 				Pair.OutputPin->PinFriendlyName = FText::FromName(ModelPin->GetDisplayName());
 				Pair.OutputPin->bNotConnectable = ModelPin->GetDirection() != ERigVMPinDirection::IO;
+				Pair.OutputPin->bOrphanedPin = ModelPin->IsOrphanPin() ? 1 : 0; 
 
 				if (InParentPin != nullptr)
 		{
@@ -501,6 +505,7 @@ void UControlRigGraphNode::CreateOutputPins(URigVMPin* InParentPin)
 			{
 				Pair.OutputPin->PinFriendlyName = FText::FromName(ModelPin->GetDisplayName());
 				Pair.OutputPin->bNotConnectable = ModelPin->GetDirection() != ERigVMPinDirection::Output;
+				Pair.OutputPin->bOrphanedPin = ModelPin->IsOrphanPin() ? 1 : 0; 
 
 				if (InParentPin != nullptr)
 	{
diff --git a/Engine/Plugins/Experimental/ControlRig/Source/ControlRigDeveloper/Public/ControlRigBlueprint.h b/Engine/Plugins/Experimental/ControlRig/Source/ControlRigDeveloper/Public/ControlRigBlueprint.h
--- a/Engine/Plugins/Experimental/ControlRig/Source/ControlRigDeveloper/Public/ControlRigBlueprint.h
+++ b/Engine/Plugins/Experimental/ControlRig/Source/ControlRigDeveloper/Public/ControlRigBlueprint.h
@@ -22,6 +22,7 @@
 
 #if WITH_EDITOR
 #include "Kismet2/Kismet2NameValidators.h"
+#include "Kismet2/CompilerResultsLog.h"
 #endif
 
 #include "ControlRigBlueprint.generated.h"
@@ -38,6 +39,7 @@ DECLARE_EVENT_TwoParams(UControlRigBlueprint, FOnNodeDoubleClicked, UControlRigB
 DECLARE_EVENT_OneParam(UControlRigBlueprint, FOnGraphImported, UEdGraph*);
 DECLARE_EVENT_OneParam(UControlRigBlueprint, FOnPostEditChangeChainProperty, FPropertyChangedChainEvent&);
 DECLARE_EVENT_ThreeParams(UControlRigBlueprint, FOnLocalizeFunctionDialogRequested, URigVMLibraryNode*, UControlRigBlueprint*, bool);
+DECLARE_EVENT_ThreeParams(UControlRigBlueprint, FOnReportCompilerMessage, EMessageSeverity::Type, UObject*, const FString&);
 
 USTRUCT()
 struct CONTROLRIGDEVELOPER_API FControlRigPublicFunctionArg
@@ -393,6 +395,9 @@ private:
 	UPROPERTY(transient)
 	bool bVMRecompilationRequired;
 
+	UPROPERTY(transient)
+	bool bIsCompiling;
+
 	UPROPERTY(transient)
 	int32 VMRecompilationBracket;
 
@@ -491,15 +496,20 @@ public:
 	FOnLocalizeFunctionDialogRequested& OnRequestLocalizeFunctionDialog() { return RequestLocalizeFunctionDialog; }
 	void BroadcastRequestLocalizeFunctionDialog(URigVMLibraryNode* InFunction, bool bForce = false);
 
+	FOnReportCompilerMessage& OnReportCompilerMessage() { return ReportCompilerMessageEvent; }
+	void BroadCastReportCompilerMessage(EMessageSeverity::Type InSeverity, UObject* InSubject, const FString& InMessage);
+
 private:
 
 	FOnExternalVariablesChanged ExternalVariablesChangedEvent;
 	void BroadcastExternalVariablesChangedEvent();
+	FCompilerResultsLog CompileLog;
 
 	FOnNodeDoubleClicked NodeDoubleClickedEvent;
 	FOnGraphImported GraphImportedEvent;
 	FOnPostEditChangeChainProperty PostEditChangeChainPropertyEvent;
 	FOnLocalizeFunctionDialogRequested RequestLocalizeFunctionDialog;
+	FOnReportCompilerMessage ReportCompilerMessageEvent;
 
 #endif
 
diff --git a/Engine/Plugins/Experimental/ControlRig/Source/ControlRigEditor/Private/Editor/ControlRigEditor.cpp b/Engine/Plugins/Experimental/ControlRig/Source/ControlRigEditor/Private/Editor/ControlRigEditor.cpp
--- a/Engine/Plugins/Experimental/ControlRig/Source/ControlRigEditor/Private/Editor/ControlRigEditor.cpp
+++ b/Engine/Plugins/Experimental/ControlRig/Source/ControlRigEditor/Private/Editor/ControlRigEditor.cpp
@@ -78,6 +78,7 @@
 #include "RigVMModel/Nodes/RigVMFunctionEntryNode.h"
 #include "RigVMModel/Nodes/RigVMFunctionReturnNode.h"
 #include "BlueprintActionDatabaseRegistrar.h"
+#include "IMessageLogListing.h"
 #include "SControlRigFunctionLocalizationWidget.h"
 
 #define LOCTEXT_NAMESPACE "ControlRigEditor"
@@ -128,6 +129,7 @@ FControlRigEditor::~FControlRigEditor()
 		RigBlueprint->OnGraphImported().RemoveAll(this);
 		RigBlueprint->OnPostEditChangeChainProperty().RemoveAll(this);
 		RigBlueprint->OnRequestLocalizeFunctionDialog().RemoveAll(this);
+		RigBlueprint->OnReportCompilerMessage().RemoveAll(this);
 	}
 
 	if (NodeDetailBuffer.Num() > 0 && NodeDetailStruct != nullptr)
@@ -375,6 +377,14 @@ void FControlRigEditor::InitControlRigEditor(const EToolkitMode::Type Mode, cons
 	FControlRigBlueprintUtils::HandleRefreshAllNodes(InControlRigBlueprint);
 
 	bControlRigEditorInitialized = true;
+
+	if (ControlRigBlueprints.Num() > 0)
+	{
+		if(ControlRigBlueprints[0]->Status == BS_Error)
+		{
+			Compile();
+		}
+	}
 }
 
 void FControlRigEditor::BindCommands()
@@ -1287,6 +1297,8 @@ void FControlRigEditor::Compile()
 			return;
 		}
 
+		RigBlueprint->CompileLog.Messages.Reset();
+
 		FString LastDebuggedObjectName = GetCustomDebugObjectLabel(RigBlueprint->GetObjectBeingDebugged());
 		RigBlueprint->SetObjectBeingDebugged(nullptr);
 
@@ -1966,6 +1978,13 @@ void FControlRigEditor::HandleModifiedEvent(ERigVMGraphNotifType InNotifType, UR
 
 void FControlRigEditor::HandleVMCompiledEvent(UBlueprint* InBlueprint, URigVM* InVM)
 {
+	if(UControlRigBlueprint* RigBlueprint = Cast<UControlRigBlueprint>(InBlueprint))
+	{
+		CompilerResultsListing->ClearMessages();
+		CompilerResultsListing->AddMessages(RigBlueprint->CompileLog.Messages);
+		RigBlueprint->CompileLog.Messages.Reset();
+		RigBlueprint->CompileLog.NumErrors = RigBlueprint->CompileLog.NumWarnings = 0;  
+	}
 }
 
 void FControlRigEditor::HandleControlRigExecutedEvent(UControlRig* InControlRig, const EControlRigState InState, const FName& InEventName)
@@ -2097,7 +2116,6 @@ void FControlRigEditor::CreateEditorModeManager()
 	PRAGMA_ENABLE_DEPRECATION_WARNINGS
 }
 
-
 void FControlRigEditor::Tick(float DeltaTime)
 {
 	FBlueprintEditor::Tick(DeltaTime);
diff --git a/Engine/Plugins/Experimental/ControlRig/Source/ControlRigEditor/Private/Graph/SControlRigGraphNode.cpp b/Engine/Plugins/Experimental/ControlRig/Source/ControlRigEditor/Private/Graph/SControlRigGraphNode.cpp
--- a/Engine/Plugins/Experimental/ControlRig/Source/ControlRigEditor/Private/Graph/SControlRigGraphNode.cpp
+++ b/Engine/Plugins/Experimental/ControlRig/Source/ControlRigEditor/Private/Graph/SControlRigGraphNode.cpp
@@ -865,6 +865,11 @@ FSlateColor SControlRigGraphNode::GetPinTextColor(TWeakPtr<SGraphPin> GraphPin)
 {
 	if(GraphPin.IsValid())
 	{
+		if (GraphPin.Pin()->GetPinObj()->bOrphanedPin)
+		{
+			return FLinearColor::Red;
+		}
+
 		// If there is no schema there is no owning node (or basically this is a deleted node)
 		if (GraphNode)
 		{
diff --git a/Engine/Source/Developer/RigVMDeveloper/Private/RigVMCompiler/RigVMAST.cpp b/Engine/Source/Developer/RigVMDeveloper/Private/RigVMCompiler/RigVMAST.cpp
--- a/Engine/Source/Developer/RigVMDeveloper/Private/RigVMCompiler/RigVMAST.cpp
+++ b/Engine/Source/Developer/RigVMDeveloper/Private/RigVMCompiler/RigVMAST.cpp
@@ -978,6 +978,12 @@ FRigVMExprAST* FRigVMParserAST::TraverseMutableNode(const FRigVMASTProxy& InNode
 		return SubjectToExpression.FindChecked(InNodeProxy);
 	}
 
+	URigVMNode* Node = InNodeProxy.GetSubjectChecked<URigVMNode>();
+	if(Node->HasOrphanedPins())
+	{
+		return nullptr;
+	}
+
 	FRigVMExprAST* NodeExpr = CreateExpressionForNode(InNodeProxy, InParentExpr);
 	if (NodeExpr)
 	{
@@ -988,7 +994,6 @@ FRigVMExprAST* FRigVMParserAST::TraverseMutableNode(const FRigVMASTProxy& InNode
 
 		TraversePins(InNodeProxy, NodeExpr);
 
-		URigVMNode* Node = InNodeProxy.GetSubjectChecked<URigVMNode>();
 		for (URigVMPin* SourcePin : Node->GetPins())
 		{
 			if (SourcePin->GetDirection() == ERigVMPinDirection::Output || SourcePin->GetDirection() == ERigVMPinDirection::IO)
@@ -1043,6 +1048,12 @@ FRigVMExprAST* FRigVMParserAST::TraverseNode(const FRigVMASTProxy& InNodeProxy,
 		return nullptr;
 	
 	}
+
+	if(Node->HasOrphanedPins())
+	{
+		return nullptr;
+	}
+	
 	if (SubjectToExpression.Contains(InNodeProxy))
 	{
 		FRigVMExprAST* NodeExpr = SubjectToExpression.FindChecked(InNodeProxy);
diff --git a/Engine/Source/Developer/RigVMDeveloper/Private/RigVMCompiler/RigVMCompiler.cpp b/Engine/Source/Developer/RigVMDeveloper/Private/RigVMCompiler/RigVMCompiler.cpp
--- a/Engine/Source/Developer/RigVMDeveloper/Private/RigVMCompiler/RigVMCompiler.cpp
+++ b/Engine/Source/Developer/RigVMDeveloper/Private/RigVMCompiler/RigVMCompiler.cpp
@@ -124,6 +124,34 @@ bool URigVMCompiler::Compile(URigVMGraph* InGraph, URigVMController* InControlle
 	}
 	OutOperands->Reset();
 
+#if WITH_EDITOR
+
+	// traverse all graphs and try to clear out orphan pins
+	TArray<URigVMGraph*> VisitedGraphs;
+	VisitedGraphs.Add(InGraph);
+	for(int32 GraphIndex=0; GraphIndex<VisitedGraphs.Num(); GraphIndex++)
+	{
+		URigVMGraph* VisitedGraph = VisitedGraphs[GraphIndex];
+		for(URigVMNode* ModelNode : VisitedGraph->GetNodes())
+		{
+			FRigVMControllerGraphGuard Guard(InController, VisitedGraph);
+			if(!InController->RemoveUnusedOrphanedPins(ModelNode, true))
+			{
+				static const FString LinkedMessage = TEXT("Node @@ uses pins that no longer exist. Please rewire the links and re-compile.");
+				Settings.ASTSettings.Report(EMessageSeverity::Error, ModelNode, LinkedMessage);
+			}
+
+			if(URigVMLibraryNode* LibraryNode = Cast<URigVMLibraryNode>(ModelNode))
+			{
+				if(URigVMGraph* ContainedGraph = LibraryNode->GetContainedGraph())
+				{
+					VisitedGraphs.AddUnique(ContainedGraph);
+				}
+			}
+		}
+	}
+#endif
+	
 	OutVM->ClearExternalVariables();
 	for (const FRigVMExternalVariable& ExternalVariable : InExternalVariables)
 	{
diff --git a/Engine/Source/Developer/RigVMDeveloper/Private/RigVMModel/RigVMController.cpp b/Engine/Source/Developer/RigVMDeveloper/Private/RigVMModel/RigVMController.cpp
--- a/Engine/Source/Developer/RigVMDeveloper/Private/RigVMModel/RigVMController.cpp
+++ b/Engine/Source/Developer/RigVMDeveloper/Private/RigVMModel/RigVMController.cpp
@@ -493,7 +493,7 @@ URigVMVariableNode* URigVMController::AddVariableNodeFromObjectPath(const FName&
 	return AddVariableNode(InVariableName, InCPPType, CPPTypeObject, bIsGetter, InDefaultValue, InPosition, InNodeName, bSetupUndoRedo);
 }
 
-void URigVMController::RefreshVariableNode(const FName& InNodeName, const FName& InVariableName, const FString& InCPPType, UObject* InCPPTypeObject, bool bSetupUndoRedo)
+void URigVMController::RefreshVariableNode(const FName& InNodeName, const FName& InVariableName, const FString& InCPPType, UObject* InCPPTypeObject, bool bSetupUndoRedo, bool bSetupOrphanPinss)
 {
 	if (!IsValidGraph())
 	{
@@ -539,9 +539,6 @@ void URigVMController::RefreshVariableNode(const FName& InNodeName, const FName&
 							ValuePin->Modify();
 						}
 
-						BreakAllLinks(ValuePin, ValuePin->GetDirection() == ERigVMPinDirection::Input, bSetupUndoRedo);
-						BreakAllLinksRecursive(ValuePin, ValuePin->GetDirection() == ERigVMPinDirection::Input, false, bSetupUndoRedo);
-
 						// if this is an unsupported datatype...
 						if (InCPPType == FName(NAME_None).ToString())
 						{
@@ -549,24 +546,12 @@ void URigVMController::RefreshVariableNode(const FName& InNodeName, const FName&
 							return;
 						}
 
-						ValuePin->CPPType = InCPPType;
-						ValuePin->CPPTypeObject = InCPPTypeObject;
-						ValuePin->CPPTypeObjectPath = *InCPPTypeObject->GetPathName();
-
-						TArray<URigVMPin*> SubPins = ValuePin->GetSubPins();
-						for(URigVMPin * SubPin : SubPins)
-						{
-							ValuePin->SubPins.Remove(SubPin);
-						}
-
-						if (ValuePin->IsStruct())
+						FString CPPTypeObjectPath;
+						if(InCPPTypeObject)
 						{
-							FString DefaultValue = ValuePin->DefaultValue;
-							CreateDefaultValueForStructIfRequired(ValuePin->GetScriptStruct(), DefaultValue);
-							AddPinsForStruct(ValuePin->GetScriptStruct(), ValuePin->GetNode(), ValuePin, ValuePin->Direction, DefaultValue, false);
+							CPPTypeObjectPath = InCPPTypeObject->GetPathName();
 						}
-
-						Notify(ERigVMGraphNotifType::PinTypeChanged, ValuePin);
+						ChangePinType(ValuePin, InCPPType, *CPPTypeObjectPath, bSetupUndoRedo, bSetupOrphanPinss);
 					}
 				}
 			}
@@ -6021,7 +6006,7 @@ bool URigVMController::RenameExposedPin(const FName& InOldPinName, const FName&
 	return true;
 }
 
-bool URigVMController::ChangeExposedPinType(const FName& InPinName, const FString& InCPPType, const FName& InCPPTypeObjectPath, bool bSetupUndoRedo)
+bool URigVMController::ChangeExposedPinType(const FName& InPinName, const FString& InCPPType, const FName& InCPPTypeObjectPath, bool bSetupUndoRedo, bool bSetupOrphanPins)
 {
 	if (!IsValidGraph())
 	{
@@ -6062,7 +6047,7 @@ bool URigVMController::ChangeExposedPinType(const FName& InPinName, const FStrin
 	if (bSetupUndoRedo)
 	{
 		FRigVMControllerGraphGuard GraphGuard(this, LibraryNode->GetGraph(), bSetupUndoRedo);
-		if (!ChangePinType(Pin, InCPPType, InCPPTypeObjectPath, bSetupUndoRedo))
+		if (!ChangePinType(Pin, InCPPType, InCPPTypeObjectPath, bSetupUndoRedo, bSetupOrphanPins))
 		{
 			if (bSetupUndoRedo)
 			{
@@ -6076,14 +6061,14 @@ bool URigVMController::ChangeExposedPinType(const FName& InPinName, const FStrin
 	{
 		if (URigVMPin* EntryPin = EntryNode->FindPin(Pin->GetName()))
 		{
-			ChangePinType(EntryPin, InCPPType, InCPPTypeObjectPath, bSetupUndoRedo);
+			ChangePinType(EntryPin, InCPPType, InCPPTypeObjectPath, bSetupUndoRedo, bSetupOrphanPins);
 		}
 	}
 	if (URigVMFunctionReturnNode* ReturnNode = Graph->GetReturnNode())
 	{
 		if (URigVMPin* ReturnPin = ReturnNode->FindPin(Pin->GetName()))
 		{
-			ChangePinType(ReturnPin, InCPPType, InCPPTypeObjectPath, bSetupUndoRedo);
+			ChangePinType(ReturnPin, InCPPType, InCPPTypeObjectPath, bSetupUndoRedo, bSetupOrphanPins);
 		}
 	}
 
@@ -6100,7 +6085,7 @@ bool URigVMController::ChangeExposedPinType(const FName& InPinName, const FStrin
 					if (URigVMPin* ReferencedNodePin = FunctionReferencePtr->FindPin(Pin->GetName()))
 					{
 						FRigVMControllerGraphGuard GraphGuard(this, FunctionReferencePtr->GetGraph(), false);
-						ChangePinType(ReferencedNodePin, InCPPType, InCPPTypeObjectPath, bSetupUndoRedo);
+						ChangePinType(ReferencedNodePin, InCPPType, InCPPTypeObjectPath, bSetupUndoRedo, bSetupOrphanPins);
 					}
 				}
 			}
@@ -7547,7 +7532,7 @@ int32 URigVMController::DetachLinksFromPinObjects(const TArray<URigVMLink*>* InL
 	return Links.Num();
 }
 
-int32 URigVMController::ReattachLinksToPinObjects(bool bFollowCoreRedirectors, const TArray<URigVMLink*>* InLinks, bool bNotify)
+int32 URigVMController::ReattachLinksToPinObjects(bool bFollowCoreRedirectors, const TArray<URigVMLink*>* InLinks, bool bNotify, bool bSetupOrphanedPins)
 {
 	URigVMGraph* Graph = GetGraph();
 	check(Graph);
@@ -7603,6 +7588,83 @@ int32 URigVMController::ReattachLinksToPinObjects(bool bFollowCoreRedirectors, c
 
 		URigVMPin* SourcePin = Link->GetSourcePin();
 		URigVMPin* TargetPin = Link->GetTargetPin();
+
+		if((SourcePin != nullptr) && (TargetPin != nullptr))
+		{
+			if (!URigVMPin::CanLink(SourcePin, TargetPin, nullptr, nullptr))
+			{
+				if(SourcePin->GetNode()->HasOrphanedPins() && bSetupOrphanedPins)
+				{
+					SourcePin = nullptr;
+				}
+				else if(TargetPin->GetNode()->HasOrphanedPins() && bSetupOrphanedPins)
+				{
+					TargetPin = nullptr;
+				}
+				else
+				{
+					ReportWarningf(TEXT("Unable to re-create link %s -> %s"), *Link->SourcePinPath, *Link->TargetPinPath);
+					TargetPin->Links.Remove(Link);
+					SourcePin->Links.Remove(Link);
+					continue;
+				}
+			}
+		}
+
+		if(bSetupOrphanedPins)
+		{
+			bool bRelayedBackToOrphanPin = false;
+			for(int32 PinIndex=0; PinIndex<2; PinIndex++)
+			{
+				URigVMPin*& PinToFind = PinIndex == 0 ? SourcePin : TargetPin;
+				
+				if(PinToFind == nullptr)
+				{
+					const FString& PinPathToFind = PinIndex == 0 ? Link->SourcePinPath : Link->TargetPinPath;
+					FString NodeName, RemainingPinPath;
+					URigVMPin::SplitPinPathAtStart(PinPathToFind, NodeName, RemainingPinPath);
+					check(!NodeName.IsEmpty() && !RemainingPinPath.IsEmpty());
+
+					URigVMNode* Node = Graph->FindNode(NodeName);
+					if(Node == nullptr)
+					{
+						continue;
+					}
+
+					RemainingPinPath = FString::Printf(TEXT("%s%s"), *URigVMPin::OrphanPinPrefix, *RemainingPinPath);
+					PinToFind = Node->FindPin(RemainingPinPath);
+
+					if(PinToFind != nullptr)
+					{
+						if(PinIndex == 0)
+						{
+							Link->SourcePinPath = PinToFind->GetPinPath();
+							Link->SourcePin = nullptr;
+							SourcePin = Link->GetSourcePin();
+						}
+						else
+						{
+							Link->TargetPinPath = PinToFind->GetPinPath();
+							Link->TargetPin = nullptr;
+							TargetPin = Link->GetTargetPin();
+						}
+						bRelayedBackToOrphanPin = true;
+					}
+				}
+			}
+
+			if(bRelayedBackToOrphanPin && (SourcePin != nullptr) && (TargetPin != nullptr))
+			{
+				if (!URigVMPin::CanLink(SourcePin, TargetPin, nullptr, nullptr))
+				{
+					ReportWarningf(TEXT("Unable to re-create link %s -> %s"), *Link->SourcePinPath, *Link->TargetPinPath);
+					TargetPin->Links.Remove(Link);
+					SourcePin->Links.Remove(Link);
+					continue;
+				}
+			}
+		}
+		
 		if (SourcePin == nullptr)
 		{
 			ReportWarningf(TEXT("Unable to re-create link %s -> %s"), *Link->SourcePinPath, *Link->TargetPinPath);
@@ -7857,7 +7919,7 @@ bool URigVMController::ShouldRedirectPin(const FString& InOldPinPath, FString& I
 	return false;
 }
 
-void URigVMController::RepopulatePinsOnNode(URigVMNode* InNode, bool bFollowCoreRedirectors, bool bNotify)
+void URigVMController::RepopulatePinsOnNode(URigVMNode* InNode, bool bFollowCoreRedirectors, bool bNotify, bool bSetupOrphanedPins)
 {
 	if (InNode == nullptr)
 	{
@@ -7871,6 +7933,7 @@ void URigVMController::RepopulatePinsOnNode(URigVMNode* InNode, bool bFollowCore
 	URigVMFunctionReturnNode* ReturnNode = Cast<URigVMFunctionReturnNode>(InNode);
 	URigVMCollapseNode* CollapseNode = Cast<URigVMCollapseNode>(InNode);
 	URigVMFunctionReferenceNode* FunctionRefNode = Cast<URigVMFunctionReferenceNode>(InNode);
+	URigVMVariableNode* VariableNode = Cast<URigVMVariableNode>(InNode);
 
 	TGuardValue<bool> EventuallySuspendNotifs(bSuspendNotifications, !bNotify);
 	FScopeLock Lock(&PinPathCoreRedirectorsLock);
@@ -7898,13 +7961,7 @@ void URigVMController::RepopulatePinsOnNode(URigVMNode* InNode, bool bFollowCore
 	// step 2/3: clear pins on the node and repopulate the node with new pins
 	if (UnitNode != nullptr)
 	{
-		TArray<URigVMPin*> Pins = InNode->GetPins();
-		for (URigVMPin* Pin : Pins)
-		{
-			RemovePin(Pin, false, bNotify);
-		}
-		InNode->Pins.Reset();
-		Pins.Reset();
+		RemovePinsDuringRepopulate(UnitNode, UnitNode->Pins, bNotify, bSetupOrphanedPins);
 
 		UScriptStruct* ScriptStruct = UnitNode->GetScriptStruct();
 		if (ScriptStruct == nullptr)
@@ -7930,23 +7987,25 @@ void URigVMController::RepopulatePinsOnNode(URigVMNode* InNode, bool bFollowCore
 		CreateDefaultValueForStructIfRequired(ScriptStruct, ExportedDefaultValue);
 		AddPinsForStruct(ScriptStruct, UnitNode, nullptr, ERigVMPinDirection::Invalid, ExportedDefaultValue, false, bNotify);
 	}
-	else if (RerouteNode != nullptr)
+	else if ((RerouteNode != nullptr) || (VariableNode != nullptr))
 	{
-		if (RerouteNode->GetPins().Num() == 0)
+		if (InNode->GetPins().Num() == 0)
 		{
 			return;
 		}
 
-		URigVMPin* ValuePin = RerouteNode->Pins[0];
-
-		// only repopulate the value pin, which may host a struct
-		TArray<URigVMPin*> Pins = ValuePin->SubPins;
-		for (URigVMPin* Pin : Pins)
+		URigVMPin* ValuePin = nullptr;
+		if(RerouteNode)
 		{
-			RemovePin(Pin, false, bNotify);
+			ValuePin = RerouteNode->Pins[0];
+		}
+		else
+		{
+			ValuePin = VariableNode->FindPin(URigVMVariableNode::ValueName);
 		}
-		ValuePin->SubPins.Reset();
-		Pins.Reset();
+		check(ValuePin);
+		
+		RemovePinsDuringRepopulate(InNode, ValuePin->SubPins, bNotify, bSetupOrphanedPins);
 
 		if (ValuePin->IsStruct())
 		{
@@ -7955,17 +8014,17 @@ void URigVMController::RepopulatePinsOnNode(URigVMNode* InNode, bool bFollowCore
 			{
 				ReportErrorf(
 					TEXT("Control Rig '%s', Node '%s' has no struct assigned. Do you have a broken redirect?"),
-					*RerouteNode->GetOutermost()->GetPathName(),
-					*RerouteNode->GetName()
+					*InNode->GetOutermost()->GetPathName(),
+					*InNode->GetName()
 				);
 
-				RemoveNode(RerouteNode, false, true);
+				RemoveNode(InNode, false, true);
 				return;
 			}
 
 			FString ExportedDefaultValue;
 			CreateDefaultValueForStructIfRequired(ScriptStruct, ExportedDefaultValue);
-			AddPinsForStruct(ScriptStruct, RerouteNode, ValuePin, ValuePin->Direction, ExportedDefaultValue, false);
+			AddPinsForStruct(ScriptStruct, InNode, ValuePin, ValuePin->Direction, ExportedDefaultValue, false);
 		}
 	}
 	else if (EntryNode || ReturnNode)
@@ -7973,14 +8032,7 @@ void URigVMController::RepopulatePinsOnNode(URigVMNode* InNode, bool bFollowCore
 		if (URigVMLibraryNode* LibraryNode = Cast<URigVMLibraryNode>(InNode->GetGraph()->GetOuter()))
 		{
 			bool bIsEntryNode = EntryNode != nullptr;
-
-			TArray<URigVMPin*> Pins = InNode->GetPins();
-			for (URigVMPin* Pin : Pins)
-			{
-				RemovePin(Pin, false, bNotify);
-			}
-			InNode->Pins.Reset();
-			Pins.Reset();
+			RemovePinsDuringRepopulate(InNode, InNode->Pins, bNotify, bSetupOrphanedPins);
 
 			TArray<URigVMPin*> SortedLibraryPins;
 
@@ -8054,7 +8106,7 @@ void URigVMController::RepopulatePinsOnNode(URigVMNode* InNode, bool bFollowCore
 		FRigVMControllerGraphGuard GraphGuard(this, CollapseNode->GetContainedGraph(), false);
 		for (URigVMNode* ContainedNode : CollapseNode->GetContainedNodes())
 		{
-			RepopulatePinsOnNode(ContainedNode, bFollowCoreRedirectors);
+			RepopulatePinsOnNode(ContainedNode, bFollowCoreRedirectors, bNotify, bSetupOrphanedPins);
 		}
 	}
 	else if (FunctionRefNode)
@@ -8064,13 +8116,7 @@ void URigVMController::RepopulatePinsOnNode(URigVMNode* InNode, bool bFollowCore
 			// we want to make sure notify the graph of a potential name change
 			// when repopulating the function ref node
 			Notify(ERigVMGraphNotifType::NodeRenamed, FunctionRefNode);
-
-			TArray<URigVMPin*> Pins = InNode->GetPins();
-			for (URigVMPin* Pin : Pins)
-			{
-				RemovePin(Pin, false, bNotify);
-			}
-			InNode->Pins.Reset();
+			RemovePinsDuringRepopulate(InNode, InNode->Pins, bNotify, bSetupOrphanedPins);
 
 			TMap<FString, FPinState> ReferencedPinStates = GetPinStates(ReferencedNode);
 
@@ -8107,6 +8153,111 @@ void URigVMController::RepopulatePinsOnNode(URigVMNode* InNode, bool bFollowCore
 	}
 }
 
+void URigVMController::RemovePinsDuringRepopulate(URigVMNode* InNode, TArray<URigVMPin*>& InPins, bool bNotify, bool bSetupOrphanedPins)
+{
+	TArray<URigVMPin*> Pins = InPins;
+	for (URigVMPin* Pin : Pins)
+	{
+		if(bSetupOrphanedPins && !Pin->IsExecuteContext())
+		{
+			URigVMPin* RootPin = Pin->GetRootPin();
+			const FString OrphanedName = FString::Printf(TEXT("%s%s"), *URigVMPin::OrphanPinPrefix, *RootPin->GetName());
+
+			URigVMPin* OrphanedRootPin = nullptr;
+			
+			for(URigVMPin* OrphanedPin : InNode->OrphanedPins)
+			{
+				if(OrphanedPin->GetName() == OrphanedName)
+				{
+					OrphanedRootPin = OrphanedPin;
+					break;
+				}
+			}
+
+			if(OrphanedRootPin == nullptr)
+			{
+				if(Pin->IsRootPin()) // if we are passing root pins we can reparent them directly
+				{
+					RootPin->DisplayName = RootPin->GetFName();
+					RootPin->Rename(*OrphanedName, nullptr, REN_ForceNoResetLoaders);
+					InNode->Pins.Remove(RootPin);
+
+					if(bNotify)
+					{
+						Notify(ERigVMGraphNotifType::PinRemoved, RootPin);
+					}
+
+					InNode->OrphanedPins.Add(RootPin);
+
+					if(bNotify)
+					{
+						Notify(ERigVMGraphNotifType::PinAdded, RootPin);
+					}
+				}
+				else // while if we are iterating over sub pins - we should reparent them
+				{
+					OrphanedRootPin = NewObject<URigVMPin>(RootPin->GetNode(), *OrphanedName);
+					ConfigurePinFromPin(OrphanedRootPin, RootPin);
+					OrphanedRootPin->DisplayName = RootPin->GetFName();
+				
+					OrphanedRootPin->GetNode()->OrphanedPins.Add(OrphanedRootPin);
+
+					if(bNotify)
+					{
+						Notify(ERigVMGraphNotifType::PinAdded, OrphanedRootPin);
+					}
+				}
+			}
+
+			if(!Pin->IsRootPin() && (OrphanedRootPin != nullptr))
+			{
+				Pin->Rename(nullptr, OrphanedRootPin, REN_ForceNoResetLoaders);
+				RootPin->SubPins.Remove(Pin);
+				OrphanedRootPin->SubPins.Add(Pin);
+			}
+		}
+	}
+
+	for (URigVMPin* Pin : Pins)
+	{
+		if(!Pin->IsOrphanPin())
+		{
+			RemovePin(Pin, false, bNotify);
+		}
+	}
+	InPins.Reset();
+}
+
+bool URigVMController::RemoveUnusedOrphanedPins(URigVMNode* InNode, bool bNotify)
+{
+	if(!InNode->HasOrphanedPins())
+	{
+		return true;
+	}
+	
+	TArray<URigVMPin*> RemainingOrphanPins;
+	for(int32 PinIndex=0; PinIndex < InNode->OrphanedPins.Num(); PinIndex++)
+	{
+		URigVMPin* OrphanedPin = InNode->OrphanedPins[PinIndex];
+
+		const int32 NumSourceLinks = OrphanedPin->GetSourceLinks(true).Num(); 
+		const int32 NumTargetLinks = OrphanedPin->GetTargetLinks(true).Num();
+
+		if(NumSourceLinks + NumTargetLinks == 0)
+		{
+			RemovePin(OrphanedPin, false, bNotify);
+		}
+		else
+		{
+			RemainingOrphanPins.Add(OrphanedPin);
+		}
+	}
+
+	InNode->OrphanedPins = RemainingOrphanPins;
+	
+	return !InNode->HasOrphanedPins();
+}
+
 #endif
 
 void URigVMController::SetupDefaultUnitNodeDelegates(TDelegate<FName(FRigVMExternalVariable, FString)> InCreateExternalVariableDelegate)
@@ -8223,6 +8374,9 @@ TMap<FString, FString> URigVMController::GetRedirectedPinPaths(URigVMNode* InNod
 URigVMController::FPinState URigVMController::GetPinState(URigVMPin* InPin) const
 {
 	FPinState State;
+	State.Direction = InPin->GetDirection();
+	State.CPPType = InPin->GetCPPType();
+	State.CPPTypeObject = InPin->GetCPPTypeObject();
 	State.DefaultValue = InPin->GetDefaultValue();
 	State.BoundVariable = InPin->GetBoundVariablePath();
 	State.bIsExpanded = InPin->IsExpanded();
@@ -8552,7 +8706,7 @@ bool URigVMController::ChangePinType(const FString& InPinPath, const FString& In
 	return false;
 }
 
-bool URigVMController::ChangePinType(URigVMPin* InPin, const FString& InCPPType, const FName& InCPPTypeObjectPath, bool bSetupUndoRedo)
+bool URigVMController::ChangePinType(URigVMPin* InPin, const FString& InCPPType, const FName& InCPPTypeObjectPath, bool bSetupUndoRedo, bool bSetupOrphanPins)
 {
 	if (InPin->CPPType == InCPPType)
 	{
@@ -8582,23 +8736,50 @@ bool URigVMController::ChangePinType(URigVMPin* InPin, const FString& InCPPType,
 		ActionStack->BeginAction(Action);
 	}
 
+	TArray<URigVMLink*> Links;
+
 	if (bSetupUndoRedo)
 	{
-		BreakAllLinks(InPin, true, true);
-		BreakAllLinks(InPin, false, true);
-		BreakAllLinksRecursive(InPin, true, false, true);
-		BreakAllLinksRecursive(InPin, false, false, true);
-
+		if(!bSetupOrphanPins)
+		{
+			BreakAllLinks(InPin, true, true);
+			BreakAllLinks(InPin, false, true);
+			BreakAllLinksRecursive(InPin, true, false, true);
+			BreakAllLinksRecursive(InPin, false, false, true);
+		}
 		ActionStack->AddAction(FRigVMChangePinTypeAction(InPin, InCPPType, InCPPTypeObjectPath));
 	}
+	
+	if(bSetupOrphanPins)
+	{
+		Links.Append(InPin->GetSourceLinks(true));
+		Links.Append(InPin->GetTargetLinks(true));
+		DetachLinksFromPinObjects(&Links, true);
+
+		const FString OrphanedName = FString::Printf(TEXT("%s%s"), *URigVMPin::OrphanPinPrefix, *InPin->GetName());
+		if(InPin->GetNode()->FindPin(OrphanedName) == nullptr)
+		{
+			URigVMPin* OrphanedPin = NewObject<URigVMPin>(InPin->GetNode(), *OrphanedName);
+			ConfigurePinFromPin(OrphanedPin, InPin);
+			OrphanedPin->DisplayName = InPin->GetFName();
+
+			if(OrphanedPin->IsStruct())
+			{
+				AddPinsForStruct(OrphanedPin->GetScriptStruct(), OrphanedPin->GetNode(), OrphanedPin, OrphanedPin->Direction, OrphanedPin->GetDefaultValue(), false, true);
+			}
+				
+			InPin->GetNode()->OrphanedPins.Add(OrphanedPin);
+		}
+	}
 
 	TArray<URigVMPin*> Pins = InPin->SubPins;
 	for (URigVMPin* Pin : Pins)
 	{
 		RemovePin(Pin, false, true);
 	}
+	
 	InPin->SubPins.Reset();
-
+	
 	InPin->CPPType = InCPPType;
 	InPin->CPPTypeObjectPath = InCPPTypeObjectPath;
 	InPin->CPPTypeObject = CPPTypeObject;
@@ -8619,6 +8800,12 @@ bool URigVMController::ChangePinType(URigVMPin* InPin, const FString& InCPPType,
 		ActionStack->EndAction(Action);
 	}
 
+	if(Links.Num() > 0)
+	{
+		ReattachLinksToPinObjects(false, &Links, true, true);
+		RemoveUnusedOrphanedPins(InPin->GetNode(), true);
+	}
+
 	return true;
 }
 
diff --git a/Engine/Source/Developer/RigVMDeveloper/Private/RigVMModel/RigVMNode.cpp b/Engine/Source/Developer/RigVMDeveloper/Private/RigVMModel/RigVMNode.cpp
--- a/Engine/Source/Developer/RigVMDeveloper/Private/RigVMModel/RigVMNode.cpp
+++ b/Engine/Source/Developer/RigVMDeveloper/Private/RigVMModel/RigVMNode.cpp
@@ -145,9 +145,30 @@ URigVMPin* URigVMNode::FindPin(const FString& InPinPath) const
 			return Pin->FindSubPin(Right);
 		}
 	}
+
+	if(Left.StartsWith(URigVMPin::OrphanPinPrefix))
+	{
+		for (URigVMPin* Pin : OrphanedPins)
+		{
+			if (Pin->GetName() == Left)
+			{
+				if (Right.IsEmpty())
+				{
+					return Pin;
+				}
+				return Pin->FindSubPin(Right);
+			}
+		}
+	}
+	
 	return nullptr;
 }
 
+const TArray<URigVMPin*>& URigVMNode::GetOrphanedPins() const
+{
+	return OrphanedPins;
+}
+
 URigVMGraph* URigVMNode::GetGraph() const
 {
 	if (URigVMGraph* Graph = Cast<URigVMGraph>(GetOuter()))
diff --git a/Engine/Source/Developer/RigVMDeveloper/Private/RigVMModel/RigVMPin.cpp b/Engine/Source/Developer/RigVMDeveloper/Private/RigVMModel/RigVMPin.cpp
--- a/Engine/Source/Developer/RigVMDeveloper/Private/RigVMModel/RigVMPin.cpp
+++ b/Engine/Source/Developer/RigVMDeveloper/Private/RigVMModel/RigVMPin.cpp
@@ -47,6 +47,7 @@ URigVMPin* URigVMInjectionInfo::GetPin() const
 
 const URigVMPin::FPinOverrideMap URigVMPin::EmptyPinOverrideMap;
 const URigVMPin::FPinOverride URigVMPin::EmptyPinOverride = URigVMPin::FPinOverride(FRigVMASTProxy(), EmptyPinOverrideMap);
+const FString URigVMPin::OrphanPinPrefix = TEXT("Orphan::");
 
 bool URigVMPin::SplitPinPathAtStart(const FString& InPinPath, FString& LeftMost, FString& Right)
 {
@@ -867,6 +868,18 @@ FRigVMExternalVariable URigVMPin::ToExternalVariable() const
 	return ExternalVariable;
 }
 
+bool URigVMPin::IsOrphanPin() const
+{
+	if(URigVMPin* RootPin = GetRootPin())
+	{
+		if(RootPin != this)
+		{
+			return RootPin->IsOrphanPin();
+		}
+	}
+	return GetNode()->OrphanedPins.Contains(this); 
+}
+
 void URigVMPin::UpdateCPPTypeObjectIfRequired() const
 {
 	if (CPPTypeObject == nullptr)
diff --git a/Engine/Source/Developer/RigVMDeveloper/Public/RigVMModel/RigVMController.h b/Engine/Source/Developer/RigVMDeveloper/Public/RigVMModel/RigVMController.h
--- a/Engine/Source/Developer/RigVMDeveloper/Public/RigVMModel/RigVMController.h
+++ b/Engine/Source/Developer/RigVMDeveloper/Public/RigVMModel/RigVMController.h
@@ -128,7 +128,7 @@ public:
 
 	// Refreshes the variable node with the new data
 	UFUNCTION(BlueprintCallable, Category = RigVMController)
-	void RefreshVariableNode(const FName& InNodeName, const FName& InVariableName, const FString& InCPPType, UObject* InCPPTypeObject, bool bSetupUndoRedo);
+	void RefreshVariableNode(const FName& InNodeName, const FName& InVariableName, const FString& InCPPType, UObject* InCPPTypeObject, bool bSetupUndoRedo, bool bDetachAndReattachLinks = true);
 
 	// Removes all nodes related to a given variable
 	UFUNCTION(BlueprintCallable, Category = RigVMController)
@@ -533,7 +533,7 @@ public:
 
 	// Changes the type of an exposed pin in the graph controlled by this
 	UFUNCTION(BlueprintCallable, Category = RigVMController)
-	bool ChangeExposedPinType(const FName& InPinName, const FString& InCPPType, const FName& InCPPTypeObjectPath, bool bSetupUndoRedo = true);
+	bool ChangeExposedPinType(const FName& InPinName, const FString& InCPPType, const FName& InCPPTypeObjectPath, bool bSetupUndoRedo = true, bool bSetupOrphanPins = true);
 
 	// Sets the index for an exposed pin. This can be used to move the pin up and down on the node.
 	UFUNCTION(BlueprintCallable, Category = RigVMController)
@@ -570,18 +570,22 @@ public:
 	FRigVMController_RequestLocalizeFunctionDelegate RequestLocalizeFunctionDelegate;
 
 	int32 DetachLinksFromPinObjects(const TArray<URigVMLink*>* InLinks = nullptr, bool bNotify = false);
-	int32 ReattachLinksToPinObjects(bool bFollowCoreRedirectors = false, const TArray<URigVMLink*>* InLinks = nullptr, bool bNotify = false);
+	int32 ReattachLinksToPinObjects(bool bFollowCoreRedirectors = false, const TArray<URigVMLink*>* InLinks = nullptr, bool bNotify = false, bool bSetupOrphanedPins = false);
 	void AddPinRedirector(bool bInput, bool bOutput, const FString& OldPinPath, const FString& NewPinPath);
 
 	// Removes nodes which went stale.
 	void RemoveStaleNodes();
 
 #if WITH_EDITOR
-
 	bool ShouldRedirectPin(UScriptStruct* InOwningStruct, const FString& InOldRelativePinPath, FString& InOutNewRelativePinPath) const;
 	bool ShouldRedirectPin(const FString& InOldPinPath, FString& InOutNewPinPath) const;
 
-	void RepopulatePinsOnNode(URigVMNode* InNode, bool bFollowCoreRedirectors = true, bool bNotify = false);
+	void RepopulatePinsOnNode(URigVMNode* InNode, bool bFollowCoreRedirectors = true, bool bNotify = false, bool bSetupOrphanedPins = false);
+	void RemovePinsDuringRepopulate(URigVMNode* InNode, TArray<URigVMPin*>& InPins, bool bNotify, bool bSetupOrphanedPins);
+
+	// removes any orphan pins that no longer holds a link
+	bool RemoveUnusedOrphanedPins(URigVMNode* InNode, bool bNotify);
+	
 #endif
 
 	FRigVMUnitNodeCreatedContext& GetUnitNodeCreatedContext() { return UnitNodeCreatedContext; }
@@ -676,6 +680,9 @@ private:
 
 	struct FPinState
 	{
+		ERigVMPinDirection Direction;
+		FString CPPType;
+		UObject* CPPTypeObject;
 		FString DefaultValue;
 		FString BoundVariable;
 		bool bIsExpanded;
@@ -696,7 +703,7 @@ private:
 	void PotentiallyResolvePrototypeNode(URigVMPrototypeNode* InNode, bool bSetupUndoRedo);
 	void PotentiallyResolvePrototypeNode(URigVMPrototypeNode* InNode, bool bSetupUndoRedo, TArray<URigVMNode*>& NodesVisited);
 	bool ChangePinType(const FString& InPinPath, const FString& InCPPType, const FName& InCPPTypeObjectPath, bool bSetupUndoRedo);
-	bool ChangePinType(URigVMPin* InPin, const FString& InCPPType, const FName& InCPPTypeObjectPath, bool bSetupUndoRedo);
+	bool ChangePinType(URigVMPin* InPin, const FString& InCPPType, const FName& InCPPTypeObjectPath, bool bSetupUndoRedo, bool bSetupOrphanPins = true);
 
 #if WITH_EDITOR
 	void RewireLinks(URigVMPin* OldPin, URigVMPin* NewPin, bool bAsInput, bool bSetupUndoRedo, TArray<URigVMLink*> InLinks = TArray<URigVMLink*>());
diff --git a/Engine/Source/Developer/RigVMDeveloper/Public/RigVMModel/RigVMNode.h b/Engine/Source/Developer/RigVMDeveloper/Public/RigVMModel/RigVMNode.h
--- a/Engine/Source/Developer/RigVMDeveloper/Public/RigVMModel/RigVMNode.h
+++ b/Engine/Source/Developer/RigVMDeveloper/Public/RigVMModel/RigVMNode.h
@@ -68,6 +68,14 @@ public:
 	UFUNCTION(BlueprintCallable, Category = RigVMNode)
 	URigVMPin* FindPin(const FString& InPinPath) const;
 
+	// Returns all of the top-level orphaned Pins of this Node.
+	UFUNCTION(BlueprintCallable, Category = RigVMNode)
+    virtual const TArray<URigVMPin*>& GetOrphanedPins() const;
+
+	// Returns true if the node has orphaned pins - which leads to a compiler error
+	UFUNCTION(BlueprintPure, Category = RigVMNode)
+    FORCEINLINE bool HasOrphanedPins() const { return GetOrphanedPins().Num() > 0; }
+
 	// Returns the Graph of this Node
 	UFUNCTION(BlueprintCallable, Category = RigVMNode)
 	URigVMGraph* GetGraph() const;
@@ -227,6 +235,9 @@ private:
 	UPROPERTY()
 	TArray<URigVMPin*> Pins;
 
+	UPROPERTY()
+	TArray<URigVMPin*> OrphanedPins;
+
 	int32 GetSliceContextBracket;
 
 	friend class URigVMController;
diff --git a/Engine/Source/Developer/RigVMDeveloper/Public/RigVMModel/RigVMPin.h b/Engine/Source/Developer/RigVMDeveloper/Public/RigVMModel/RigVMPin.h
--- a/Engine/Source/Developer/RigVMDeveloper/Public/RigVMModel/RigVMPin.h
+++ b/Engine/Source/Developer/RigVMDeveloper/Public/RigVMModel/RigVMPin.h
@@ -359,6 +359,9 @@ public:
 	// Returns an external variable matching this pin's type
 	FRigVMExternalVariable ToExternalVariable() const;
 
+	// Returns true if the pin has been orphaned
+	bool IsOrphanPin() const;
+
 private:
 
 	void UpdateCPPTypeObjectIfRequired() const;
@@ -419,6 +422,8 @@ private:
 	UPROPERTY()
 	FString BoundVariablePath;
 
+	static const FString OrphanPinPrefix;
+
 	friend class URigVMController;
 	friend class URigVMGraph;
 	friend class URigVMNode;
