Control Rig: Follow up on orphan pins: Ensure to only setup orphan pins during postload and avoid compilation of invalid graphs

diff --git a/Engine/Source/Developer/RigVMDeveloper/Private/RigVMCompiler/RigVMCompiler.cpp b/Engine/Source/Developer/RigVMDeveloper/Private/RigVMCompiler/RigVMCompiler.cpp
--- a/Engine/Source/Developer/RigVMDeveloper/Private/RigVMCompiler/RigVMCompiler.cpp
+++ b/Engine/Source/Developer/RigVMDeveloper/Private/RigVMCompiler/RigVMCompiler.cpp
@@ -129,6 +129,8 @@ bool URigVMCompiler::Compile(URigVMGraph* InGraph, URigVMController* InControlle
 	// traverse all graphs and try to clear out orphan pins
 	TArray<URigVMGraph*> VisitedGraphs;
 	VisitedGraphs.Add(InGraph);
+
+	int32 NodesWithOrphanPins = 0;
 	for(int32 GraphIndex=0; GraphIndex<VisitedGraphs.Num(); GraphIndex++)
 	{
 		URigVMGraph* VisitedGraph = VisitedGraphs[GraphIndex];
@@ -139,6 +141,7 @@ bool URigVMCompiler::Compile(URigVMGraph* InGraph, URigVMController* InControlle
 			{
 				static const FString LinkedMessage = TEXT("Node @@ uses pins that no longer exist. Please rewire the links and re-compile.");
 				Settings.ASTSettings.Report(EMessageSeverity::Error, ModelNode, LinkedMessage);
+				NodesWithOrphanPins++;
 			}
 
 			if(URigVMLibraryNode* LibraryNode = Cast<URigVMLibraryNode>(ModelNode))
@@ -150,6 +153,11 @@ bool URigVMCompiler::Compile(URigVMGraph* InGraph, URigVMController* InControlle
 			}
 		}
 	}
+
+	if(NodesWithOrphanPins > 0)
+	{
+		return false;
+	}
 #endif
 	
 	OutVM->ClearExternalVariables();
diff --git a/Engine/Source/Developer/RigVMDeveloper/Private/RigVMModel/RigVMController.cpp b/Engine/Source/Developer/RigVMDeveloper/Private/RigVMModel/RigVMController.cpp
--- a/Engine/Source/Developer/RigVMDeveloper/Private/RigVMModel/RigVMController.cpp
+++ b/Engine/Source/Developer/RigVMDeveloper/Private/RigVMModel/RigVMController.cpp
@@ -7589,7 +7589,7 @@ int32 URigVMController::ReattachLinksToPinObjects(bool bFollowCoreRedirectors, c
 		URigVMPin* SourcePin = Link->GetSourcePin();
 		URigVMPin* TargetPin = Link->GetTargetPin();
 
-		if((SourcePin != nullptr) && (TargetPin != nullptr))
+		if(bSetupOrphanedPins && (SourcePin != nullptr) && (TargetPin != nullptr))
 		{
 			if (!URigVMPin::CanLink(SourcePin, TargetPin, nullptr, nullptr))
 			{
