Control Rig: Correct typo in RigVMController

diff --git a/Engine/Source/Developer/RigVMDeveloper/Private/RigVMModel/RigVMController.cpp b/Engine/Source/Developer/RigVMDeveloper/Private/RigVMModel/RigVMController.cpp
--- a/Engine/Source/Developer/RigVMDeveloper/Private/RigVMModel/RigVMController.cpp
+++ b/Engine/Source/Developer/RigVMDeveloper/Private/RigVMModel/RigVMController.cpp
@@ -493,7 +493,7 @@ URigVMVariableNode* URigVMController::AddVariableNodeFromObjectPath(const FName&
 	return AddVariableNode(InVariableName, InCPPType, CPPTypeObject, bIsGetter, InDefaultValue, InPosition, InNodeName, bSetupUndoRedo);
 }
 
-void URigVMController::RefreshVariableNode(const FName& InNodeName, const FName& InVariableName, const FString& InCPPType, UObject* InCPPTypeObject, bool bSetupUndoRedo, bool bSetupOrphanPinss)
+void URigVMController::RefreshVariableNode(const FName& InNodeName, const FName& InVariableName, const FString& InCPPType, UObject* InCPPTypeObject, bool bSetupUndoRedo, bool bSetupOrphanPins)
 {
 	if (!IsValidGraph())
 	{
@@ -551,7 +551,7 @@ void URigVMController::RefreshVariableNode(const FName& InNodeName, const FName&
 						{
 							CPPTypeObjectPath = InCPPTypeObject->GetPathName();
 						}
-						ChangePinType(ValuePin, InCPPType, *CPPTypeObjectPath, bSetupUndoRedo, bSetupOrphanPinss);
+						ChangePinType(ValuePin, InCPPType, *CPPTypeObjectPath, bSetupUndoRedo, bSetupOrphanPins);
 					}
 				}
 			}
