Control Rig: Ensure to use REN_ForceNoResetLoaders for rename operations during load

diff --git a/Engine/Source/Developer/RigVMDeveloper/Private/RigVMModel/RigVMController.cpp b/Engine/Source/Developer/RigVMDeveloper/Private/RigVMModel/RigVMController.cpp
--- a/Engine/Source/Developer/RigVMDeveloper/Private/RigVMModel/RigVMController.cpp
+++ b/Engine/Source/Developer/RigVMDeveloper/Private/RigVMModel/RigVMController.cpp
@@ -8179,7 +8179,7 @@ void URigVMController::RemovePinsDuringRepopulate(URigVMNode* InNode, TArray<URi
 				if(Pin->IsRootPin()) // if we are passing root pins we can reparent them directly
 				{
 					RootPin->DisplayName = RootPin->GetFName();
-					RootPin->Rename(*OrphanedName, nullptr, REN_ForceNoResetLoaders);
+					RootPin->Rename(*OrphanedName, nullptr, REN_ForceNoResetLoaders | REN_DoNotDirty | REN_DontCreateRedirectors | REN_NonTransactional);
 					InNode->Pins.Remove(RootPin);
 
 					if(bNotify)
@@ -8211,7 +8211,7 @@ void URigVMController::RemovePinsDuringRepopulate(URigVMNode* InNode, TArray<URi
 
 			if(!Pin->IsRootPin() && (OrphanedRootPin != nullptr))
 			{
-				Pin->Rename(nullptr, OrphanedRootPin, REN_ForceNoResetLoaders);
+				Pin->Rename(nullptr, OrphanedRootPin, REN_ForceNoResetLoaders | REN_DoNotDirty | REN_DontCreateRedirectors | REN_NonTransactional);
 				RootPin->SubPins.Remove(Pin);
 				OrphanedRootPin->SubPins.Add(Pin);
 			}
@@ -8405,7 +8405,7 @@ void URigVMController::ApplyPinState(URigVMPin* InPin, const FPinState& InPinSta
 {
 	for (URigVMInjectionInfo* InjectionInfo : InPinState.InjectionInfos)
 	{
-		InjectionInfo->Rename(nullptr, InPin);
+		InjectionInfo->Rename(nullptr, InPin, REN_ForceNoResetLoaders | REN_DoNotDirty | REN_DontCreateRedirectors | REN_NonTransactional);
 		InjectionInfo->InputPin = InjectionInfo->UnitNode->FindPin(InjectionInfo->InputPin->GetName());
 		InjectionInfo->OutputPin = InjectionInfo->UnitNode->FindPin(InjectionInfo->OutputPin->GetName());
 		InPin->InjectionInfos.Add(InjectionInfo);
