Control Rig: Make sure to reattach links after a promotion from function to collapse node

diff --git a/Engine/Source/Developer/RigVMDeveloper/Private/RigVMModel/RigVMController.cpp b/Engine/Source/Developer/RigVMDeveloper/Private/RigVMModel/RigVMController.cpp
--- a/Engine/Source/Developer/RigVMDeveloper/Private/RigVMModel/RigVMController.cpp
+++ b/Engine/Source/Developer/RigVMDeveloper/Private/RigVMModel/RigVMController.cpp
@@ -3299,6 +3299,11 @@ URigVMFunctionReferenceNode* URigVMController::PromoteCollapseNodeToFunctionRefe
 				AddLink(LinkPath.Key, LinkPath.Value, bSetupUndoRedo);
 			}
 		}
+
+		{
+			FRigVMControllerGraphGuard GraphGuard(this, FunctionDefinition->GetContainedGraph(), false);
+			ReattachLinksToPinObjects(false, nullptr, true, false);
+		}
 	}
 
 	if (bSetupUndoRedo)
