Control Rig: Avoid undo / redo during compilation. Improve RemoveExpressions performance within AST Parser

diff --git a/Engine/Source/Developer/RigVMDeveloper/Private/RigVMCompiler/RigVMAST.cpp b/Engine/Source/Developer/RigVMDeveloper/Private/RigVMCompiler/RigVMAST.cpp
--- a/Engine/Source/Developer/RigVMDeveloper/Private/RigVMCompiler/RigVMAST.cpp
+++ b/Engine/Source/Developer/RigVMDeveloper/Private/RigVMCompiler/RigVMAST.cpp
@@ -2461,57 +2461,97 @@ const FRigVMBlockExprAST* FRigVMParserAST::GetObsoleteBlock(bool bCreateIfMissin
 	return ObsoleteBlock;
 }
 
-void FRigVMParserAST::RemoveExpression(FRigVMExprAST* InExpr, bool bRefreshIndices, bool bRecurseToChildren)
+void FRigVMParserAST::RemoveExpressions(TArray<FRigVMExprAST*> InExprs)
 {
-	TArray<FRigVMExprAST*> Parents = InExpr->Parents;
-	for (FRigVMExprAST* Parent : Parents)
+	DECLARE_SCOPE_HIERARCHICAL_COUNTER_FUNC()
+
+	if(InExprs.IsEmpty())
 	{
-		InExpr->RemoveParent(Parent);
+		return;
 	}
-	TArray<FRigVMExprAST*> Children = InExpr->Children;
-	for (FRigVMExprAST* Child : Children)
+	
+	RefreshExprIndices();
+
+	TArray<FRigVMExprAST*> ExpressionsToRemove;
+	ExpressionsToRemove.Append(InExprs);
+
+	TArray<int32> NumRemainingParents;
+	NumRemainingParents.AddUninitialized(Expressions.Num());
+	for(int32 ExpressionIndex = 0; ExpressionIndex < Expressions.Num(); ExpressionIndex++)
+	{
+		NumRemainingParents[ExpressionIndex] = Expressions[ExpressionIndex]->Parents.Num();
+	}
+
+	TArray<bool> bRemoveExpression;
+	bRemoveExpression.AddZeroed(Expressions.Num());
+	for(int32 ExpressionIndex = 0; ExpressionIndex < ExpressionsToRemove.Num(); ExpressionIndex++)
 	{
-		Child->RemoveParent(InExpr);
+		FRigVMExprAST* Expr = ExpressionsToRemove[ExpressionIndex];
+		bRemoveExpression[Expr->GetIndex()] = true;
 
-		if (bRecurseToChildren && Child->NumParents() == 0)
+		for (FRigVMExprAST* Child : Expr->Children)
 		{
-			RemoveExpression(Child, false, bRecurseToChildren);
+			if(--NumRemainingParents[Child->GetIndex()] == 0)
+			{
+				ExpressionsToRemove.Add(Child);
+			}
 		}
 	}
 
-	Expressions.Remove(InExpr);
+	TArray<FRigVMExprAST*> RemainingExpressions;
+	RemainingExpressions.Reserve(Expressions.Num() - ExpressionsToRemove.Num());
+	
+	for(int32 ExpressionIndex = 0; ExpressionIndex < Expressions.Num(); ExpressionIndex++)
+	{
+		if(!bRemoveExpression[ExpressionIndex])
+		{
+			FRigVMExprAST* Expr = Expressions[ExpressionIndex];
+			RemainingExpressions.Add(Expr);
+
+			for(int32 ChildIndex = 0; ChildIndex < Expr->Children.Num(); ChildIndex++)
+			{
+				FRigVMExprAST* ChildExpr = Expr->Children[ChildIndex];
+				if(bRemoveExpression[ChildExpr->GetIndex()])
+				{
+					Expr->Children.RemoveAt(ChildIndex);
+					break;
+				}
+			}
+
+			for(int32 ParentIndex = 0; ParentIndex < Expr->Parents.Num(); ParentIndex++)
+			{
+				FRigVMExprAST* ParentExpr = Expr->Parents[ParentIndex];
+				if(bRemoveExpression[ParentExpr->GetIndex()])
+				{
+					Expr->Parents.RemoveAt(ParentIndex);
+					break;
+				}
+			}
+		}
+	}
+
+	Expressions = RemainingExpressions;
 
 	TArray<FRigVMASTProxy> KeysToRemove;
 	for (TPair<FRigVMASTProxy, FRigVMExprAST*> Pair : SubjectToExpression)
 	{
-		if (Pair.Value == InExpr)
+		if (bRemoveExpression[Pair.Value->GetIndex()])
 		{
 			KeysToRemove.Add(Pair.Key);
 		}
 	}
+	
 	for (const FRigVMASTProxy& KeyToRemove : KeysToRemove)
 	{
 		SubjectToExpression.Remove(KeyToRemove);
 	}
 
-	delete InExpr;
-
-	if (bRefreshIndices)
+	for(int32 ExpressionIndex = ExpressionsToRemove.Num() - 1; ExpressionIndex >= 0; ExpressionIndex--)
 	{
-		RefreshExprIndices();
+		delete ExpressionsToRemove[ExpressionIndex];
 	}
-}
 
-void FRigVMParserAST::RemoveExpressions(TArray<FRigVMExprAST*> InExprs, bool bRefreshIndices, bool bRecurseToChildren)
-{
-	for (FRigVMExprAST* InExpr : InExprs)
-	{
-		RemoveExpression(InExpr, false, bRecurseToChildren);
-	}
-	if (bRefreshIndices)
-	{
-		RefreshExprIndices();
-	}
+	RefreshExprIndices();
 }
 
 void FRigVMParserAST::TraverseParents(const FRigVMExprAST* InExpr, TFunctionRef<bool(const FRigVMExprAST*)> InContinuePredicate)
diff --git a/Engine/Source/Developer/RigVMDeveloper/Private/RigVMCompiler/RigVMCompiler.cpp b/Engine/Source/Developer/RigVMDeveloper/Private/RigVMCompiler/RigVMCompiler.cpp
--- a/Engine/Source/Developer/RigVMDeveloper/Private/RigVMCompiler/RigVMCompiler.cpp
+++ b/Engine/Source/Developer/RigVMDeveloper/Private/RigVMCompiler/RigVMCompiler.cpp
@@ -136,7 +136,7 @@ bool URigVMCompiler::Compile(URigVMGraph* InGraph, URigVMController* InControlle
 		URigVMGraph* VisitedGraph = VisitedGraphs[GraphIndex];
 		for(URigVMNode* ModelNode : VisitedGraph->GetNodes())
 		{
-			FRigVMControllerGraphGuard Guard(InController, VisitedGraph);
+			FRigVMControllerGraphGuard Guard(InController, VisitedGraph, false);
 			if(!InController->RemoveUnusedOrphanedPins(ModelNode, true))
 			{
 				static const FString LinkedMessage = TEXT("Node @@ uses pins that no longer exist. Please rewire the links and re-compile.");
diff --git a/Engine/Source/Developer/RigVMDeveloper/Public/RigVMCompiler/RigVMAST.h b/Engine/Source/Developer/RigVMDeveloper/Public/RigVMCompiler/RigVMAST.h
--- a/Engine/Source/Developer/RigVMDeveloper/Public/RigVMCompiler/RigVMAST.h
+++ b/Engine/Source/Developer/RigVMDeveloper/Public/RigVMCompiler/RigVMAST.h
@@ -1280,17 +1280,9 @@ private:
 		return Expr;
 	}
 
-	// removes a single expression from the parser
-	// @param InExpr the expression to remove
-	// @param bRefreshIndices flag to determine if the expression indices need to be refreshed
-	// @param bRecurseToChildren flag to determine if nested expressions should also be removed
-	void RemoveExpression(FRigVMExprAST* InExpr, bool bRefreshIndices = true, bool bRecurseToChildren = false);
-
 	// removes an array of expressions from the parser
 	// @param InExprs the expressions to remove
-	// @param bRefreshIndices flag to determine if the expression indices need to be refreshed
-	// @param bRecurseToChildren flag to determine if nested expressions should also be removed
-	void RemoveExpressions(TArray<FRigVMExprAST*> InExprs, bool bRefreshIndices = true, bool bRecurseToChildren = true);
+	void RemoveExpressions(TArray<FRigVMExprAST*> InExprs);
 
 	// a static helper function to traverse along all parents of an expression,
 	// provided a predicate to return true if the traverse should continue,
