Control Rig: Avoid cycles between blueprints

diff --git a/Engine/Plugins/Experimental/ControlRig/Source/ControlRigDeveloper/Private/ControlRigBlueprint.cpp b/Engine/Plugins/Experimental/ControlRig/Source/ControlRigDeveloper/Private/ControlRigBlueprint.cpp
--- a/Engine/Plugins/Experimental/ControlRig/Source/ControlRigDeveloper/Private/ControlRigBlueprint.cpp
+++ b/Engine/Plugins/Experimental/ControlRig/Source/ControlRigDeveloper/Private/ControlRigBlueprint.cpp
@@ -1019,6 +1019,30 @@ URigVMController* UControlRigBlueprint::GetOrCreateController(URigVMGraph* InGra
 		return false;
 	});
 
+	Controller->IsDependencyCyclicDelegate.BindLambda([WeakThis](UObject* InDependentObject, UObject* InDependencyObject) -> bool
+	{
+	    if(InDependentObject == nullptr || InDependencyObject == nullptr)
+	    {
+	        return false;
+	    }
+
+		UControlRigBlueprint* DependentBlueprint = InDependentObject->GetTypedOuter<UControlRigBlueprint>();
+		UControlRigBlueprint* DependencyBlueprint = InDependencyObject->GetTypedOuter<UControlRigBlueprint>();
+
+		if(DependentBlueprint == nullptr || DependencyBlueprint == nullptr)
+		{
+			return false;
+		}
+
+		if(DependentBlueprint == DependencyBlueprint)
+		{
+			return false;
+		}
+
+		const TArray<UControlRigBlueprint*> DependencyDependencies = DependencyBlueprint->GetDependencies(true);
+		return DependencyDependencies.Contains(DependentBlueprint);
+	});
+
 #if WITH_EDITOR
 
 	// this sets up three delegates:
diff --git a/Engine/Source/Developer/RigVMDeveloper/Private/RigVMModel/RigVMController.cpp b/Engine/Source/Developer/RigVMDeveloper/Private/RigVMModel/RigVMController.cpp
--- a/Engine/Source/Developer/RigVMDeveloper/Private/RigVMModel/RigVMController.cpp
+++ b/Engine/Source/Developer/RigVMDeveloper/Private/RigVMModel/RigVMController.cpp
@@ -7121,6 +7121,18 @@ bool URigVMController::CanAddFunctionRefForDefinition(URigVMLibraryNode* InFunct
 		}
 	}
 
+	if(IsDependencyCyclicDelegate.IsBound())
+	{
+		if(IsDependencyCyclicDelegate.Execute(Graph, InFunctionDefinition))
+		{
+			if(bReportErrors)
+			{
+				ReportAndNotifyError(TEXT("Function is not available for placement in this graph host due to dependency cycles."));
+			}
+			return false;
+		}
+	}
+
 	URigVMLibraryNode* ParentLibraryNode = Cast<URigVMLibraryNode>(Graph->GetOuter());
 	while (ParentLibraryNode)
 	{
diff --git a/Engine/Source/Developer/RigVMDeveloper/Public/RigVMModel/RigVMController.h b/Engine/Source/Developer/RigVMDeveloper/Public/RigVMModel/RigVMController.h
--- a/Engine/Source/Developer/RigVMDeveloper/Public/RigVMModel/RigVMController.h
+++ b/Engine/Source/Developer/RigVMDeveloper/Public/RigVMModel/RigVMController.h
@@ -25,6 +25,7 @@ DECLARE_DELEGATE_RetVal(TArray<FRigVMExternalVariable>, FRigVMController_GetExte
 DECLARE_DELEGATE_RetVal(const FRigVMByteCode*, FRigVMController_GetByteCodeDelegate)
 DECLARE_DELEGATE_RetVal_OneParam(bool, FRigVMController_IsFunctionAvailableDelegate, URigVMLibraryNode*)
 DECLARE_DELEGATE_RetVal_OneParam(bool, FRigVMController_RequestLocalizeFunctionDelegate, URigVMLibraryNode*)
+DECLARE_DELEGATE_RetVal_TwoParams(bool, FRigVMController_IsDependencyCyclicDelegate, UObject*, UObject*)
 
 /**
  * The Controller is the sole authority to perform changes
@@ -568,6 +569,9 @@ public:
 
 	// A delegate to localize a function on demand
 	FRigVMController_RequestLocalizeFunctionDelegate RequestLocalizeFunctionDelegate;
+	
+	// A delegate to validate if we are allowed to introduce a dependency between two objects
+	FRigVMController_IsDependencyCyclicDelegate IsDependencyCyclicDelegate; 
 
 	int32 DetachLinksFromPinObjects(const TArray<URigVMLink*>* InLinks = nullptr, bool bNotify = false);
 	int32 ReattachLinksToPinObjects(bool bFollowCoreRedirectors = false, const TArray<URigVMLink*>* InLinks = nullptr, bool bNotify = false, bool bSetupOrphanedPins = false);
