Control Rig: Fix cook warnings generated due to duplicated links and invalid links to orphan pins. Duplicated links are now removed once detected. Invalid Links on orpha pins are now allowed.

Also Control Rig Editor now allows opening of BP with compile errors

diff --git a/Engine/Plugins/Experimental/ControlRig/Source/ControlRigEditor/Private/Editor/ControlRigEditor.cpp b/Engine/Plugins/Experimental/ControlRig/Source/ControlRigEditor/Private/Editor/ControlRigEditor.cpp
--- a/Engine/Plugins/Experimental/ControlRig/Source/ControlRigEditor/Private/Editor/ControlRigEditor.cpp
+++ b/Engine/Plugins/Experimental/ControlRig/Source/ControlRigEditor/Private/Editor/ControlRigEditor.cpp
@@ -2831,15 +2831,6 @@ void FControlRigEditor::UpdateControlRig()
 			Blueprint->RecompileVM();
 			Blueprint->ClearFlags(RF_Transient);
 
-			if (Blueprint->bErrorsDuringCompilation)
-			{
-				AnimInstance->ResetControlRigTracks();
-				AnimInstance->RecalcRequiredBones();
-				GetBlueprintObj()->SetObjectBeingDebugged(nullptr);
-				ControlRig = nullptr;
-				return;
-			}
-
 			ControlRig->OnInitialized_AnyThread().AddSP(this, &FControlRigEditor::HandleControlRigExecutedEvent);
 			ControlRig->OnExecuted_AnyThread().AddSP(this, &FControlRigEditor::HandleControlRigExecutedEvent);
 			ControlRig->RequestInit();
diff --git a/Engine/Source/Developer/RigVMDeveloper/Private/RigVMModel/RigVMController.cpp b/Engine/Source/Developer/RigVMDeveloper/Private/RigVMModel/RigVMController.cpp
--- a/Engine/Source/Developer/RigVMDeveloper/Private/RigVMModel/RigVMController.cpp
+++ b/Engine/Source/Developer/RigVMDeveloper/Private/RigVMModel/RigVMController.cpp
@@ -7611,6 +7611,12 @@ int32 URigVMController::ReattachLinksToPinObjects(bool bFollowCoreRedirectors, c
 
 		if(bSetupOrphanedPins && (SourcePin != nullptr) && (TargetPin != nullptr))
 		{
+			// ignore duplicated links that have been processed
+			if (SourcePin->IsLinkedTo(TargetPin))
+			{
+				continue;
+			}
+
 			if (!URigVMPin::CanLink(SourcePin, TargetPin, nullptr, nullptr))
 			{
 				if(SourcePin->GetNode()->HasOrphanedPins() && bSetupOrphanedPins)
@@ -7672,17 +7678,6 @@ int32 URigVMController::ReattachLinksToPinObjects(bool bFollowCoreRedirectors, c
 					}
 				}
 			}
-
-			if(bRelayedBackToOrphanPin && (SourcePin != nullptr) && (TargetPin != nullptr))
-			{
-				if (!URigVMPin::CanLink(SourcePin, TargetPin, nullptr, nullptr))
-				{
-					ReportWarningf(TEXT("Unable to re-create link %s -> %s"), *Link->SourcePinPath, *Link->TargetPinPath);
-					TargetPin->Links.Remove(Link);
-					SourcePin->Links.Remove(Link);
-					continue;
-				}
-			}
 		}
 		
 		if (SourcePin == nullptr)
