Control Rig: Variable Remapping

diff --git a/Engine/Plugins/Experimental/ControlRig/Source/ControlRigDeveloper/Private/ControlRigBlueprint.cpp b/Engine/Plugins/Experimental/ControlRig/Source/ControlRigDeveloper/Private/ControlRigBlueprint.cpp
--- a/Engine/Plugins/Experimental/ControlRig/Source/ControlRigDeveloper/Private/ControlRigBlueprint.cpp
+++ b/Engine/Plugins/Experimental/ControlRig/Source/ControlRigDeveloper/Private/ControlRigBlueprint.cpp
@@ -1154,16 +1154,19 @@ URigVMController* UControlRigBlueprint::GetOrCreateController(URigVMGraph* InGra
 	// this delegate is used by the controller to determine variable validity
 	// during a bind process. the controller itself doesn't own the variables,
 	// so we need a delegate to request them from the owning blueprint
-	Controller->GetExternalVariablesDelegate.BindLambda([WeakThis]() -> TArray<FRigVMExternalVariable> {
+	Controller->GetExternalVariablesDelegate.BindLambda([](URigVMGraph* InGraph) -> TArray<FRigVMExternalVariable> {
 
-		if (WeakThis.IsValid())
+		if (InGraph)
 		{
-			if (UControlRigBlueprintGeneratedClass* RigClass = WeakThis->GetControlRigBlueprintGeneratedClass())
+			if(UControlRigBlueprint* Blueprint = InGraph->GetTypedOuter<UControlRigBlueprint>())
 			{
-				if (UControlRig* CDO = Cast<UControlRig>(RigClass->GetDefaultObject(true /* create if needed */)))
+				if (UControlRigBlueprintGeneratedClass* RigClass = Blueprint->GetControlRigBlueprintGeneratedClass())
 				{
-					return CDO->GetExternalVariablesImpl(true /* rely on variables within blueprint */);
-				}
+                    if (UControlRig* CDO = Cast<UControlRig>(RigClass->GetDefaultObject(true /* create if needed */)))
+                    {
+                        return CDO->GetExternalVariablesImpl(true /* rely on variables within blueprint */);
+                    }
+                }
 			}
 		}
 		return TArray<FRigVMExternalVariable>();
@@ -2551,6 +2554,7 @@ void UControlRigBlueprint::HandleModifiedEvent(ERigVMGraphNotifType InNotifType,
 			case ERigVMGraphNotifType::ParameterRemoved:
 			case ERigVMGraphNotifType::ParameterRenamed:
 			case ERigVMGraphNotifType::PinBoundVariableChanged:
+			case ERigVMGraphNotifType::VariableRemappingChanged:
 			{
 				RequestAutoVMRecompilation();
 				MarkPackageDirty();
diff --git a/Engine/Plugins/Experimental/ControlRig/Source/ControlRigDeveloper/Private/Graph/ControlRigGraphNode.cpp b/Engine/Plugins/Experimental/ControlRig/Source/ControlRigDeveloper/Private/Graph/ControlRigGraphNode.cpp
--- a/Engine/Plugins/Experimental/ControlRig/Source/ControlRigDeveloper/Private/Graph/ControlRigGraphNode.cpp
+++ b/Engine/Plugins/Experimental/ControlRig/Source/ControlRigDeveloper/Private/Graph/ControlRigGraphNode.cpp
@@ -369,6 +369,20 @@ void UControlRigGraphNode::AllocateDefaultPins()
 	CreateInputPins();
 	CreateInputOutputPins();
 	CreateOutputPins();
+
+	// Fill the variable list
+	ExternalVariables.Reset();
+	if (URigVMFunctionReferenceNode* FunctionReferenceNode = Cast<URigVMFunctionReferenceNode>(GetModelNode()))
+	{
+		if(FunctionReferenceNode->RequiresVariableRemapping())
+		{
+			TArray<FRigVMExternalVariable> CurrentExternalVariables = FunctionReferenceNode->GetContainedGraph()->GetExternalVariables();
+			for(const FRigVMExternalVariable& CurrentExternalVariable : CurrentExternalVariables)
+			{
+				ExternalVariables.Add(MakeShared<FRigVMExternalVariable>(CurrentExternalVariable));
+			}
+		}
+	}
 }
 
 void UControlRigGraphNode::CreateExecutionPins()
diff --git a/Engine/Plugins/Experimental/ControlRig/Source/ControlRigDeveloper/Public/ControlRigBlueprint.h b/Engine/Plugins/Experimental/ControlRig/Source/ControlRigDeveloper/Public/ControlRigBlueprint.h
--- a/Engine/Plugins/Experimental/ControlRig/Source/ControlRigDeveloper/Public/ControlRigBlueprint.h
+++ b/Engine/Plugins/Experimental/ControlRig/Source/ControlRigDeveloper/Public/ControlRigBlueprint.h
@@ -581,4 +581,5 @@ private:
 	friend class UAnimGraphNode_ControlRig;
 	friend class UControlRigThumbnailRenderer;
 	friend class FControlRigGraphDetails;
+	friend class FControlRigEditorModule;
 };
diff --git a/Engine/Plugins/Experimental/ControlRig/Source/ControlRigDeveloper/Public/Graph/ControlRigGraphNode.h b/Engine/Plugins/Experimental/ControlRig/Source/ControlRigDeveloper/Public/Graph/ControlRigGraphNode.h
--- a/Engine/Plugins/Experimental/ControlRig/Source/ControlRigDeveloper/Public/Graph/ControlRigGraphNode.h
+++ b/Engine/Plugins/Experimental/ControlRig/Source/ControlRigDeveloper/Public/Graph/ControlRigGraphNode.h
@@ -191,6 +191,7 @@ private:
 	TArray<URigVMPin*> InputOutputPins;
 	TArray<URigVMPin*> InputPins;
 	TArray<URigVMPin*> OutputPins;
+	TArray<TSharedPtr<FRigVMExternalVariable>> ExternalVariables;
 	
 	struct PinPair
 	{
diff --git a/Engine/Plugins/Experimental/ControlRig/Source/ControlRigEditor/Private/ControlRigEditorModule.cpp b/Engine/Plugins/Experimental/ControlRig/Source/ControlRigEditor/Private/ControlRigEditorModule.cpp
--- a/Engine/Plugins/Experimental/ControlRig/Source/ControlRigEditor/Private/ControlRigEditorModule.cpp
+++ b/Engine/Plugins/Experimental/ControlRig/Source/ControlRigEditor/Private/ControlRigEditorModule.cpp
@@ -1931,6 +1931,54 @@ void FControlRigEditorModule::GetContextMenuActions(const UControlRigGraphSchema
                                     RigBlueprint->BroadcastRequestLocalizeFunctionDialog(FunctionReferenceNode->GetReferencedNode(), true);
                                 })
                             ));
+
+							if(!FunctionReferenceNode->IsFullyRemapped())
+							{
+								FToolMenuSection& VariablesSection = Menu->AddSection("EdGraphSchemaVariables", LOCTEXT("Variables", "Variables"));
+								VariablesSection.AddMenuEntry(
+                                    "MakeVariablesFromFunctionReferenceNode",
+                                    LOCTEXT("MakeVariablesFromFunctionReferenceNode", "Create required variables"),
+                                    LOCTEXT("MakeVariablesFromFunctionReferenceNode_Tooltip", "Creates all required variables for this function and binds them"),
+                                    FSlateIcon(),
+                                    FUIAction(FExecuteAction::CreateLambda([Controller, FunctionReferenceNode, RigBlueprint]() {
+
+                                    	const TArray<FRigVMExternalVariable> ExternalVariables = FunctionReferenceNode->GetExternalVariables(false);
+                                    	if(!ExternalVariables.IsEmpty())
+                                    	{
+											FScopedTransaction Transaction(LOCTEXT("MakeVariablesFromFunctionReferenceNode", "Create required variables"));
+                                    		RigBlueprint->Modify();
+
+                                    		UControlRigBlueprint* ReferencedBlueprint = FunctionReferenceNode->GetReferencedNode()->GetTypedOuter<UControlRigBlueprint>();
+                                    		// ReferencedBlueprint != RigBlueprint - since only FunctionReferenceNodes from other assets have the potential to be unmapped
+                                    		
+                                    		for(const FRigVMExternalVariable& ExternalVariable : ExternalVariables)
+                                    		{
+                                    			FString DefaultValue;
+                                    			if(ReferencedBlueprint)
+                                    			{
+                                    				for(const FBPVariableDescription& NewVariable : ReferencedBlueprint->NewVariables)
+                                    				{
+                                    					if(NewVariable.VarName == ExternalVariable.Name)
+                                    					{
+                                    						DefaultValue = NewVariable.DefaultValue;
+                                    						break;
+                                    					}
+                                    				}
+                                    			}
+                                    			
+                                                FName NewVariableName = RigBlueprint->AddCRMemberVariableFromExternal(ExternalVariable, DefaultValue);
+                                    			if(!NewVariableName.IsNone())
+                                    			{
+                                    				Controller->SetRemappedVariable(FunctionReferenceNode, ExternalVariable.Name, NewVariableName);
+                                    			}
+                                    		}
+
+                                    		FBlueprintEditorUtils::MarkBlueprintAsModified(RigBlueprint);
+                                    	}
+                                        
+                                    })
+                                ));
+							}
 						}
 					}
 
diff --git a/Engine/Plugins/Experimental/ControlRig/Source/ControlRigEditor/Private/Graph/NodeSpawners/ControlRigFunctionRefNodeSpawner.cpp b/Engine/Plugins/Experimental/ControlRig/Source/ControlRigEditor/Private/Graph/NodeSpawners/ControlRigFunctionRefNodeSpawner.cpp
--- a/Engine/Plugins/Experimental/ControlRig/Source/ControlRigEditor/Private/Graph/NodeSpawners/ControlRigFunctionRefNodeSpawner.cpp
+++ b/Engine/Plugins/Experimental/ControlRig/Source/ControlRigEditor/Private/Graph/NodeSpawners/ControlRigFunctionRefNodeSpawner.cpp
@@ -258,6 +258,36 @@ UControlRigGraphNode* UControlRigFunctionRefNodeSpawner::SpawnNode(UEdGraph* Par
 				UControlRigUnitNodeSpawner::HookupMutableNode(ModelNode, RigBlueprint);
 			}
 
+			if (!bIsTemplateNode)
+			{
+				for(URigVMNode* OtherModelNode : ModelNode->GetGraph()->GetNodes())
+				{
+					if(OtherModelNode == ModelNode)
+					{
+						continue;
+					}
+					
+					URigVMFunctionReferenceNode* ExistingFunctionReferenceNode = Cast<URigVMFunctionReferenceNode>(OtherModelNode);
+					if(ExistingFunctionReferenceNode == nullptr)
+					{
+						continue;
+					}
+
+					if(ExistingFunctionReferenceNode->GetReferencedNode() != InFunction)
+					{
+						continue;
+					}
+
+					for(TPair<FName, FName> MappedVariablePair : ExistingFunctionReferenceNode->GetVariableMap())
+					{
+						if(!MappedVariablePair.Value.IsNone())
+						{
+							Controller->SetRemappedVariable(ModelNode, MappedVariablePair.Key, MappedVariablePair.Value, bUndo);
+						}
+					}
+				}
+			}
+
 			if (bUndo)
 			{
 				Controller->CloseUndoBracket();
diff --git a/Engine/Plugins/Experimental/ControlRig/Source/ControlRigEditor/Private/Graph/SControlRigGraphNode.cpp b/Engine/Plugins/Experimental/ControlRig/Source/ControlRigEditor/Private/Graph/SControlRigGraphNode.cpp
--- a/Engine/Plugins/Experimental/ControlRig/Source/ControlRigEditor/Private/Graph/SControlRigGraphNode.cpp
+++ b/Engine/Plugins/Experimental/ControlRig/Source/ControlRigEditor/Private/Graph/SControlRigGraphNode.cpp
@@ -28,9 +28,11 @@
 #include "IDocumentation.h"
 #include "DetailLayoutBuilder.h"
 #include "EditorStyleSet.h"
+#include "SControlRigGraphPinVariableBinding.h"
 
 #if WITH_EDITOR
 #include "Editor.h"
+#include "EdGraphSchema_K2.h"
 #endif
 
 #define LOCTEXT_NAMESPACE "SControlRigGraphNode"
@@ -82,7 +84,7 @@ void SControlRigGraphNode::Construct( const FArguments& InArgs )
 			.Visibility(this, &SControlRigGraphNode::GetExecutionTreeVisibility)
 			.TreeItemsSource(&ControlRigGraphNode->ExecutePins)
 			.SelectionMode(ESelectionMode::None)
-			.OnGenerateRow(this, &SControlRigGraphNode::MakeTableRowWidget)
+			.OnGenerateRow(this, &SControlRigGraphNode::MakePinTableRowWidget)
 			.OnGetChildren(this, &SControlRigGraphNode::HandleGetChildrenForTree)
 			.OnExpansionChanged(this, &SControlRigGraphNode::HandleExpansionChanged)
 			.OnSetExpansionRecursive(this, &SControlRigGraphNode::HandleExpandRecursively, &ExecutionTree)
@@ -97,7 +99,7 @@ void SControlRigGraphNode::Construct( const FArguments& InArgs )
 			.Visibility(this, &SControlRigGraphNode::GetInputTreeVisibility)
 			.TreeItemsSource(&ControlRigGraphNode->InputPins)
 			.SelectionMode(ESelectionMode::None)
-			.OnGenerateRow(this, &SControlRigGraphNode::MakeTableRowWidget)
+			.OnGenerateRow(this, &SControlRigGraphNode::MakePinTableRowWidget)
 			.OnGetChildren(this, &SControlRigGraphNode::HandleGetChildrenForTree)
 			.OnExpansionChanged(this, &SControlRigGraphNode::HandleExpansionChanged)
 			.OnSetExpansionRecursive(this, &SControlRigGraphNode::HandleExpandRecursively, &InputTree)
@@ -112,7 +114,7 @@ void SControlRigGraphNode::Construct( const FArguments& InArgs )
 			.Visibility(this, &SControlRigGraphNode::GetInputOutputTreeVisibility)
 			.TreeItemsSource(&ControlRigGraphNode->InputOutputPins)
 			.SelectionMode(ESelectionMode::None)
-			.OnGenerateRow(this, &SControlRigGraphNode::MakeTableRowWidget)
+			.OnGenerateRow(this, &SControlRigGraphNode::MakePinTableRowWidget)
 			.OnGetChildren(this, &SControlRigGraphNode::HandleGetChildrenForTree)
 			.OnExpansionChanged(this, &SControlRigGraphNode::HandleExpansionChanged)
 			.OnSetExpansionRecursive(this, &SControlRigGraphNode::HandleExpandRecursively, &InputOutputTree)
@@ -127,7 +129,7 @@ void SControlRigGraphNode::Construct( const FArguments& InArgs )
 			.Visibility(this, &SControlRigGraphNode::GetOutputTreeVisibility)
 			.TreeItemsSource(&ControlRigGraphNode->OutputPins)
 			.SelectionMode(ESelectionMode::None)
-			.OnGenerateRow(this, &SControlRigGraphNode::MakeTableRowWidget)
+			.OnGenerateRow(this, &SControlRigGraphNode::MakePinTableRowWidget)
 			.OnGetChildren(this, &SControlRigGraphNode::HandleGetChildrenForTree)
 			.OnExpansionChanged(this, &SControlRigGraphNode::HandleExpansionChanged)
 			.OnSetExpansionRecursive(this, &SControlRigGraphNode::HandleExpandRecursively, &OutputTree)
@@ -135,6 +137,18 @@ void SControlRigGraphNode::Construct( const FArguments& InArgs )
 			.ItemHeight(20.0f)
 		];
 
+	LeftNodeBox->AddSlot()
+        .AutoHeight()
+        [
+            SAssignNew(VariableRemappingList, SListView<TSharedPtr<FRigVMExternalVariable>>)
+            .Visibility(this, &SControlRigGraphNode::GetVariableListVisibility)
+            .ListItemsSource(&ControlRigGraphNode->ExternalVariables)
+            .SelectionMode(ESelectionMode::None)
+            .OnGenerateRow(this, &SControlRigGraphNode::MakeVariableTableRowWidget)
+            .ExternalScrollbar(ScrollBar)
+            .ItemHeight(20.0f)
+        ];
+
 	struct Local
 	{
 		static void SetItemExpansion_Recursive(UControlRigGraphNode* InControlRigGraphNode, TSharedPtr<STreeView<URigVMPin*>>& TreeWidget, const TArray<URigVMPin*>& InItems)
@@ -168,6 +182,7 @@ void SControlRigGraphNode::Construct( const FArguments& InArgs )
 	InputTree->Tick(DummyGeometry, 0.f, 0.f);
 	InputOutputTree->Tick(DummyGeometry, 0.f, 0.f);
 	OutputTree->Tick(DummyGeometry, 0.f, 0.f);
+	VariableRemappingList->Tick(DummyGeometry, 0.f, 0.f);
 
 	const FSlateBrush* ImageBrush = FControlRigEditorStyle::Get().GetBrush(TEXT("ControlRig.Bug.Dot"));
 
@@ -386,6 +401,12 @@ EVisibility SControlRigGraphNode::GetOutputTreeVisibility() const
 	return ControlRigGraphNode->OutputPins.Num() > 0 ? EVisibility::Visible : EVisibility::Collapsed;
 }
 
+EVisibility SControlRigGraphNode::GetVariableListVisibility() const
+{
+	UControlRigGraphNode* ControlRigGraphNode = CastChecked<UControlRigGraphNode>(GraphNode);
+	return ControlRigGraphNode->ExternalVariables.Num() > 0 ? EVisibility::Visible : EVisibility::Collapsed;
+}
+
 TSharedRef<SWidget> SControlRigGraphNode::CreateTitleWidget(TSharedPtr<SNodeTitle> InNodeTitle)
 {
 	NodeTitle = InNodeTitle;
@@ -611,6 +632,76 @@ public:
 	bool bLeftAligned;
 };
 
+class SControlRigVariableListRow : public STableRow<TSharedPtr<FRigVMExternalVariable>>
+{
+public:
+	SLATE_BEGIN_ARGS(SControlRigVariableListRow) {}
+	SLATE_END_ARGS()
+
+	void Construct(const FArguments& InArgs, const TSharedRef<STableViewBase>& InOwnerTableView)
+	{
+		STableRow<TSharedPtr<FRigVMExternalVariable>>::Construct(STableRow<TSharedPtr<FRigVMExternalVariable>>::FArguments(), InOwnerTableView);
+	}
+
+	const FSlateBrush* GetBorder() const
+	{
+		bool bShowBG = (CVarEnableShowBackground.GetValueOnGameThread() > 0);
+
+		if (bShowBG)
+			return STableRow<TSharedPtr<FRigVMExternalVariable>>::GetBorder();
+		else
+		// We want a transparent background.
+		return FCoreStyle::Get().GetBrush("NoBrush");
+	}
+
+	virtual void ConstructChildren( ETableViewMode::Type InOwnerTableMode, const TAttribute<FMargin>& InPadding, const TSharedRef<SWidget>& InContent ) override
+	{
+		const UGraphEditorSettings* Settings = GetDefault<UGraphEditorSettings>();
+		FMargin InputPadding = Settings->GetInputPinPadding();
+		InputPadding.Top = InputPadding.Bottom = 3.0f;
+		InputPadding.Right = 8.0f;
+
+		this->Content = InContent;
+
+		SHorizontalBox::FSlot* InnerContentSlotNativePtr = nullptr;
+
+		TSharedRef<SHorizontalBox> ContentBox = SNew(SHorizontalBox);
+
+		ContentBox->AddSlot()
+		.AutoWidth()
+		.HAlign(HAlign_Left)
+		.VAlign(VAlign_Center)
+		.Padding(InputPadding)
+		[
+			SAssignNew(LeftContentBox, SBox)
+		];
+
+		ContentBox->AddSlot()
+		.AutoWidth()
+		.HAlign(HAlign_Left)
+		.VAlign(VAlign_Center)
+		.MaxWidth(200.f)
+		.Expose(InnerContentSlotNativePtr)
+		[
+			SAssignNew(RightContentBox, SBox)
+			[
+				InContent
+			]
+		];
+
+		this->ChildSlot
+		[
+			ContentBox
+		];
+
+		InnerContentSlot = InnerContentSlotNativePtr;
+	}
+
+	/** Exposed boxes to slot pin widgets into */
+	TSharedPtr<SBox> LeftContentBox;
+	TSharedPtr<SBox> RightContentBox;
+};
+
 TSharedRef<SWidget> SControlRigGraphNode::AddContainerPinContent(URigVMPin* InItem, FText InTooltipText)
 {
 	return SNew(SButton)
@@ -633,7 +724,7 @@ TSharedRef<SWidget> SControlRigGraphNode::AddContainerPinContent(URigVMPin* InIt
 	];
 }
 
-TSharedRef<ITableRow> SControlRigGraphNode::MakeTableRowWidget(URigVMPin* InItem, const TSharedRef<STableViewBase>& OwnerTable)
+TSharedRef<ITableRow> SControlRigGraphNode::MakePinTableRowWidget(URigVMPin* InItem, const TSharedRef<STableViewBase>& OwnerTable)
 {
 	const bool bLeaf = InItem->GetSubPins().Num() == 0;
 	const bool bIsContainer = InItem->IsArray();
@@ -797,7 +888,37 @@ TSharedRef<ITableRow> SControlRigGraphNode::MakeTableRowWidget(URigVMPin* InItem
 
 	return ControlRigPinTreeRow;
 }
-	
+
+TSharedRef<ITableRow> SControlRigGraphNode::MakeVariableTableRowWidget(TSharedPtr<FRigVMExternalVariable> InVariable,
+	const TSharedRef<STableViewBase>& OwnerTable)
+{
+	TSharedRef<SControlRigVariableListRow> ControlRigVariableListRow = SNew(SControlRigVariableListRow, OwnerTable);
+
+	UControlRigGraphNode* RigNode = Cast<UControlRigGraphNode>(GraphNode);
+	URigVMFunctionReferenceNode* FunctionReferenceNode = Cast<URigVMFunctionReferenceNode>(RigNode->GetModelNode());
+	check(FunctionReferenceNode);
+	TWeakObjectPtr<URigVMFunctionReferenceNode> WeakFunctionReferenceNode = FunctionReferenceNode;
+
+	UControlRigBlueprint* ControlRigBlueprint = FunctionReferenceNode->GetReferencedNode()->GetTypedOuter<UControlRigBlueprint>(); 
+	check(ControlRigBlueprint);
+	TWeakObjectPtr<UControlRigBlueprint> WeakControlRigBlueprint = ControlRigBlueprint;
+
+	ControlRigVariableListRow->LeftContentBox->SetContent(SNew(STextBlock)
+		.Text(FText::FromName(InVariable->Name))
+        .TextStyle(FEditorStyle::Get(), NAME_DefaultPinLabelStyle)
+		.ColorAndOpacity(this, &SControlRigGraphNode::GetVariableLabelTextColor, WeakFunctionReferenceNode, InVariable->Name)
+		.ToolTipText(this, &SControlRigGraphNode::GetVariableLabelTooltipText, WeakControlRigBlueprint, InVariable->Name)
+    );
+
+	ControlRigVariableListRow->RightContentBox->SetContent(SNew(SControlRigVariableBinding)
+		.Blueprint(Blueprint.Get())
+		.FunctionReferenceNode(FunctionReferenceNode)
+		.InnerVariableName(InVariable->Name)
+	);
+
+	return ControlRigVariableListRow;
+}
+
 void SControlRigGraphNode::HandleGetChildrenForTree(URigVMPin* InItem, TArray<URigVMPin*>& OutChildren)
 {
 	OutChildren.Append(InItem->GetSubPins());
@@ -881,6 +1002,41 @@ FSlateColor SControlRigGraphNode::GetPinTextColor(TWeakPtr<SGraphPin> GraphPin)
 	return FLinearColor::White;
 }
 
+FSlateColor SControlRigGraphNode::GetVariableLabelTextColor(
+	TWeakObjectPtr<URigVMFunctionReferenceNode> FunctionReferenceNode, FName InVariableName) const
+{
+	if(FunctionReferenceNode.IsValid())
+	{
+		if(FunctionReferenceNode->GetOuterVariableName(InVariableName).IsNone())
+		{
+			return FLinearColor::Red;
+		}
+	}
+	return FLinearColor::White;
+}
+
+FText SControlRigGraphNode::GetVariableLabelTooltipText(TWeakObjectPtr<UControlRigBlueprint> InBlueprint,
+	FName InVariableName) const
+{
+	if(InBlueprint.IsValid())
+	{
+		for(const FBPVariableDescription& Variable : InBlueprint->NewVariables)
+		{
+			if(Variable.VarName == InVariableName)
+			{
+				FString Message = FString::Printf(TEXT("Variable from %s"), *InBlueprint->GetPathName()); 
+				if(Variable.HasMetaData(FBlueprintMetadata::MD_Tooltip))
+				{
+					const FString Tooltip = Variable.GetMetaData(FBlueprintMetadata::MD_Tooltip);
+					Message = FString::Printf(TEXT("%s\n%s"), *Message, *Tooltip);
+				}
+				return FText::FromString(Message);
+			}
+		}
+	}
+	return FText();
+}
+
 FReply SControlRigGraphNode::HandleAddArrayElement(URigVMPin* InItem)
 {
 	if(InItem)
diff --git a/Engine/Plugins/Experimental/ControlRig/Source/ControlRigEditor/Private/Graph/SControlRigGraphNode.h b/Engine/Plugins/Experimental/ControlRig/Source/ControlRigEditor/Private/Graph/SControlRigGraphNode.h
--- a/Engine/Plugins/Experimental/ControlRig/Source/ControlRigEditor/Private/Graph/SControlRigGraphNode.h
+++ b/Engine/Plugins/Experimental/ControlRig/Source/ControlRigEditor/Private/Graph/SControlRigGraphNode.h
@@ -6,6 +6,7 @@
 #include "Widgets/DeclarativeSyntaxSupport.h"
 #include "SGraphNode.h"
 #include "Widgets/Views/STreeView.h"
+#include "Widgets/Views/SListView.h"
 #include "Widgets/Images/SImage.h"
 #include "Widgets/Text/STextBlock.h"
 #include "RigVMModel/RigVMPin.h"
@@ -69,7 +70,10 @@ private:
 
 	EVisibility GetOutputTreeVisibility() const;
 
-	TSharedRef<ITableRow> MakeTableRowWidget(URigVMPin* InItem, const TSharedRef<STableViewBase>& OwnerTable);
+	EVisibility GetVariableListVisibility() const;
+
+	TSharedRef<ITableRow> MakePinTableRowWidget(URigVMPin* InItem, const TSharedRef<STableViewBase>& OwnerTable);
+	TSharedRef<ITableRow> MakeVariableTableRowWidget(TSharedPtr<FRigVMExternalVariable> InVariable, const TSharedRef<STableViewBase>& OwnerTable);
 
 	void HandleGetChildrenForTree(URigVMPin* InItem, TArray<URigVMPin*>& OutChildren);
 
@@ -79,6 +83,8 @@ private:
 	FText GetPinLabel(TWeakPtr<SGraphPin> GraphPin) const;
 
 	FSlateColor GetPinTextColor(TWeakPtr<SGraphPin> GraphPin) const;
+	FSlateColor GetVariableLabelTextColor(TWeakObjectPtr<URigVMFunctionReferenceNode> FunctionReferenceNode, FName InVariableName) const;
+	FText GetVariableLabelTooltipText(TWeakObjectPtr<UControlRigBlueprint> InBlueprint, FName InVariableName) const;
 
 	TSharedRef<SWidget> AddContainerPinContent(URigVMPin* InItem, FText InTooltipText);
 
@@ -104,6 +110,9 @@ private:
 	/** Widget representing collapsible output pins */
 	TSharedPtr<STreeView<URigVMPin*>> OutputTree;
 
+	/** Widget representing the variable remapping information */
+	TSharedPtr<SListView<TSharedPtr<FRigVMExternalVariable>>> VariableRemappingList;
+
 	/** Dummy scrollbar, as we cant create a tree view without one! */
 	TSharedPtr<SScrollBar> ScrollBar;
 
diff --git a/Engine/Plugins/Experimental/ControlRig/Source/ControlRigEditor/Private/Graph/SControlRigGraphPinVariableBinding.cpp b/Engine/Plugins/Experimental/ControlRig/Source/ControlRigEditor/Private/Graph/SControlRigGraphPinVariableBinding.cpp
--- a/Engine/Plugins/Experimental/ControlRig/Source/ControlRigEditor/Private/Graph/SControlRigGraphPinVariableBinding.cpp
+++ b/Engine/Plugins/Experimental/ControlRig/Source/ControlRigEditor/Private/Graph/SControlRigGraphPinVariableBinding.cpp
@@ -13,6 +13,8 @@
 void SControlRigVariableBinding::Construct(const FArguments& InArgs)
 {
 	this->ModelPin = InArgs._ModelPin;
+	this->FunctionReferenceNode = InArgs._FunctionReferenceNode;
+	this->InnerVariableName = InArgs._InnerVariableName;
 	this->Blueprint = InArgs._Blueprint;
 	this->bCanRemoveBinding = InArgs._CanRemoveBinding;
 
@@ -32,6 +34,7 @@ void SControlRigVariableBinding::Construct(const FArguments& InArgs)
 	BindingArgs.bGeneratePureBindings = true;
 	BindingArgs.bAllowNewBindings = true;
 	BindingArgs.bAllowArrayElementBindings = false;
+	BindingArgs.bAllowStructMemberBindings = false;
 	BindingArgs.bAllowUObjectFunctions = false;
 
 	this->ChildSlot
@@ -44,9 +47,17 @@ FText SControlRigVariableBinding::GetBindingText() const
 {
 	if (ModelPin)
 	{
-		FString VariablePath = ModelPin->GetBoundVariablePath();
+		const FString VariablePath = ModelPin->GetBoundVariablePath();
 		return FText::FromString(VariablePath);
 	}
+	else if(FunctionReferenceNode && !InnerVariableName.IsNone())
+	{
+		const FName BoundVariable = FunctionReferenceNode->GetOuterVariableName(InnerVariableName);
+		if(!BoundVariable.IsNone())
+		{
+			return FText::FromName(BoundVariable);
+		}
+	}
 
 	return FText();
 }
@@ -59,13 +70,28 @@ const FSlateBrush* SControlRigVariableBinding::GetBindingImage() const
 
 FLinearColor SControlRigVariableBinding::GetBindingColor() const
 {
-	if (Blueprint && ModelPin)
+	if (Blueprint)
 	{
 		const UEdGraphSchema_K2* Schema = GetDefault<UEdGraphSchema_K2>();
 
+		FName BoundVariable(NAME_None);
+
+		if(ModelPin)
+		{
+			BoundVariable = *ModelPin->GetBoundVariableName();
+		}
+		else if(FunctionReferenceNode && !InnerVariableName.IsNone())
+		{
+			BoundVariable = FunctionReferenceNode->GetOuterVariableName(InnerVariableName);
+			if(BoundVariable.IsNone())
+			{
+				return FLinearColor::Red;
+			}
+		}
+
 		for (const FBPVariableDescription& VariableDescription : Blueprint->NewVariables)
 		{
-			if (VariableDescription.VarName.ToString() == ModelPin->GetBoundVariableName())
+			if (VariableDescription.VarName == BoundVariable)
 			{
 				return Schema->GetPinTypeColor(VariableDescription.VarType);
 			}
@@ -81,10 +107,40 @@ bool SControlRigVariableBinding::OnCanBindProperty(FProperty* InProperty) const
 		return true;
 	}
 
-	if (InProperty && ModelPin)
+	if (InProperty)
 	{
-		FRigVMExternalVariable ExternalVariable = FRigVMExternalVariable::Make(InProperty, nullptr);
-		return ModelPin->CanBeBoundToVariable(ExternalVariable);
+		const FRigVMExternalVariable ExternalVariable = FRigVMExternalVariable::Make(InProperty, nullptr);
+		if(ModelPin)
+		{
+			return ModelPin->CanBeBoundToVariable(ExternalVariable);
+		}
+		else if(FunctionReferenceNode && !InnerVariableName.IsNone())
+		{
+			TArray<FRigVMExternalVariable> InnerVariables = FunctionReferenceNode->GetContainedGraph()->GetExternalVariables();
+			for(const FRigVMExternalVariable& InnerVariable : InnerVariables)
+			{
+				if(InnerVariable.Name == InnerVariableName)
+				{
+					if(!InnerVariable.bIsReadOnly && ExternalVariable.bIsReadOnly)
+					{
+						return false;
+					}
+					if(InnerVariable.bIsArray != ExternalVariable.bIsArray)
+					{
+						return false;
+					}
+					if(InnerVariable.TypeObject && InnerVariable.TypeObject != ExternalVariable.TypeObject)
+					{
+						return false;
+					}
+					else if(InnerVariable.TypeName != ExternalVariable.TypeName)
+					{
+						return false;
+					}
+					return true;
+				}
+			}
+		}
 	}
 
 	return false;
@@ -101,7 +157,7 @@ bool SControlRigVariableBinding::OnCanBindToClass(UClass* InClass) const
 
 void SControlRigVariableBinding::OnAddBinding(FName InPropertyName, const TArray<FBindingChainElement>& InBindingChain)
 {
-	if (Blueprint && ModelPin)
+	if (Blueprint)
 	{
 		TArray<FString> Parts;
 		for (const FBindingChainElement& ChainElement : InBindingChain)
@@ -109,7 +165,16 @@ void SControlRigVariableBinding::OnAddBinding(FName InPropertyName, const TArray
 			ensure(ChainElement.Field);
 			Parts.Add(ChainElement.Field.GetName());
 		}
-		Blueprint->GetController(ModelPin->GetGraph())->BindPinToVariable(ModelPin->GetPinPath(), FString::Join(Parts, TEXT(".")), true /* undo */);
+
+		if(ModelPin)
+		{
+			Blueprint->GetController(ModelPin->GetGraph())->BindPinToVariable(ModelPin->GetPinPath(), FString::Join(Parts, TEXT(".")), true /* undo */);
+		}
+		else if(FunctionReferenceNode && !InnerVariableName.IsNone())
+		{
+			const FName BoundVariableName = *FString::Join(Parts, TEXT("."));
+			Blueprint->GetController(FunctionReferenceNode->GetGraph())->SetRemappedVariable(FunctionReferenceNode, InnerVariableName, BoundVariableName);
+		}
 	}
 }
 
@@ -120,9 +185,16 @@ bool SControlRigVariableBinding::OnCanRemoveBinding(FName InPropertyName)
 
 void SControlRigVariableBinding::OnRemoveBinding(FName InPropertyName)
 {
-	if (Blueprint && ModelPin)
+	if (Blueprint)
 	{
-		Blueprint->GetController(ModelPin->GetGraph())->UnbindPinFromVariable(ModelPin->GetPinPath(), true /* undo */);
+		if(ModelPin)
+		{
+			Blueprint->GetController(ModelPin->GetGraph())->UnbindPinFromVariable(ModelPin->GetPinPath(), true /* undo */);
+		}
+		else if(FunctionReferenceNode && !InnerVariableName.IsNone())
+		{
+			Blueprint->GetController(FunctionReferenceNode->GetGraph())->SetRemappedVariable(FunctionReferenceNode, InnerVariableName, NAME_None);
+		}
 	}
 }
 
diff --git a/Engine/Plugins/Experimental/ControlRig/Source/ControlRigEditor/Private/Graph/SControlRigGraphPinVariableBinding.h b/Engine/Plugins/Experimental/ControlRig/Source/ControlRigEditor/Private/Graph/SControlRigGraphPinVariableBinding.h
--- a/Engine/Plugins/Experimental/ControlRig/Source/ControlRigEditor/Private/Graph/SControlRigGraphPinVariableBinding.h
+++ b/Engine/Plugins/Experimental/ControlRig/Source/ControlRigEditor/Private/Graph/SControlRigGraphPinVariableBinding.h
@@ -7,6 +7,7 @@
 #include "Widgets/SWidget.h"
 #include "SGraphPin.h"
 #include "RigVMModel/RigVMPin.h"
+#include "RigVMModel/Nodes/RigVMFunctionReferenceNode.h"
 #include "IPropertyAccessEditor.h"
 #include "ControlRigBlueprint.h"
 
@@ -15,10 +16,16 @@ class SControlRigVariableBinding : public SCompoundWidget
 public:
 
 	SLATE_BEGIN_ARGS(SControlRigVariableBinding)
-	: _CanRemoveBinding(true)
+	: _ModelPin(nullptr)
+    , _FunctionReferenceNode(nullptr)
+    , _InnerVariableName(NAME_None)
+    , _Blueprint(nullptr)
+	, _CanRemoveBinding(true)
 	{}
 
 		SLATE_ARGUMENT(URigVMPin*, ModelPin)
+		SLATE_ARGUMENT(URigVMFunctionReferenceNode*, FunctionReferenceNode)
+		SLATE_ARGUMENT(FName, InnerVariableName)
 		SLATE_ARGUMENT(UControlRigBlueprint*, Blueprint)
 		SLATE_ARGUMENT(bool, CanRemoveBinding)
 
@@ -38,6 +45,8 @@ protected:
 	void OnRemoveBinding(FName InPropertyName);
 
 	URigVMPin* ModelPin;
+	URigVMFunctionReferenceNode* FunctionReferenceNode;
+	FName InnerVariableName;
 	UControlRigBlueprint* Blueprint;
 	FPropertyBindingWidgetArgs BindingArgs;
 	bool bCanRemoveBinding;
diff --git a/Engine/Source/Developer/RigVMDeveloper/Private/RigVMCompiler/RigVMCompiler.cpp b/Engine/Source/Developer/RigVMDeveloper/Private/RigVMCompiler/RigVMCompiler.cpp
--- a/Engine/Source/Developer/RigVMDeveloper/Private/RigVMCompiler/RigVMCompiler.cpp
+++ b/Engine/Source/Developer/RigVMDeveloper/Private/RigVMCompiler/RigVMCompiler.cpp
@@ -127,10 +127,12 @@ bool URigVMCompiler::Compile(URigVMGraph* InGraph, URigVMController* InControlle
 #if WITH_EDITOR
 
 	// traverse all graphs and try to clear out orphan pins
+	// also check on function references with unmapped variables
 	TArray<URigVMGraph*> VisitedGraphs;
 	VisitedGraphs.Add(InGraph);
 
 	int32 NodesWithOrphanPins = 0;
+	int32 NodesWithUnmappedVariables = 0;
 	for(int32 GraphIndex=0; GraphIndex<VisitedGraphs.Num(); GraphIndex++)
 	{
 		URigVMGraph* VisitedGraph = VisitedGraphs[GraphIndex];
@@ -144,6 +146,16 @@ bool URigVMCompiler::Compile(URigVMGraph* InGraph, URigVMController* InControlle
 				NodesWithOrphanPins++;
 			}
 
+			if(URigVMFunctionReferenceNode* FunctionReferenceNode = Cast<URigVMFunctionReferenceNode>(ModelNode))
+			{
+				if(!FunctionReferenceNode->IsFullyRemapped())
+				{
+					static const FString UnmappedMessage = TEXT("Node @@ has unmapped variables. Please adjust the node and re-compile.");
+					Settings.ASTSettings.Report(EMessageSeverity::Error, ModelNode, UnmappedMessage);
+					NodesWithUnmappedVariables++;
+				}
+			}
+
 			if(URigVMLibraryNode* LibraryNode = Cast<URigVMLibraryNode>(ModelNode))
 			{
 				if(URigVMGraph* ContainedGraph = LibraryNode->GetContainedGraph())
@@ -154,7 +166,7 @@ bool URigVMCompiler::Compile(URigVMGraph* InGraph, URigVMController* InControlle
 		}
 	}
 
-	if(NodesWithOrphanPins > 0)
+	if((NodesWithOrphanPins + NodesWithUnmappedVariables) > 0)
 	{
 		return false;
 	}
@@ -1215,7 +1227,28 @@ FString URigVMCompiler::GetPinHash(const URigVMPin* InPin, const FRigVMVarExprAS
 	{
 		if (InPin->GetName() == TEXT("Value") && !bIsLiteral && !bIsDebugValue)
 		{
-			return FString::Printf(TEXT("%sVariable::%s%s"), *Prefix, *VariableNode->GetVariableName().ToString(), *Suffix);
+			FName VariableName = VariableNode->GetVariableName();
+			
+			// determine if this variable needs to be remapped
+			if(InVarExpr)
+			{
+				FRigVMASTProxy ParentProxy = InVarExpr->GetProxy();
+				while(ParentProxy.GetCallstack().Num() > 1)
+				{
+					ParentProxy = ParentProxy.GetParent();
+
+					if(URigVMFunctionReferenceNode* FunctionReferenceNode = ParentProxy.GetSubject<URigVMFunctionReferenceNode>())
+					{
+						const FName RemappedVariableName = FunctionReferenceNode->GetOuterVariableName(VariableName);
+						if(!RemappedVariableName.IsNone())
+						{
+							VariableName = RemappedVariableName;
+						}
+					}
+				}
+			}
+			
+			return FString::Printf(TEXT("%sVariable::%s%s"), *Prefix, *VariableName.ToString(), *Suffix);
 		}
 	}
 	else
diff --git a/Engine/Source/Developer/RigVMDeveloper/Private/RigVMModel/Nodes/RigVMFunctionReferenceNode.cpp b/Engine/Source/Developer/RigVMDeveloper/Private/RigVMModel/Nodes/RigVMFunctionReferenceNode.cpp
--- a/Engine/Source/Developer/RigVMDeveloper/Private/RigVMModel/Nodes/RigVMFunctionReferenceNode.cpp
+++ b/Engine/Source/Developer/RigVMDeveloper/Private/RigVMModel/Nodes/RigVMFunctionReferenceNode.cpp
@@ -57,6 +57,143 @@ URigVMGraph* URigVMFunctionReferenceNode::GetContainedGraph() const
 	return nullptr;
 }
 
+bool URigVMFunctionReferenceNode::RequiresVariableRemapping() const
+{
+	TArray<FRigVMExternalVariable> InnerVariables;
+	return RequiresVariableRemappingInternal(InnerVariables);
+}
+
+bool URigVMFunctionReferenceNode::RequiresVariableRemappingInternal(TArray<FRigVMExternalVariable>& InnerVariables) const
+{
+	bool bHostedInDifferencePackage = false;
+	if(URigVMLibraryNode* ReferencedNode = GetReferencedNode())
+	{
+		// we only need to remap variables if we are referencing
+		// a function which is hosted in another asset.
+		const UObject* ReferencedPackage = ReferencedNode->GetOutermost(); 
+		const UObject* Package = GetOutermost();
+		if(ReferencedPackage != Package)
+		{
+			bHostedInDifferencePackage = true;
+		}
+		else
+		{
+			// this case might happen within unit tests
+			// where the graphs are not parented to a package
+			const UObject* TransientPackage = GetTransientPackage();
+			if((ReferencedPackage == TransientPackage) && (Package == TransientPackage))
+			{
+				const UObject* ReferencedOuter = ReferencedNode;
+				while(ReferencedOuter)
+				{
+					const UObject* Outer = ReferencedOuter->GetOuter();
+					if(Outer != TransientPackage)
+					{
+						ReferencedOuter = Outer;
+					}
+					else
+					{
+						break;
+					}
+				}
+
+				const UObject* CurrentOuter = this;
+				while(CurrentOuter)
+				{
+					const UObject* Outer = CurrentOuter->GetOuter();
+					if(Outer != TransientPackage)
+					{
+						CurrentOuter = Outer;
+					}
+					else
+					{
+						break;
+					}
+				}
+
+				if(CurrentOuter != ReferencedOuter)
+				{
+					bHostedInDifferencePackage = true;
+				}
+			}
+		}
+	}
+
+	if(bHostedInDifferencePackage)
+	{
+		InnerVariables = Super::GetExternalVariables();
+		if(InnerVariables.Num() == 0)
+		{
+			return false;
+		}
+	}
+
+	return bHostedInDifferencePackage;
+}
+
+bool URigVMFunctionReferenceNode::IsFullyRemapped() const
+{
+	TArray<FRigVMExternalVariable> InnerVariables;
+	if(!RequiresVariableRemappingInternal(InnerVariables))
+	{
+		return true;
+	}
+
+	for(const FRigVMExternalVariable& InnerVariable : InnerVariables)
+	{
+		const FName InnerVariableName = InnerVariable.Name;
+		const FName* OuterVariableName = VariableMap.Find(InnerVariableName);
+		if(OuterVariableName == nullptr)
+		{
+			return false;
+		}
+		check(!OuterVariableName->IsNone());
+	}
+
+	return true;
+}
+
+TArray<FRigVMExternalVariable> URigVMFunctionReferenceNode::GetExternalVariables() const
+{
+	return GetExternalVariables(true);
+}
+
+TArray<FRigVMExternalVariable> URigVMFunctionReferenceNode::GetExternalVariables(bool bRemapped) const
+{
+	TArray<FRigVMExternalVariable> Variables;
+	
+	if(!bRemapped)
+	{
+		Variables = Super::GetExternalVariables();
+	}
+	else
+	{
+		if(RequiresVariableRemappingInternal(Variables))
+		{
+			for(FRigVMExternalVariable& Variable : Variables)
+			{
+				const FName* OuterVariableName = VariableMap.Find(Variable.Name);
+				if(OuterVariableName != nullptr)
+				{
+					check(!OuterVariableName->IsNone());
+					Variable.Name = *OuterVariableName;
+				}
+			}
+		}
+	}
+	
+	return Variables; 
+}
+
+FName URigVMFunctionReferenceNode::GetOuterVariableName(const FName& InInnerVariableName) const
+{
+	if(const FName* OuterVariableName = VariableMap.Find(InInnerVariableName))
+	{
+		return *OuterVariableName;
+	}
+	return NAME_None;
+}
+
 URigVMLibraryNode* URigVMFunctionReferenceNode::GetReferencedNode() const
 {
 	if (!ReferencedNodePtr.IsValid())
diff --git a/Engine/Source/Developer/RigVMDeveloper/Private/RigVMModel/Nodes/RigVMLibraryNode.cpp b/Engine/Source/Developer/RigVMDeveloper/Private/RigVMModel/Nodes/RigVMLibraryNode.cpp
--- a/Engine/Source/Developer/RigVMDeveloper/Private/RigVMModel/Nodes/RigVMLibraryNode.cpp
+++ b/Engine/Source/Developer/RigVMDeveloper/Private/RigVMModel/Nodes/RigVMLibraryNode.cpp
@@ -149,3 +149,19 @@ bool URigVMLibraryNode::Contains(URigVMLibraryNode* InContainedNode, bool bRecur
 	return false;
 }
 
+TArray<FRigVMExternalVariable> URigVMLibraryNode::GetExternalVariables() const
+{
+	TArray<FRigVMExternalVariable> Variables;
+
+	if(URigVMGraph* ContainedGraph = GetContainedGraph())
+	{
+		TArray<FRigVMExternalVariable> ContainedVariables = ContainedGraph->GetExternalVariables();
+		for(const FRigVMExternalVariable& ContainedVariable : ContainedVariables)
+		{
+			FRigVMExternalVariable::MergeExternalVariable(Variables, ContainedVariable);
+		}
+	}
+
+	return Variables;
+}
+
diff --git a/Engine/Source/Developer/RigVMDeveloper/Private/RigVMModel/RigVMController.cpp b/Engine/Source/Developer/RigVMDeveloper/Private/RigVMModel/RigVMController.cpp
--- a/Engine/Source/Developer/RigVMDeveloper/Private/RigVMModel/RigVMController.cpp
+++ b/Engine/Source/Developer/RigVMDeveloper/Private/RigVMModel/RigVMController.cpp
@@ -65,14 +65,14 @@ void URigVMController::SetGraph(URigVMGraph* InGraph)
 {
 	ensure(Graphs.Num() < 2);
 
-	URigVMGraph* PreviousGraph = GetGraph();
-	if (PreviousGraph)
+	URigVMGraph* LastGraph = GetGraph();
+	if (LastGraph)
 	{
-		if(PreviousGraph == InGraph)
+		if(LastGraph == InGraph)
 		{
 			return;
 		}
-		PreviousGraph->OnModified().RemoveAll(this);
+		LastGraph->OnModified().RemoveAll(this);
 	}
 
 	Graphs.Reset();
@@ -81,20 +81,22 @@ void URigVMController::SetGraph(URigVMGraph* InGraph)
 		PushGraph(InGraph, false);
 	}
 
-	URigVMGraph* CurrentGraph = GetGraph();
-	if (CurrentGraph)
-	{
-		CurrentGraph->OnModified().AddUObject(this, &URigVMController::HandleModifiedEvent);
-	}
-
-	HandleModifiedEvent(ERigVMGraphNotifType::GraphChanged, CurrentGraph, nullptr);
+	HandleModifiedEvent(ERigVMGraphNotifType::GraphChanged, GetGraph(), nullptr);
 }
 
 void URigVMController::PushGraph(URigVMGraph* InGraph, bool bSetupUndoRedo)
 {
+	URigVMGraph* LastGraph = GetGraph();
+	if (LastGraph)
+	{
+		LastGraph->OnModified().RemoveAll(this);
+	}
+
 	check(InGraph);
 	Graphs.Push(InGraph);
 
+	InGraph->OnModified().AddUObject(this, &URigVMController::HandleModifiedEvent);
+
 	if (bSetupUndoRedo)
 	{
 		ActionStack->AddAction(FRigVMPushGraphAction(InGraph));
@@ -104,9 +106,21 @@ void URigVMController::PushGraph(URigVMGraph* InGraph, bool bSetupUndoRedo)
 URigVMGraph* URigVMController::PopGraph(bool bSetupUndoRedo)
 {
 	ensure(Graphs.Num() > 1);
+	
 	URigVMGraph* LastGraph = GetGraph();
+	if (LastGraph)
+	{
+		LastGraph->OnModified().RemoveAll(this);
+	}
+
 	Graphs.Pop();
 
+	URigVMGraph* CurrentGraph = GetGraph();
+	if (CurrentGraph)
+	{
+		CurrentGraph->OnModified().AddUObject(this, &URigVMController::HandleModifiedEvent);
+	}
+
 	if (bSetupUndoRedo)
 	{
 		ActionStack->AddAction(FRigVMPopGraphAction(LastGraph));
@@ -192,8 +206,6 @@ void URigVMController::HandleModifiedEvent(ERigVMGraphNotifType InNotifType, URi
 		case ERigVMGraphNotifType::LinkAdded:
 		case ERigVMGraphNotifType::LinkRemoved:
 		case ERigVMGraphNotifType::PinArraySizeChanged:
-		case ERigVMGraphNotifType::VariableAdded:
-		case ERigVMGraphNotifType::VariableRemoved:
 		case ERigVMGraphNotifType::ParameterAdded:
 		case ERigVMGraphNotifType::ParameterRemoved:
 		{
@@ -223,6 +235,26 @@ void URigVMController::HandleModifiedEvent(ERigVMGraphNotifType InNotifType, URi
 			}
 			break;
 		}
+		case ERigVMGraphNotifType::VariableAdded:
+		case ERigVMGraphNotifType::VariableRemoved:
+		case ERigVMGraphNotifType::VariableRemappingChanged:
+		{
+			URigVMGraph* RootGraph = InGraph->GetRootGraph();
+			if(URigVMFunctionLibrary* FunctionLibrary = Cast<URigVMFunctionLibrary>(RootGraph->GetRootGraph()))
+			{
+				URigVMNode* Node = CastChecked<URigVMNode>(InSubject);
+				check(Node);
+
+				if(URigVMLibraryNode* Function = FunctionLibrary->FindFunctionForNode(Node))
+				{
+					FunctionLibrary->ForEachReference(Function->GetFName(), [this](URigVMFunctionReferenceNode* Reference)
+                    {
+                        FRigVMControllerGraphGuard GraphGuard(this, Reference->GetGraph(), false);
+                        Reference->GetGraph()->Notify(ERigVMGraphNotifType::VariableRemappingChanged, Reference);
+                    });
+				}
+			}
+		}
 	}
 
 	ModifiedEventStatic.Broadcast(InNotifType, InGraph, InSubject);
@@ -574,12 +606,7 @@ void URigVMController::OnExternalVariableRemoved(const FName& InVarName, bool bS
 	URigVMGraph* Graph = GetGraph();
 	check(Graph);
 
-	if (Graph->IsA<URigVMFunctionLibrary>())
-	{
-		return;
-	}
-
-	FString VarNameStr = InVarName.ToString();
+	const FString VarNameStr = InVarName.ToString();
 
 	if (bSetupUndoRedo)
 	{
@@ -600,13 +627,47 @@ void URigVMController::OnExternalVariableRemoved(const FName& InVarName, bool bS
 				}
 			}
 		}
+		else if(URigVMCollapseNode* CollapseNode = Cast<URigVMCollapseNode>(Node))
+		{
+			FRigVMControllerGraphGuard GraphGuard(this, CollapseNode->GetContainedGraph(), bSetupUndoRedo);
+
+			// call this function for the contained graph recursively 
+			OnExternalVariableRemoved(InVarName, bSetupUndoRedo);
+
+			// if we are a function we need to notify all references!
+			if(URigVMFunctionLibrary* FunctionLibrary = Cast<URigVMFunctionLibrary>(Graph))
+			{
+				FunctionLibrary->ForEachReference(CollapseNode->GetFName(), [this, InVarName](URigVMFunctionReferenceNode* Reference)
+				{
+					if(Reference->VariableMap.Contains(InVarName))
+					{
+						Reference->Modify();
+                        Reference->VariableMap.Remove(InVarName);
+
+                        FRigVMControllerGraphGuard GraphGuard(this, Reference->GetGraph(), false);
+                        Notify(ERigVMGraphNotifType::VariableRemappingChanged, Reference);
+					}
+				});
+			}
+		}
+		else if(URigVMFunctionReferenceNode* FunctionReferenceNode = Cast<URigVMFunctionReferenceNode>(Node))
+		{
+			TMap<FName, FName> VariableMap = FunctionReferenceNode->GetVariableMap();
+			for(const TPair<FName, FName>& VariablePair : VariableMap)
+			{
+				if(VariablePair.Value == InVarName)
+				{
+					SetRemappedVariable(FunctionReferenceNode, VariablePair.Key, NAME_None, bSetupUndoRedo);
+				}
+			}
+		}
 
 		TArray<URigVMPin*> AllPins = Node->GetAllPinsRecursively();
 		for (URigVMPin* Pin : AllPins)
 		{
 			if (Pin->GetBoundVariableName() == InVarName.ToString())
 			{
-				BindPinToVariable(Pin, FString(), true);
+				BindPinToVariable(Pin, FString(), bSetupUndoRedo);
 			}
 		}
 	}
@@ -632,12 +693,7 @@ void URigVMController::OnExternalVariableRenamed(const FName& InOldVarName, cons
 	URigVMGraph* Graph = GetGraph();
 	check(Graph);
 
-	if (Graph->IsA<URigVMFunctionLibrary>())
-	{
-		return;
-	}
-
-	FString VarNameStr = InOldVarName.ToString();
+	const FString VarNameStr = InOldVarName.ToString();
 
 	if (bSetupUndoRedo)
 	{
@@ -658,6 +714,41 @@ void URigVMController::OnExternalVariableRenamed(const FName& InOldVarName, cons
 				}
 			}
 		}
+		else if(URigVMCollapseNode* CollapseNode = Cast<URigVMCollapseNode>(Node))
+		{
+			FRigVMControllerGraphGuard GraphGuard(this, CollapseNode->GetContainedGraph(), bSetupUndoRedo);
+			OnExternalVariableRenamed(InOldVarName, InNewVarName, bSetupUndoRedo);
+
+			// if we are a function we need to notify all references!
+			if(URigVMFunctionLibrary* FunctionLibrary = Cast<URigVMFunctionLibrary>(Graph))
+			{
+				FunctionLibrary->ForEachReference(CollapseNode->GetFName(), [this, InOldVarName, InNewVarName](URigVMFunctionReferenceNode* Reference)
+                {
+					if(Reference->VariableMap.Contains(InOldVarName))
+					{
+						Reference->Modify();
+
+						FName MappedVariable = Reference->VariableMap.FindChecked(InOldVarName);
+						Reference->VariableMap.Remove(InOldVarName);
+						Reference->VariableMap.FindOrAdd(InNewVarName) = MappedVariable; 
+
+						FRigVMControllerGraphGuard GraphGuard(this, Reference->GetGraph(), false);
+                        Notify(ERigVMGraphNotifType::VariableRemappingChanged, Reference);
+                    }
+                });
+			}
+		}
+		else if(URigVMFunctionReferenceNode* FunctionReferenceNode = Cast<URigVMFunctionReferenceNode>(Node))
+		{
+			TMap<FName, FName> VariableMap = FunctionReferenceNode->GetVariableMap();
+			for(const TPair<FName, FName>& VariablePair : VariableMap)
+			{
+				if(VariablePair.Value == InOldVarName)
+				{
+					SetRemappedVariable(FunctionReferenceNode, VariablePair.Key, InNewVarName, bSetupUndoRedo);
+				}
+			}
+		}
 
 		TArray<URigVMPin*> AllPins = Node->GetAllPinsRecursively();
 		for (URigVMPin* Pin : AllPins)
@@ -666,7 +757,7 @@ void URigVMController::OnExternalVariableRenamed(const FName& InOldVarName, cons
 			{
 				FString OldVariablePath = Pin->GetBoundVariablePath();
 				FString NewVariablePath = OldVariablePath.Replace(*InOldVarName.ToString(), *InNewVarName.ToString());
-				BindPinToVariable(Pin, NewVariablePath, true);
+				BindPinToVariable(Pin, NewVariablePath, bSetupUndoRedo);
 			}
 		}
 	}
@@ -692,12 +783,7 @@ void URigVMController::OnExternalVariableTypeChanged(const FName& InVarName, con
 	URigVMGraph* Graph = GetGraph();
 	check(Graph);
 
-	if (Graph->IsA<URigVMFunctionLibrary>())
-	{
-		return;
-	}
-
-	FString VarNameStr = InVarName.ToString();
+	const FString VarNameStr = InVarName.ToString();
 
 	if (bSetupUndoRedo)
 	{
@@ -718,6 +804,38 @@ void URigVMController::OnExternalVariableTypeChanged(const FName& InVarName, con
 				}
 			}
 		}
+		else if(URigVMCollapseNode* CollapseNode = Cast<URigVMCollapseNode>(Node))
+		{
+			FRigVMControllerGraphGuard GraphGuard(this, CollapseNode->GetContainedGraph(), bSetupUndoRedo);
+			OnExternalVariableTypeChanged(InVarName, InCPPType, InCPPTypeObject, bSetupUndoRedo);
+
+			// if we are a function we need to notify all references!
+			if(URigVMFunctionLibrary* FunctionLibrary = Cast<URigVMFunctionLibrary>(Graph))
+			{
+				FunctionLibrary->ForEachReference(CollapseNode->GetFName(), [this, InVarName](URigVMFunctionReferenceNode* Reference)
+                {
+                    if(Reference->VariableMap.Contains(InVarName))
+                    {
+                        Reference->Modify();
+                        Reference->VariableMap.Remove(InVarName); 
+
+                        FRigVMControllerGraphGuard GraphGuard(this, Reference->GetGraph(), false);
+                        Notify(ERigVMGraphNotifType::VariableRemappingChanged, Reference);
+                    }
+                });
+			}
+		}
+		else if(URigVMFunctionReferenceNode* FunctionReferenceNode = Cast<URigVMFunctionReferenceNode>(Node))
+		{
+			TMap<FName, FName> VariableMap = FunctionReferenceNode->GetVariableMap();
+			for(const TPair<FName, FName>& VariablePair : VariableMap)
+			{
+				if(VariablePair.Value == InVarName)
+				{
+					SetRemappedVariable(FunctionReferenceNode, VariablePair.Key, NAME_None, bSetupUndoRedo);
+				}
+			}
+		}
 
 		TArray<URigVMPin*> AllPins = Node->GetAllPinsRecursively();
 		for (URigVMPin* Pin : AllPins)
@@ -725,9 +843,9 @@ void URigVMController::OnExternalVariableTypeChanged(const FName& InVarName, con
 			if (Pin->GetBoundVariableName() == InVarName.ToString())
 			{
 				FString BoundVariablePath = Pin->GetBoundVariablePath();
-				BindPinToVariable(Pin, FString(), true);
+				BindPinToVariable(Pin, FString(), bSetupUndoRedo);
 				// try to bind it again - maybe it can be bound (due to cast rules etc)
-				BindPinToVariable(Pin, BoundVariablePath, true);
+				BindPinToVariable(Pin, BoundVariablePath, bSetupUndoRedo);
 			}
 		}
 	}
@@ -3528,13 +3646,10 @@ bool URigVMController::RemoveNode(URigVMNode* InNode, bool bSetupUndoRedo, bool
 					References->FunctionReferences.RemoveAll(
 						[FunctionReferenceNode](TSoftObjectPtr<URigVMFunctionReferenceNode>& FunctionReferencePtr) {
 
-							/*
 							if (!FunctionReferencePtr.IsValid())
 							{
 								FunctionReferencePtr.LoadSynchronous();
 							}
-							*/
-
 							if (!FunctionReferencePtr.IsValid())
 							{
 								return true;
@@ -3663,19 +3778,11 @@ bool URigVMController::RenameNode(URigVMNode* InNode, const FName& InNewName, bo
 	{
 		if (URigVMFunctionLibrary* FunctionLibrary = Cast<URigVMFunctionLibrary>(LibraryNode->GetGraph()))
 		{
-			FRigVMFunctionReferenceArray* ReferencesEntry = FunctionLibrary->FunctionReferences.Find(LibraryNode);
-			if (ReferencesEntry)
+			FunctionLibrary->ForEachReference(LibraryNode->GetFName(), [this, InNewName](URigVMFunctionReferenceNode* ReferenceNode)
 			{
-				for (TSoftObjectPtr<URigVMFunctionReferenceNode> FunctionReferencePtr : ReferencesEntry->FunctionReferences)
-				{
-					// only update valid, living references
-					if (FunctionReferencePtr.IsValid())
-					{
-						FRigVMControllerGraphGuard GraphGuard(this, FunctionReferencePtr->GetGraph(), false);
-						RenameNode(FunctionReferencePtr.Get(), InNewName, false);
-					}
-				}
-			}
+				FRigVMControllerGraphGuard GraphGuard(this, ReferenceNode->GetGraph(), false);
+                RenameNode(ReferenceNode, InNewName, false);
+			});
 		}
 	}
 
@@ -3937,19 +4044,11 @@ bool URigVMController::SetNodeColor(URigVMNode* InNode, const FLinearColor& InCo
 	{
 		if (URigVMFunctionLibrary* FunctionLibrary = Cast<URigVMFunctionLibrary>(LibraryNode->GetGraph()))
 		{
-			FRigVMFunctionReferenceArray* ReferencesEntry = FunctionLibrary->FunctionReferences.Find(LibraryNode);
-			if (ReferencesEntry)
-			{
-				for (TSoftObjectPtr<URigVMFunctionReferenceNode> FunctionReferencePtr : ReferencesEntry->FunctionReferences)
-				{
-					if (FunctionReferencePtr.IsValid())
-					{
-						URigVMNode* ReferenceNode = FunctionReferencePtr.Get();
-						FRigVMControllerGraphGuard GraphGuard(this, ReferenceNode->GetGraph(), false);
-						Notify(ERigVMGraphNotifType::NodeColorChanged, ReferenceNode);
-					}
-				}
-			}
+			FunctionLibrary->ForEachReference(LibraryNode->GetFName(), [this](URigVMFunctionReferenceNode* ReferenceNode)
+            {
+                FRigVMControllerGraphGuard GraphGuard(this, ReferenceNode->GetGraph(), false);
+				Notify(ERigVMGraphNotifType::NodeColorChanged, ReferenceNode);
+            });
 		}
 	}
 
@@ -5987,22 +6086,14 @@ bool URigVMController::RenameExposedPin(const FName& InOldPinName, const FName&
 
 	if (URigVMFunctionLibrary* FunctionLibrary = Cast<URigVMFunctionLibrary>(LibraryNode->GetGraph()))
 	{
-		FRigVMFunctionReferenceArray* ReferencesEntry = FunctionLibrary->FunctionReferences.Find(LibraryNode);
-		if (ReferencesEntry)
-		{
-			for (TSoftObjectPtr<URigVMFunctionReferenceNode> FunctionReferencePtr : ReferencesEntry->FunctionReferences)
+		FunctionLibrary->ForEachReference(LibraryNode->GetFName(), [this, InOldPinName, PinName](URigVMFunctionReferenceNode* ReferenceNode)
+        {
+			if (URigVMPin* EntryPin = ReferenceNode->FindPin(InOldPinName.ToString()))
 			{
-				// only update valid, living references
-				if (FunctionReferencePtr.IsValid())
-				{
-					if (URigVMPin* EntryPin = FunctionReferencePtr->FindPin(InOldPinName.ToString()))
-					{
-						FRigVMControllerGraphGuard GraphGuard(this, FunctionReferencePtr->GetGraph(), false);
-						Local::RenamePin(this, EntryPin, PinName);
-					}
-				}
-			}
-		}
+                FRigVMControllerGraphGuard GraphGuard(this, ReferenceNode->GetGraph(), false);
+                Local::RenamePin(this, EntryPin, PinName);
+            }
+        });
 	}
 
 	if (bSetupUndoRedo)
@@ -6081,22 +6172,14 @@ bool URigVMController::ChangeExposedPinType(const FName& InPinName, const FStrin
 
 	if (URigVMFunctionLibrary* FunctionLibrary = Cast<URigVMFunctionLibrary>(LibraryNode->GetGraph()))
 	{
-		FRigVMFunctionReferenceArray* ReferencesEntry = FunctionLibrary->FunctionReferences.Find(LibraryNode);
-		if (ReferencesEntry)
-		{
-			for (TSoftObjectPtr<URigVMFunctionReferenceNode> FunctionReferencePtr : ReferencesEntry->FunctionReferences)
+		FunctionLibrary->ForEachReference(LibraryNode->GetFName(), [this, &Pin, InCPPType, InCPPTypeObjectPath, bSetupUndoRedo, bSetupOrphanPins](URigVMFunctionReferenceNode* ReferenceNode)
+        {
+			if (URigVMPin* ReferencedNodePin = ReferenceNode->FindPin(Pin->GetName()))
 			{
-				// only update valid, living references
-				if (FunctionReferencePtr.IsValid())
-				{
-					if (URigVMPin* ReferencedNodePin = FunctionReferencePtr->FindPin(Pin->GetName()))
-					{
-						FRigVMControllerGraphGuard GraphGuard(this, FunctionReferencePtr->GetGraph(), false);
-						ChangePinType(ReferencedNodePin, InCPPType, InCPPTypeObjectPath, bSetupUndoRedo, bSetupOrphanPins);
-					}
-				}
+				FRigVMControllerGraphGuard GraphGuard(this, ReferenceNode->GetGraph(), false);
+				ChangePinType(ReferencedNodePin, InCPPType, InCPPTypeObjectPath, bSetupUndoRedo, bSetupOrphanPins);
 			}
-		}
+        });
 	}
 
 	if (bSetupUndoRedo)
@@ -6244,6 +6327,93 @@ URigVMFunctionReferenceNode* URigVMController::AddFunctionReferenceNode(URigVMLi
 	return FunctionRefNode;
 }
 
+bool URigVMController::SetRemappedVariable(URigVMFunctionReferenceNode* InFunctionRefNode,
+	const FName& InInnerVariableName, const FName& InOuterVariableName, bool bSetupUndoRedo)
+{
+	if(!InFunctionRefNode)
+	{
+		return false;
+	}
+
+	if (!IsValidGraph())
+	{
+		return false;
+	}
+
+	if(InInnerVariableName.IsNone())
+	{
+		return false;
+	}
+
+	const FName OldOuterVariableName = InFunctionRefNode->GetOuterVariableName(InInnerVariableName);
+	if(OldOuterVariableName == InOuterVariableName)
+	{
+		return false;
+	}
+
+	if(!InFunctionRefNode->RequiresVariableRemapping())
+	{
+		return false;
+	}
+	
+	URigVMGraph* Graph = GetGraph();
+	check(Graph);
+
+	FRigVMExternalVariable InnerExternalVariable;
+	{
+		FRigVMControllerGraphGuard GraphGuard(this, InFunctionRefNode->GetContainedGraph());
+		InnerExternalVariable = GetExternalVariableByName(InInnerVariableName);
+	}
+
+	if(!InnerExternalVariable.IsValid(true))
+	{
+		ReportErrorf(TEXT("External variable '%s' cannot be found."), *InInnerVariableName.ToString());
+		return false;
+	}
+
+	ensure(InnerExternalVariable.Name == InInnerVariableName);
+
+	if(InOuterVariableName.IsNone())
+	{
+		InFunctionRefNode->Modify();
+		InFunctionRefNode->VariableMap.Remove(InInnerVariableName);
+	}
+	else
+	{
+		const FRigVMExternalVariable OuterExternalVariable = GetExternalVariableByName(InOuterVariableName);
+		if(!OuterExternalVariable.IsValid(true))
+		{
+			ReportErrorf(TEXT("External variable '%s' cannot be found."), *InOuterVariableName.ToString());
+			return false;
+		}
+
+		ensure(OuterExternalVariable.Name == InOuterVariableName);
+
+		if((InnerExternalVariable.TypeObject != nullptr) && (InnerExternalVariable.TypeObject != OuterExternalVariable.TypeObject))
+		{
+			ReportErrorf(TEXT("Inner and Outer External variables '%s' and '%s' are not compatible."), *InInnerVariableName.ToString(), *InOuterVariableName.ToString());
+			return false;
+		}
+		if((InnerExternalVariable.TypeObject == nullptr) && (InnerExternalVariable.TypeName != OuterExternalVariable.TypeName))
+		{
+			ReportErrorf(TEXT("Inner and Outer External variables '%s' and '%s' are not compatible."), *InInnerVariableName.ToString(), *InOuterVariableName.ToString());
+			return false;
+		}
+
+		InFunctionRefNode->Modify();
+		InFunctionRefNode->VariableMap.FindOrAdd(InInnerVariableName) = InOuterVariableName;
+	}
+
+	Notify(ERigVMGraphNotifType::VariableRemappingChanged, InFunctionRefNode);
+
+	if(bSetupUndoRedo)
+	{
+		ActionStack->AddAction(FRigVMSetRemappedVariableAction(InFunctionRefNode, InInnerVariableName, OldOuterVariableName, InOuterVariableName));
+	}
+	
+	return true;
+}
+
 URigVMLibraryNode* URigVMController::AddFunctionToLibrary(const FName& InFunctionName, bool bMutable, const FVector2D& InNodePosition, bool bSetupUndoRedo)
 {
 	if (!IsValidGraph())
@@ -8894,7 +9064,7 @@ TArray<FRigVMExternalVariable> URigVMController::GetExternalVariables()
 {
 	if (GetExternalVariablesDelegate.IsBound())
 	{
-		return GetExternalVariablesDelegate.Execute();
+		return GetExternalVariablesDelegate.Execute(GetGraph());
 	}
 	return TArray<FRigVMExternalVariable>();
 }
@@ -8914,23 +9084,15 @@ void URigVMController::RefreshFunctionReferences(URigVMLibraryNode* InFunctionDe
 
 	if (URigVMFunctionLibrary* FunctionLibrary = Cast<URigVMFunctionLibrary>(InFunctionDefinition->GetGraph()))
 	{
-		FRigVMFunctionReferenceArray* ReferencesEntry = FunctionLibrary->FunctionReferences.Find(InFunctionDefinition);
-		if (ReferencesEntry)
+		FunctionLibrary->ForEachReference(InFunctionDefinition->GetFName(), [this, bSetupUndoRedo](URigVMFunctionReferenceNode* ReferenceNode)
 		{
-			for (TSoftObjectPtr<URigVMFunctionReferenceNode> FunctionReferencePtr : ReferencesEntry->FunctionReferences)
-			{
-				// only update valid, living references
-				if (FunctionReferencePtr.IsValid())
-				{
-					FRigVMControllerGraphGuard GraphGuard(this, FunctionReferencePtr->GetGraph(), bSetupUndoRedo);
+			FRigVMControllerGraphGuard GraphGuard(this, ReferenceNode->GetGraph(), bSetupUndoRedo);
 
-					TArray<URigVMLink*> Links = FunctionReferencePtr->GetLinks();
-					DetachLinksFromPinObjects(&Links, true);
-					RepopulatePinsOnNode(FunctionReferencePtr.Get(), false, true);
-					TGuardValue<bool> ReportGuard(bReportWarningsAndErrors, false);
-					ReattachLinksToPinObjects(false, &Links, true);
-				}
-			}
-		}
+			TArray<URigVMLink*> Links = ReferenceNode->GetLinks();
+			DetachLinksFromPinObjects(&Links, true);
+			RepopulatePinsOnNode(ReferenceNode, false, true);
+			TGuardValue<bool> ReportGuard(bReportWarningsAndErrors, false);
+			ReattachLinksToPinObjects(false, &Links, true);
+		});
 	}
 }
diff --git a/Engine/Source/Developer/RigVMDeveloper/Private/RigVMModel/RigVMControllerActions.cpp b/Engine/Source/Developer/RigVMDeveloper/Private/RigVMModel/RigVMControllerActions.cpp
--- a/Engine/Source/Developer/RigVMDeveloper/Private/RigVMModel/RigVMControllerActions.cpp
+++ b/Engine/Source/Developer/RigVMDeveloper/Private/RigVMModel/RigVMControllerActions.cpp
@@ -1689,3 +1689,42 @@ bool FRigVMSetPinIndexAction::Redo(URigVMController* InController)
 	}
 	return FRigVMBaseAction::Redo(InController);
 }
+
+FRigVMSetRemappedVariableAction::FRigVMSetRemappedVariableAction(URigVMFunctionReferenceNode* InFunctionRefNode,
+	const FName& InInnerVariableName, const FName& InOldOuterVariableName, const FName& InNewOuterVariableName)
+	: NodePath()
+	, InnerVariableName(InInnerVariableName)
+	, OldOuterVariableName(InOldOuterVariableName)
+	, NewOuterVariableName(InNewOuterVariableName)
+{
+	if(InFunctionRefNode)
+	{
+		NodePath = InFunctionRefNode->GetName();
+	}
+}
+
+bool FRigVMSetRemappedVariableAction::Undo(URigVMController* InController)
+{
+	if (!FRigVMBaseAction::Undo(InController))
+	{
+		return false;
+	}
+	if (URigVMFunctionReferenceNode* Node = Cast<URigVMFunctionReferenceNode>(InController->GetGraph()->FindNode(NodePath)))
+	{
+		return InController->SetRemappedVariable(Node, InnerVariableName, OldOuterVariableName, false);
+	}
+	return false;
+}
+
+bool FRigVMSetRemappedVariableAction::Redo(URigVMController* InController)
+{
+	if (URigVMFunctionReferenceNode* Node = Cast<URigVMFunctionReferenceNode>(InController->GetGraph()->FindNode(NodePath)))
+	{
+		return InController->SetRemappedVariable(Node, InnerVariableName, NewOuterVariableName, false);
+	}
+	else
+	{
+		return false;
+	}
+	return FRigVMBaseAction::Redo(InController);
+}
diff --git a/Engine/Source/Developer/RigVMDeveloper/Private/RigVMModel/RigVMFunctionLibrary.cpp b/Engine/Source/Developer/RigVMDeveloper/Private/RigVMModel/RigVMFunctionLibrary.cpp
--- a/Engine/Source/Developer/RigVMDeveloper/Private/RigVMModel/RigVMFunctionLibrary.cpp
+++ b/Engine/Source/Developer/RigVMDeveloper/Private/RigVMModel/RigVMFunctionLibrary.cpp
@@ -46,44 +46,85 @@ URigVMLibraryNode* URigVMFunctionLibrary::FindFunction(const FName& InFunctionNa
 	return Cast<URigVMLibraryNode>(FindNodeByName(*FunctionNameStr));
 }
 
+URigVMLibraryNode* URigVMFunctionLibrary::FindFunctionForNode(URigVMNode* InNode) const
+{
+	UObject* Subject = InNode;
+	do
+	{
+		if(Subject == nullptr)
+		{
+			return nullptr;
+		}
+		Subject = Subject->GetOuter();
+	}
+	while (Subject->GetOuter() != this);
+
+	return Cast<URigVMLibraryNode>(Subject);
+}
+
 TArray< TSoftObjectPtr<URigVMFunctionReferenceNode> > URigVMFunctionLibrary::GetReferencesForFunction(const FName& InFunctionName)
 {
 	TArray< TSoftObjectPtr<URigVMFunctionReferenceNode> > Result;
 
+	ForEachReferenceSoftPtr(InFunctionName, [&Result](TSoftObjectPtr<URigVMFunctionReferenceNode> Reference)
+	{
+		Result.Add(TSoftObjectPtr<URigVMFunctionReferenceNode>(Reference.GetUniqueID()));
+	});
+
+	return Result;
+}
+
+TArray< FString > URigVMFunctionLibrary::GetReferencePathsForFunction(const FName& InFunctionName)
+{
+	TArray< FString > Result;
+
+	ForEachReferenceSoftPtr(InFunctionName, [&Result](TSoftObjectPtr<URigVMFunctionReferenceNode> Reference)
+	{
+		Result.Add(Reference.ToString());
+	});
+
+	return Result;
+}
+
+void URigVMFunctionLibrary::ForEachReference(const FName& InFunctionName,
+	TFunction<void(URigVMFunctionReferenceNode*)> PerReferenceFunction) const
+{
 	if (URigVMLibraryNode* Function = FindFunction(InFunctionName))
 	{
-		FRigVMFunctionReferenceArray* ReferencesEntry = FunctionReferences.Find(Function);
+		const FRigVMFunctionReferenceArray* ReferencesEntry = FunctionReferences.Find(Function);
 		if (ReferencesEntry)
 		{
 			for (int32 ReferenceIndex = 0; ReferenceIndex < ReferencesEntry->Num(); ReferenceIndex++)
 			{
 				const TSoftObjectPtr<URigVMFunctionReferenceNode>& Reference = ReferencesEntry->operator [](ReferenceIndex);
-				Result.Add(TSoftObjectPtr<URigVMFunctionReferenceNode>(Reference.GetUniqueID()));
+				if (!Reference.IsValid())
+				{
+					Reference.LoadSynchronous();
+				}
+				if (Reference.IsValid())
+				{
+					PerReferenceFunction(Reference.Get());
+				}
 			}
 		}
 	}
-
-	return Result;
 }
 
-TArray< FString > URigVMFunctionLibrary::GetReferencePathsForFunction(const FName& InFunctionName)
+void URigVMFunctionLibrary::ForEachReferenceSoftPtr(const FName& InFunctionName,
+	TFunction<void(TSoftObjectPtr<URigVMFunctionReferenceNode>)> PerReferenceFunction) const
 {
-	TArray< FString > Result;
-
 	if (URigVMLibraryNode* Function = FindFunction(InFunctionName))
 	{
-		FRigVMFunctionReferenceArray* ReferencesEntry = FunctionReferences.Find(Function);
+		const FRigVMFunctionReferenceArray* ReferencesEntry = FunctionReferences.Find(Function);
 		if (ReferencesEntry)
 		{
 			for (int32 ReferenceIndex = 0; ReferenceIndex < ReferencesEntry->Num(); ReferenceIndex++)
 			{
 				const TSoftObjectPtr<URigVMFunctionReferenceNode>& Reference = ReferencesEntry->operator [](ReferenceIndex);
-				Result.Add(Reference.ToString());
+				PerReferenceFunction(Reference);
 			}
 		}
 	}
-
-	return Result;
 }
 
 void URigVMFunctionLibrary::UpdateReferencesForReferenceNode(URigVMFunctionReferenceNode* InReferenceNode)
diff --git a/Engine/Source/Developer/RigVMDeveloper/Private/RigVMModel/RigVMGraph.cpp b/Engine/Source/Developer/RigVMDeveloper/Private/RigVMModel/RigVMGraph.cpp
--- a/Engine/Source/Developer/RigVMDeveloper/Private/RigVMModel/RigVMGraph.cpp
+++ b/Engine/Source/Developer/RigVMDeveloper/Private/RigVMModel/RigVMGraph.cpp
@@ -59,6 +59,11 @@ URigVMGraph* URigVMGraph::GetRootGraph() const
 	return (URigVMGraph*)this;
 }
 
+bool URigVMGraph::IsRootGraph() const
+{
+	return GetRootGraph() == this;
+}
+
 URigVMFunctionEntryNode* URigVMGraph::GetEntryNode() const
 {
 	for (URigVMNode* Node : Nodes)
@@ -237,6 +242,29 @@ void URigVMGraph::SetDefaultFunctionLibrary(URigVMFunctionLibrary* InFunctionLib
 	DefaultFunctionLibraryPtr = InFunctionLibrary;
 }
 
+TArray<FRigVMExternalVariable> URigVMGraph::GetExternalVariables() const
+{
+	TArray<FRigVMExternalVariable> Variables;
+	
+	for(URigVMNode* Node : GetNodes())
+	{
+		if(URigVMLibraryNode* LibraryNode = Cast<URigVMLibraryNode>(Node))
+		{
+			TArray<FRigVMExternalVariable> LibraryVariables = LibraryNode->GetExternalVariables();
+			for(const FRigVMExternalVariable& LibraryVariable : LibraryVariables)
+			{
+				FRigVMExternalVariable::MergeExternalVariable(Variables, LibraryVariable);
+			}
+		}
+		else if(URigVMVariableNode* VariableNode = Cast<URigVMVariableNode>(Node))
+		{
+			FRigVMExternalVariable::MergeExternalVariable(Variables, VariableNode->GetVariableDescription().ToExternalVariable());
+		}
+	}
+	
+	return Variables;
+}
+
 FRigVMGraphModifiedEvent& URigVMGraph::OnModified()
 {
 	return ModifiedEvent;
diff --git a/Engine/Source/Developer/RigVMDeveloper/Public/RigVMModel/Nodes/RigVMFunctionReferenceNode.h b/Engine/Source/Developer/RigVMDeveloper/Public/RigVMModel/Nodes/RigVMFunctionReferenceNode.h
--- a/Engine/Source/Developer/RigVMDeveloper/Public/RigVMModel/Nodes/RigVMFunctionReferenceNode.h
+++ b/Engine/Source/Developer/RigVMDeveloper/Public/RigVMModel/Nodes/RigVMFunctionReferenceNode.h
@@ -29,17 +29,31 @@ public:
 	virtual FString GetNodeKeywords() const override;
 	virtual URigVMFunctionLibrary* GetLibrary() const override;
 	virtual URigVMGraph* GetContainedGraph() const override;
+	virtual TArray<FRigVMExternalVariable> GetExternalVariables() const override;
 	// end URigVMLibraryNode interface
 
 	URigVMLibraryNode* GetReferencedNode() const;
 
+	// Variable remapping
+	bool RequiresVariableRemapping() const;
+	bool IsFullyRemapped() const;
+	TArray<FRigVMExternalVariable> GetExternalVariables(bool bRemapped) const;
+	const TMap<FName, FName>& GetVariableMap() const { return VariableMap; }
+	FName GetOuterVariableName(const FName& InInnerVariableName) const;
+	// end Variable remapping
+
 private:
 
+	bool RequiresVariableRemappingInternal(TArray<FRigVMExternalVariable>& InnerVariables) const;
+
 	void SetReferencedNode(URigVMLibraryNode* InReferenceNode);
 
 	UPROPERTY(AssetRegistrySearchable)
 	mutable TSoftObjectPtr<URigVMLibraryNode> ReferencedNodePtr;
 
+	UPROPERTY()
+	TMap<FName, FName> VariableMap;
+
 	friend class URigVMController;
 	friend class FRigVMParserAST;
 };
diff --git a/Engine/Source/Developer/RigVMDeveloper/Public/RigVMModel/Nodes/RigVMLibraryNode.h b/Engine/Source/Developer/RigVMDeveloper/Public/RigVMModel/Nodes/RigVMLibraryNode.h
--- a/Engine/Source/Developer/RigVMDeveloper/Public/RigVMModel/Nodes/RigVMLibraryNode.h
+++ b/Engine/Source/Developer/RigVMDeveloper/Public/RigVMModel/Nodes/RigVMLibraryNode.h
@@ -39,6 +39,7 @@ public:
 	virtual URigVMFunctionEntryNode* GetEntryNode() const;
 	virtual URigVMFunctionReturnNode* GetReturnNode() const;
 	virtual bool Contains(URigVMLibraryNode* InContainedNode, bool bRecursive = true) const;
+	virtual TArray<FRigVMExternalVariable> GetExternalVariables() const;
 
 protected:
 
diff --git a/Engine/Source/Developer/RigVMDeveloper/Public/RigVMModel/RigVMController.h b/Engine/Source/Developer/RigVMDeveloper/Public/RigVMModel/RigVMController.h
--- a/Engine/Source/Developer/RigVMDeveloper/Public/RigVMModel/RigVMController.h
+++ b/Engine/Source/Developer/RigVMDeveloper/Public/RigVMModel/RigVMController.h
@@ -21,7 +21,7 @@
 class URigVMActionStack;
 
 DECLARE_DELEGATE_RetVal_OneParam(bool, FRigVMController_ShouldStructUnfoldDelegate, const UStruct*)
-DECLARE_DELEGATE_RetVal(TArray<FRigVMExternalVariable>, FRigVMController_GetExternalVariablesDelegate)
+DECLARE_DELEGATE_RetVal_OneParam(TArray<FRigVMExternalVariable>, FRigVMController_GetExternalVariablesDelegate, URigVMGraph*)
 DECLARE_DELEGATE_RetVal(const FRigVMByteCode*, FRigVMController_GetByteCodeDelegate)
 DECLARE_DELEGATE_RetVal_OneParam(bool, FRigVMController_IsFunctionAvailableDelegate, URigVMLibraryNode*)
 DECLARE_DELEGATE_RetVal_OneParam(bool, FRigVMController_RequestLocalizeFunctionDelegate, URigVMLibraryNode*)
@@ -544,6 +544,10 @@ public:
 	UFUNCTION(BlueprintCallable, Category = RigVMController)
 	URigVMFunctionReferenceNode* AddFunctionReferenceNode(URigVMLibraryNode* InFunctionDefinition, const FVector2D& InNodePosition = FVector2D::ZeroVector, const FString& InNodeName = TEXT(""), bool bSetupUndoRedo = true);
 
+	// Sets the remapped variable on a function reference node
+	UFUNCTION(BlueprintCallable, Category = RigVMController)
+    bool SetRemappedVariable(URigVMFunctionReferenceNode* InFunctionRefNode, const FName& InInnerVariableName, const FName& InOuterVariableName, bool bSetupUndoRedo = true);
+
 	// Adds a function definition to a function library graph
 	UFUNCTION(BlueprintCallable, Category = RigVMController)
 	URigVMLibraryNode* AddFunctionToLibrary(const FName& InFunctionName, bool bMutable, const FVector2D& InNodePosition = FVector2D::ZeroVector, bool bSetupUndoRedo = true);
diff --git a/Engine/Source/Developer/RigVMDeveloper/Public/RigVMModel/RigVMControllerActions.h b/Engine/Source/Developer/RigVMDeveloper/Public/RigVMModel/RigVMControllerActions.h
--- a/Engine/Source/Developer/RigVMDeveloper/Public/RigVMModel/RigVMControllerActions.h
+++ b/Engine/Source/Developer/RigVMDeveloper/Public/RigVMModel/RigVMControllerActions.h
@@ -1436,4 +1436,34 @@ public:
 
 	UPROPERTY()
 	int32 NewIndex;
-};
\ No newline at end of file
+};
+
+/**
+* An action to remap a variable inside of a function reference node renaming a node in the graph.
+*/
+USTRUCT()
+struct FRigVMSetRemappedVariableAction : public FRigVMBaseAction
+{
+	GENERATED_BODY()
+
+public:
+
+	FRigVMSetRemappedVariableAction() {}
+	FRigVMSetRemappedVariableAction(URigVMFunctionReferenceNode* InFunctionRefNode, const FName& InInnerVariableName,
+		const FName& InOldOuterVariableName, const FName& InNewOuterVariableName);
+	virtual ~FRigVMSetRemappedVariableAction() {};
+	virtual bool Undo(URigVMController* InController) override;
+	virtual bool Redo(URigVMController* InController) override;
+
+	UPROPERTY()
+	FString NodePath;
+
+	UPROPERTY()
+	FName InnerVariableName;
+
+	UPROPERTY()
+	FName OldOuterVariableName;
+
+	UPROPERTY()
+	FName NewOuterVariableName;
+};
diff --git a/Engine/Source/Developer/RigVMDeveloper/Public/RigVMModel/RigVMFunctionLibrary.h b/Engine/Source/Developer/RigVMDeveloper/Public/RigVMModel/RigVMFunctionLibrary.h
--- a/Engine/Source/Developer/RigVMDeveloper/Public/RigVMModel/RigVMFunctionLibrary.h
+++ b/Engine/Source/Developer/RigVMDeveloper/Public/RigVMModel/RigVMFunctionLibrary.h
@@ -55,6 +55,10 @@ public:
 	UFUNCTION(BlueprintCallable, Category = RigVMGraph)
 	URigVMLibraryNode* FindFunction(const FName& InFunctionName) const;
 
+	// Finds a function by a node within a function (or a sub graph of that)
+	UFUNCTION(BlueprintCallable, Category = RigVMGraph)
+    URigVMLibraryNode* FindFunctionForNode(URigVMNode* InNode) const;
+
 	// Returns all references for a given function name
 	UFUNCTION(BlueprintCallable, Category = RigVMGraph)
 	TArray< TSoftObjectPtr<URigVMFunctionReferenceNode> > GetReferencesForFunction(const FName& InFunctionName);
@@ -66,6 +70,20 @@ public:
 	// Update the references list for a given reference node
 	void UpdateReferencesForReferenceNode(URigVMFunctionReferenceNode* InReferenceNode);
 
+	/**
+	* Iterator function to invoke a lambda / TFunction for each reference of a function
+	* @param InFunctionName The function name to iterate all references for
+	* @param PerReferenceFunction The function to invoke for each reference
+	*/
+	void ForEachReference(const FName& InFunctionName, TFunction<void(URigVMFunctionReferenceNode*)> PerReferenceFunction) const;
+
+	/**
+	* Iterator function to invoke a lambda / TFunction for each reference of a function
+	* @param InFunctionName The function name to iterate all references for
+	* @param PerReferenceFunction The function to invoke for each reference
+	*/
+	void ForEachReferenceSoftPtr(const FName& InFunctionName, TFunction<void(TSoftObjectPtr<URigVMFunctionReferenceNode>)> PerReferenceFunction) const;
+
 	// Returns a function that has been previously localized based on the provided function to localize.
 	// We maintain meta data on what functions have been created locally based on which other ones,
 	// and use this method to avoid redundant localizations.
diff --git a/Engine/Source/Developer/RigVMDeveloper/Public/RigVMModel/RigVMGraph.h b/Engine/Source/Developer/RigVMDeveloper/Public/RigVMModel/RigVMGraph.h
--- a/Engine/Source/Developer/RigVMDeveloper/Public/RigVMModel/RigVMGraph.h
+++ b/Engine/Source/Developer/RigVMDeveloper/Public/RigVMModel/RigVMGraph.h
@@ -52,6 +52,10 @@ public:
 	UFUNCTION(BlueprintCallable, Category = RigVMGraph)
     URigVMGraph* GetRootGraph() const;
 
+	// Returns true if this graph is a root / top level graph
+	UFUNCTION(BlueprintCallable, Category = RigVMGraph)
+    bool IsRootGraph() const;
+
 	// Returns the entry node of this graph
 	UFUNCTION(BlueprintCallable, Category = RigVMGraph)
 	URigVMFunctionEntryNode* GetEntryNode() const;
@@ -115,6 +119,8 @@ public:
 	UFUNCTION(BlueprintCallable, Category = RigVMGraph)
 	void SetDefaultFunctionLibrary(URigVMFunctionLibrary* InFunctionLibrary);
 
+	TArray<FRigVMExternalVariable> GetExternalVariables() const;
+
 	// Returns the modified event, which can be used to 
 	// subscribe to changes happening within the Graph.
 	FRigVMGraphModifiedEvent& OnModified();
@@ -128,8 +134,6 @@ public:
 
 private:
 
-	void RepopulatePinLinks();
-
 	FRigVMGraphModifiedEvent ModifiedEvent;
 	void Notify(ERigVMGraphNotifType InNotifType, UObject* InSubject);
 
diff --git a/Engine/Source/Developer/RigVMDeveloper/Public/RigVMModel/RigVMNotifications.h b/Engine/Source/Developer/RigVMDeveloper/Public/RigVMModel/RigVMNotifications.h
--- a/Engine/Source/Developer/RigVMDeveloper/Public/RigVMModel/RigVMNotifications.h
+++ b/Engine/Source/Developer/RigVMDeveloper/Public/RigVMModel/RigVMNotifications.h
@@ -52,6 +52,7 @@ enum class ERigVMGraphNotifType : uint8
 	NodeReferenceChanged, // A node has changed it's referenced function
 	NodeCategoryChanged, // A node's category has changed (Subject == URigVMNode)
 	NodeKeywordsChanged, // A node's keywords have changed (Subject == URigVMNode)
+	VariableRemappingChanged, // A function reference node's remapping has changed (Subject == URigVMFunctionReferenceNode)
 	Invalid // The max for this enum (used for guarding)
 };
 
diff --git a/Engine/Source/Runtime/RigVM/Public/RigVMCore/RigVMExternalVariable.h b/Engine/Source/Runtime/RigVM/Public/RigVMCore/RigVMExternalVariable.h
--- a/Engine/Source/Runtime/RigVM/Public/RigVMCore/RigVMExternalVariable.h
+++ b/Engine/Source/Runtime/RigVM/Public/RigVMCore/RigVMExternalVariable.h
@@ -283,7 +283,7 @@ struct RIGVM_API FRigVMExternalVariable
 	{
 		FRigVMExternalVariable Variable;
 		Variable.Name = InName;
-		Variable.TypeName = TBaseStructure<T>::Get()->GetFName();
+		Variable.TypeName = *TBaseStructure<T>::Get()->GetStructCPPName();
 		Variable.TypeObject = TBaseStructure<T>::Get();
 		Variable.bIsArray = false;
 		Variable.Size = TBaseStructure<T>::Get()->GetStructureSize();
@@ -299,7 +299,7 @@ struct RIGVM_API FRigVMExternalVariable
 	{
 		FRigVMExternalVariable Variable;
 		Variable.Name = InName;
-		Variable.TypeName = TBaseStructure<T>::Get()->GetFName();
+		Variable.TypeName = *TBaseStructure<T>::Get()->GetStructCPPName();
 		Variable.TypeObject = TBaseStructure<T>::Get();
 		Variable.bIsArray = true;
 		Variable.Size = TBaseStructure<T>::Get()->GetStructureSize();
@@ -315,7 +315,7 @@ struct RIGVM_API FRigVMExternalVariable
 	{
 		FRigVMExternalVariable Variable;
 		Variable.Name = InName;
-		Variable.TypeName = T::StaticStruct()->GetFName();
+		Variable.TypeName = *T::StaticStruct()->GetStructCPPName();
 		Variable.TypeObject = T::StaticStruct();
 		Variable.bIsArray = false;
 		Variable.Size = T::StaticStruct()->GetStructureSize();
@@ -331,7 +331,7 @@ struct RIGVM_API FRigVMExternalVariable
 	{
 		FRigVMExternalVariable Variable;
 		Variable.Name = InName;
-		Variable.TypeName = T::StaticStruct()->GetFName();
+		Variable.TypeName = *T::StaticStruct()->GetStructCPPName();
 		Variable.TypeObject = T::StaticStruct();
 		Variable.bIsArray = true;
 		Variable.Size = T::StaticStruct()->GetStructureSize();
@@ -347,7 +347,7 @@ struct RIGVM_API FRigVMExternalVariable
 	{
 		FRigVMExternalVariable Variable;
 		Variable.Name = InName;
-		Variable.TypeName = T::StaticClass()->GetFName();
+		Variable.TypeName = *T::StaticClass()->GetStructCPPName();
 		Variable.TypeObject = T::StaticClass();
 		Variable.bIsArray = false;
 		Variable.Size = T::StaticClass()->GetStructureSize();
@@ -363,7 +363,7 @@ struct RIGVM_API FRigVMExternalVariable
 	{
 		FRigVMExternalVariable Variable;
 		Variable.Name = InName;
-		Variable.TypeName = T::StaticClass()->GetFName();
+		Variable.TypeName = *T::StaticClass()->GetStructCPPName();
 		Variable.TypeObject = T::StaticClass();
 		Variable.bIsArray = true;
 		Variable.Size = T::StaticClass()->GetStructureSize();
@@ -431,6 +431,26 @@ struct RIGVM_API FRigVMExternalVariable
 		);
 	}
 
+	FORCEINLINE static void MergeExternalVariable(TArray<FRigVMExternalVariable>& OutVariables, const FRigVMExternalVariable& InVariable)
+	{
+		if(!InVariable.IsValid(true))
+		{
+			return;
+		}
+
+		for(const FRigVMExternalVariable& ExistingVariable : OutVariables)
+		{
+			if(ExistingVariable.Name == InVariable.Name)
+			{
+				ensure(ExistingVariable.TypeName == InVariable.TypeName);
+				ensure(ExistingVariable.TypeObject == InVariable.TypeObject);
+				return;
+			}
+		}
+
+		OutVariables.Add(InVariable);
+	}
+
 	FName Name;
 	FName TypeName;
 	UObject* TypeObject;
