Control Rig: Introducing the notion of bulk edits to the rigvmcontroller

diff --git a/Engine/Plugins/Experimental/ControlRig/Source/ControlRigDeveloper/Private/ControlRigBlueprint.cpp b/Engine/Plugins/Experimental/ControlRig/Source/ControlRigDeveloper/Private/ControlRigBlueprint.cpp
--- a/Engine/Plugins/Experimental/ControlRig/Source/ControlRigDeveloper/Private/ControlRigBlueprint.cpp
+++ b/Engine/Plugins/Experimental/ControlRig/Source/ControlRigDeveloper/Private/ControlRigBlueprint.cpp
@@ -1308,6 +1308,22 @@ URigVMController* UControlRigBlueprint::GetOrCreateController(URigVMGraph* InGra
 			return NAME_None;
 		}
 	));
+
+	TWeakObjectPtr<URigVMController> WeakController = Controller;
+	Controller->RequestBulkEditDialogDelegate.BindLambda([WeakThis, WeakController](URigVMLibraryNode* InFunction, ERigVMControllerBulkEditType InEditType)
+	{
+		if(WeakThis.IsValid() && WeakController.IsValid())
+		{
+			UControlRigBlueprint* StrongThis = WeakThis.Get();
+            URigVMController* StrongController = WeakController.Get();
+            if(StrongThis->OnRequestBulkEditDialog().IsBound())
+			{
+				return StrongThis->OnRequestBulkEditDialog().Execute(StrongThis, StrongController, InFunction, InEditType);
+			}
+		}
+		return true;
+	});
+	
 #endif
 
 	Controller->RemoveStaleNodes();
diff --git a/Engine/Plugins/Experimental/ControlRig/Source/ControlRigDeveloper/Public/ControlRigBlueprint.h b/Engine/Plugins/Experimental/ControlRig/Source/ControlRigDeveloper/Public/ControlRigBlueprint.h
--- a/Engine/Plugins/Experimental/ControlRig/Source/ControlRigDeveloper/Public/ControlRigBlueprint.h
+++ b/Engine/Plugins/Experimental/ControlRig/Source/ControlRigDeveloper/Public/ControlRigBlueprint.h
@@ -40,6 +40,7 @@ DECLARE_EVENT_OneParam(UControlRigBlueprint, FOnGraphImported, UEdGraph*);
 DECLARE_EVENT_OneParam(UControlRigBlueprint, FOnPostEditChangeChainProperty, FPropertyChangedChainEvent&);
 DECLARE_EVENT_ThreeParams(UControlRigBlueprint, FOnLocalizeFunctionDialogRequested, URigVMLibraryNode*, UControlRigBlueprint*, bool);
 DECLARE_EVENT_ThreeParams(UControlRigBlueprint, FOnReportCompilerMessage, EMessageSeverity::Type, UObject*, const FString&);
+DECLARE_DELEGATE_RetVal_FourParams(bool, FControlRigOnBulkEditDialogRequestedDelegate, UControlRigBlueprint*, URigVMController*, URigVMLibraryNode*, ERigVMControllerBulkEditType);
 
 USTRUCT()
 struct CONTROLRIGDEVELOPER_API FControlRigPublicFunctionArg
@@ -512,6 +513,8 @@ public:
 	FOnLocalizeFunctionDialogRequested& OnRequestLocalizeFunctionDialog() { return RequestLocalizeFunctionDialog; }
 	void BroadcastRequestLocalizeFunctionDialog(URigVMLibraryNode* InFunction, bool bForce = false);
 
+	FControlRigOnBulkEditDialogRequestedDelegate& OnRequestBulkEditDialog() { return RequestBulkEditDialog; }
+
 	FOnReportCompilerMessage& OnReportCompilerMessage() { return ReportCompilerMessageEvent; }
 	void BroadCastReportCompilerMessage(EMessageSeverity::Type InSeverity, UObject* InSubject, const FString& InMessage);
 
@@ -526,6 +529,7 @@ private:
 	FOnPostEditChangeChainProperty PostEditChangeChainPropertyEvent;
 	FOnLocalizeFunctionDialogRequested RequestLocalizeFunctionDialog;
 	FOnReportCompilerMessage ReportCompilerMessageEvent;
+	FControlRigOnBulkEditDialogRequestedDelegate RequestBulkEditDialog;
 
 #endif
 
diff --git a/Engine/Plugins/Experimental/ControlRig/Source/ControlRigEditor/Private/Editor/ControlRigEditor.cpp b/Engine/Plugins/Experimental/ControlRig/Source/ControlRigEditor/Private/Editor/ControlRigEditor.cpp
--- a/Engine/Plugins/Experimental/ControlRig/Source/ControlRigEditor/Private/Editor/ControlRigEditor.cpp
+++ b/Engine/Plugins/Experimental/ControlRig/Source/ControlRigEditor/Private/Editor/ControlRigEditor.cpp
@@ -80,6 +80,7 @@
 #include "BlueprintActionDatabaseRegistrar.h"
 #include "IMessageLogListing.h"
 #include "SControlRigFunctionLocalizationWidget.h"
+#include "SControlRigFunctionBulkEditWidget.h"
 
 #define LOCTEXT_NAMESPACE "ControlRigEditor"
 
@@ -130,6 +131,7 @@ FControlRigEditor::~FControlRigEditor()
 		RigBlueprint->OnGraphImported().RemoveAll(this);
 		RigBlueprint->OnPostEditChangeChainProperty().RemoveAll(this);
 		RigBlueprint->OnRequestLocalizeFunctionDialog().RemoveAll(this);
+		RigBlueprint->OnRequestBulkEditDialog().Unbind();
 		RigBlueprint->OnReportCompilerMessage().RemoveAll(this);
 	}
 
@@ -358,6 +360,7 @@ void FControlRigEditor::InitControlRigEditor(const EToolkitMode::Type Mode, cons
 		InControlRigBlueprint->OnGraphImported().AddSP(this, &FControlRigEditor::OnGraphImported);
 		InControlRigBlueprint->OnPostEditChangeChainProperty().AddSP(this, &FControlRigEditor::OnBlueprintPropertyChainEvent);
 		InControlRigBlueprint->OnRequestLocalizeFunctionDialog().AddSP(this, &FControlRigEditor::OnRequestLocalizeFunctionDialog);
+		InControlRigBlueprint->OnRequestBulkEditDialog().BindSP(this, &FControlRigEditor::OnRequestBulkEditDialog);
 	}
 
 	UpdateStaleWatchedPins();
@@ -3420,6 +3423,18 @@ void FControlRigEditor::OnRequestLocalizeFunctionDialog(URigVMLibraryNode* InFun
 	}
 }
 
+bool FControlRigEditor::OnRequestBulkEditDialog(UControlRigBlueprint* InBlueprint, URigVMController* InController,
+	URigVMLibraryNode* InFunction, ERigVMControllerBulkEditType InEditType)
+{
+	TSharedRef<SControlRigFunctionBulkEditDialog> BulkEditDialog = SNew(SControlRigFunctionBulkEditDialog)
+	.Blueprint(InBlueprint)
+	.Controller(InController)
+	.Function(InFunction)
+	.EditType(InEditType);
+	
+	return BulkEditDialog->ShowModal() != EAppReturnType::Cancel;
+}
+
 void FControlRigEditor::OnCreateComment()
 {
 	TSharedPtr<SGraphEditor> GraphEditor = FocusedGraphEdPtr.Pin();
diff --git a/Engine/Plugins/Experimental/ControlRig/Source/ControlRigEditor/Private/Editor/ControlRigEditor.h b/Engine/Plugins/Experimental/ControlRig/Source/ControlRigEditor/Private/Editor/ControlRigEditor.h
--- a/Engine/Plugins/Experimental/ControlRig/Source/ControlRigEditor/Private/Editor/ControlRigEditor.h
+++ b/Engine/Plugins/Experimental/ControlRig/Source/ControlRigEditor/Private/Editor/ControlRigEditor.h
@@ -19,6 +19,7 @@
 #include "AnimPreviewInstance.h"
 #include "ScopedTransaction.h"
 #include "Graph/ControlRigGraphNode.h"
+#include "RigVMModel/RigVMController.h"
 
 class UControlRigBlueprint;
 class IPersonaToolkit;
@@ -366,6 +367,7 @@ protected:
 	virtual void OnFinishedChangingProperties(const FPropertyChangedEvent& PropertyChangedEvent) override;
 	void OnBlueprintPropertyChainEvent(FPropertyChangedChainEvent& PropertyChangedChainEvent);
 	void OnRequestLocalizeFunctionDialog(URigVMLibraryNode* InFunction, UControlRigBlueprint* InTargetBlueprint, bool bForce);
+	bool OnRequestBulkEditDialog(UControlRigBlueprint* InBlueprint, URigVMController* InController, URigVMLibraryNode* InFunction, ERigVMControllerBulkEditType InEditType);
 
 	URigVMGraph* GetFocusedModel() const;
 	URigVMController* GetFocusedController() const;
diff --git a/Engine/Plugins/Experimental/ControlRig/Source/ControlRigEditor/Private/Editor/SControlRigFunctionBulkEditWidget.cpp b/Engine/Plugins/Experimental/ControlRig/Source/ControlRigEditor/Private/Editor/SControlRigFunctionBulkEditWidget.cpp
new file mode 100644
--- /dev/null
+++ b/Engine/Plugins/Experimental/ControlRig/Source/ControlRigEditor/Private/Editor/SControlRigFunctionBulkEditWidget.cpp
@@ -0,0 +1,216 @@
+// Copyright Epic Games, Inc. All Rights Reserved.
+
+#include "SControlRigFunctionBulkEditWidget.h"
+#include "Widgets/Layout/SBorder.h"
+#include "Widgets/SBoxPanel.h"
+#include "Widgets/Input/SButton.h"
+#include "Widgets/Input/SCheckBox.h"
+#include "Widgets/Layout/SUniformGridPanel.h"
+#include "DetailLayoutBuilder.h"
+#include "Editor.h"
+#include "EditorStyleSet.h"
+#include "HAL/ConsoleManager.h"
+#include "RigVMModel/Nodes/RigVMFunctionReferenceNode.h"
+#include "RigVMModel/RigVMController.h"
+#include "ContentBrowserModule.h"
+#include "IContentBrowserSingleton.h"
+#include "Misc/ScopedSlowTask.h"
+
+#define LOCTEXT_NAMESPACE "SControlRigFunctionBulkEditWidget"
+
+//////////////////////////////////////////////////////////////
+/// SControlRigFunctionBulkEditWidget
+///////////////////////////////////////////////////////////
+
+void SControlRigFunctionBulkEditWidget::Construct(const FArguments& InArgs, UControlRigBlueprint* InBlueprint, URigVMController* InController, URigVMLibraryNode* InFunction, ERigVMControllerBulkEditType InEditType)
+{
+	Blueprint = InBlueprint;
+	Controller = InController;
+	Function = InFunction;
+	EditType = InEditType;
+
+	ChildSlot
+	[
+		SNew(SVerticalBox)
+
+		+ SVerticalBox::Slot()
+		.Padding(4.0f, 4.0f, 0.f, 0.f)
+		.AutoHeight()
+		.MaxHeight(600.f)
+		[
+			SNew(SBorder)
+			.Visibility(EVisibility::Visible)
+			.BorderImage(FEditorStyle::GetBrush("Menu.Background"))
+			[
+				MakeAssetViewForReferencedAssets()
+			]
+		]
+	];
+}
+
+TSharedPtr<SWidget> SControlRigFunctionBulkEditWidget::OnGetAssetContextMenu(const TArray<FAssetData>& SelectedAssets)
+{
+	return nullptr;
+}
+
+void SControlRigFunctionBulkEditWidget::OnAssetsActivated(const TArray<FAssetData>& ActivatedAssets,
+	EAssetTypeActivationMethod::Type ActivationMethod)
+{
+	// nothing to do here
+}
+
+void SControlRigFunctionBulkEditWidget::LoadAffectedAssets()
+{
+	TArray<FAssetData> FirstLevelReferenceAssets = Controller->GetAffectedAssets(EditType, false, true);
+
+	{
+		FScopedSlowTask SlowTask((float)FirstLevelReferenceAssets.Num(), LOCTEXT("LoadingAffectedAssets", "Loading Affected Assets"));
+		Controller->OnBulkEditProgressDelegate.BindLambda([&SlowTask](
+            TSoftObjectPtr<URigVMFunctionReferenceNode> InReference,
+            ERigVMControllerBulkEditType InEditType,
+            ERigVMControllerBulkEditProgress InProgress,
+            int32 InIndex,
+            int32 InCount)
+        {
+            if(!InReference.IsValid())
+            {
+                return;
+            }
+
+			const FString PackageName = InReference.ToSoftObjectPath().GetLongPackageName();
+			const FText Message = FText::FromString(FString::Printf(TEXT("Loading '%s' ..."), *PackageName));
+
+			SlowTask.TotalAmountOfWork = (float)InCount;
+			SlowTask.EnterProgressFrame(1.f, Message);
+        });
+		
+		SlowTask.MakeDialog();
+
+		Controller->GetAffectedAssets(EditType, true, true);
+	}
+
+	Controller->OnBulkEditProgressDelegate.Unbind();
+}
+
+TSharedRef<SWidget> SControlRigFunctionBulkEditWidget::MakeAssetViewForReferencedAssets()
+{
+	FContentBrowserModule& ContentBrowserModule = FModuleManager::Get().LoadModuleChecked<FContentBrowserModule>(TEXT("ContentBrowser"));
+	FAssetPickerConfig AssetPickerConfig;
+
+	AssetPickerConfig.bAllowDragging = false;
+	AssetPickerConfig.bCanShowClasses = false;
+	AssetPickerConfig.bAllowNullSelection = false;
+	AssetPickerConfig.bShowBottomToolbar = false;
+	AssetPickerConfig.bAutohideSearchBar = true;
+	AssetPickerConfig.bPreloadAssetsForContextMenu = false;
+	AssetPickerConfig.Filter.ClassNames.Add(TEXT("InvalidClass"));
+	AssetPickerConfig.Filter.bIncludeOnlyOnDiskAssets = false;
+	AssetPickerConfig.OnAssetsActivated = FOnAssetsActivated::CreateSP(this, &SControlRigFunctionBulkEditWidget::OnAssetsActivated);
+	AssetPickerConfig.OnGetAssetContextMenu = FOnGetAssetContextMenu::CreateSP( this, &SControlRigFunctionBulkEditWidget::OnGetAssetContextMenu );
+
+	AssetPickerConfig.InitialAssetViewType = EAssetViewType::List;
+	AssetPickerConfig.OnGetCustomSourceAssets.BindLambda([this](const FARFilter& SourceFilter, TArray<FAssetData>& AddedAssets)
+	{
+		AddedAssets.Append(Controller->GetAffectedAssets(EditType, false, true));
+	});
+	
+	return ContentBrowserModule.Get().CreateAssetPicker(AssetPickerConfig);
+}
+
+
+void SControlRigFunctionBulkEditDialog::Construct(const FArguments& InArgs)
+{
+	UserResponse = EAppReturnType::Cancel;
+
+	static const FString DialogText = TEXT(
+        "This action will affect other assets.\n\n" 
+        "Note that this action is not undoable.\n\n"
+        "All affected assets will be marked as modified\n"
+        "and require re-saving."
+    );
+
+	const FString DialogTitle = FString::Printf(TEXT("Preparing Bulk-Edit - %s"), *StaticEnum<ERigVMControllerBulkEditType>()->GetDisplayNameTextByValue((int64)InArgs._EditType).ToString()); 
+	
+	SWindow::Construct(SWindow::FArguments()
+        .Title(FText::FromString(DialogTitle))
+        .SupportsMinimize(false)
+        .SupportsMaximize(false)
+        .SizingRule( ESizingRule::Autosized )
+        .ClientSize(FVector2D(450, 450))
+        [
+	        SNew(SHorizontalBox)
+
+	        +SHorizontalBox::Slot()
+	        .AutoWidth()
+	        .HAlign(HAlign_Left)
+	        .VAlign(VAlign_Top)
+	        .Padding(2)
+	        [
+				SNew(SVerticalBox)
+
+                +SVerticalBox::Slot()
+                .AutoHeight()
+                .HAlign(HAlign_Left)
+                .Padding(8)
+                [
+
+                    SNew(STextBlock)
+                    .Text(FText::FromString(DialogText))
+                    .Font(IDetailLayoutBuilder::GetDetailFont())
+                ]
+				
+                +SVerticalBox::Slot()
+                .AutoHeight()
+                .HAlign(HAlign_Left)
+                .Padding(2)
+                [
+                    SAssignNew(BulkEditWidget, SControlRigFunctionBulkEditWidget, InArgs._Blueprint, InArgs._Controller, InArgs._Function, InArgs._EditType)
+                ]
+
+                +SVerticalBox::Slot()
+                .AutoHeight()
+                .HAlign(HAlign_Right)
+                .Padding(5)
+                [
+                    SNew(SUniformGridPanel)
+                    .SlotPadding(FEditorStyle::GetMargin("StandardDialog.SlotPadding"))
+                    .MinDesiredSlotWidth(FEditorStyle::GetFloat("StandardDialog.MinDesiredSlotWidth"))
+                    .MinDesiredSlotHeight(FEditorStyle::GetFloat("StandardDialog.MinDesiredSlotHeight"))
+                    +SUniformGridPanel::Slot(0, 0)
+                    [
+                        SNew(SButton)
+                        .HAlign(HAlign_Center)
+                        .ContentPadding(FEditorStyle::GetMargin("StandardDialog.ContentPadding"))
+                        .Text(LOCTEXT("OK", "OK"))
+                        .OnClicked(this, &SControlRigFunctionBulkEditDialog::OnButtonClick, EAppReturnType::Ok)
+                    ]
+                    +SUniformGridPanel::Slot(1, 0)
+                    [
+                        SNew(SButton)
+                        .HAlign(HAlign_Center)
+                        .ContentPadding(FEditorStyle::GetMargin("StandardDialog.ContentPadding"))
+                        .Text(LOCTEXT("Cancel", "Cancel"))
+                        .OnClicked(this, &SControlRigFunctionBulkEditDialog::OnButtonClick, EAppReturnType::Cancel)
+                    ]
+                ]
+			]
+		]);
+}
+
+FReply SControlRigFunctionBulkEditDialog::OnButtonClick(EAppReturnType::Type ButtonID)
+{
+	UserResponse = ButtonID;
+	BulkEditWidget->LoadAffectedAssets();
+	
+	RequestDestroyWindow();
+
+	return FReply::Handled();
+}
+
+EAppReturnType::Type SControlRigFunctionBulkEditDialog::ShowModal()
+{
+	GEditor->EditorAddModalWindow(SharedThis(this));
+	return UserResponse;
+}
+
+#undef LOCTEXT_NAMESPACE
diff --git a/Engine/Plugins/Experimental/ControlRig/Source/ControlRigEditor/Private/Editor/SControlRigFunctionBulkEditWidget.h b/Engine/Plugins/Experimental/ControlRig/Source/ControlRigEditor/Private/Editor/SControlRigFunctionBulkEditWidget.h
new file mode 100644
--- /dev/null
+++ b/Engine/Plugins/Experimental/ControlRig/Source/ControlRigEditor/Private/Editor/SControlRigFunctionBulkEditWidget.h
@@ -0,0 +1,67 @@
+// Copyright Epic Games, Inc. All Rights Reserved.
+
+#pragma once
+
+#include "CoreMinimal.h"
+#include "ControlRigBlueprint.h"
+#include "RigVMModel/Nodes/RigVMLibraryNode.h"
+#include "Dialogs/Dialogs.h"
+#include "IAssetTypeActions.h"
+
+class SControlRigFunctionBulkEditWidget;
+
+class SControlRigFunctionBulkEditWidget : public SCompoundWidget
+{
+public:
+	SLATE_BEGIN_ARGS(SControlRigFunctionBulkEditWidget) {}
+	SLATE_END_ARGS()
+
+	void Construct(const FArguments& InArgs, UControlRigBlueprint* InBlueprint, URigVMController* InController, URigVMLibraryNode* InFunction, ERigVMControllerBulkEditType InEditType);
+
+private:
+
+	/** Handler for when an asset context menu has been requested. */
+	TSharedPtr<SWidget> OnGetAssetContextMenu( const TArray<FAssetData>& SelectedAssets );
+
+	/** Handler for when the user double clicks, presses enter, or presses space on an asset */
+	void OnAssetsActivated(const TArray<FAssetData>& ActivatedAssets, EAssetTypeActivationMethod::Type ActivationMethod);
+
+	/** Show a progress bar and load all of the assets */
+	void LoadAffectedAssets();
+
+	UControlRigBlueprint* Blueprint;
+	URigVMController* Controller;
+	URigVMLibraryNode* Function;
+	ERigVMControllerBulkEditType EditType;
+
+	TSharedRef<SWidget> MakeAssetViewForReferencedAssets();
+
+	friend class SControlRigFunctionBulkEditDialog;
+};
+
+class SControlRigFunctionBulkEditDialog : public SWindow
+{
+public:
+	
+	SLATE_BEGIN_ARGS(SControlRigFunctionBulkEditDialog)
+	{
+	}
+
+	SLATE_ARGUMENT(UControlRigBlueprint*, Blueprint)
+	SLATE_ARGUMENT(URigVMController*, Controller)
+	SLATE_ARGUMENT(URigVMLibraryNode*, Function)
+	SLATE_ARGUMENT(ERigVMControllerBulkEditType, EditType)
+	SLATE_END_ARGS()
+
+	void Construct(const FArguments& InArgs);
+
+	EAppReturnType::Type ShowModal();
+
+protected:
+
+	TSharedPtr<SControlRigFunctionBulkEditWidget> BulkEditWidget;
+	
+	FReply OnButtonClick(EAppReturnType::Type ButtonID);
+	bool IsOkButtonEnabled() const;
+	EAppReturnType::Type UserResponse;
+};
\ No newline at end of file
diff --git a/Engine/Source/Developer/RigVMDeveloper/Private/RigVMModel/RigVMController.cpp b/Engine/Source/Developer/RigVMDeveloper/Private/RigVMModel/RigVMController.cpp
--- a/Engine/Source/Developer/RigVMDeveloper/Private/RigVMModel/RigVMController.cpp
+++ b/Engine/Source/Developer/RigVMDeveloper/Private/RigVMModel/RigVMController.cpp
@@ -22,6 +22,7 @@
 #include "Framework/Notifications/NotificationManager.h"
 #include "Widgets/Notifications/SNotificationList.h"
 #include "EditorStyleSet.h"
+#include "AssetRegistryModule.h"
 #endif
 
 TMap<URigVMController::FControlRigStructPinRedirectorKey, FString> URigVMController::PinPathCoreRedirectors;
@@ -410,6 +411,40 @@ URigVMVariableNode* URigVMController::AddVariableNode(const FName& InVariableNam
 		return nullptr;
 	}
 
+	// check if the operation will cause to dirty assets
+	if(bSetupUndoRedo)
+	{
+		if(URigVMFunctionLibrary* OuterLibrary = Graph->GetTypedOuter<URigVMFunctionLibrary>())
+		{
+			if(URigVMLibraryNode* OuterFunction = OuterLibrary->FindFunctionForNode(Graph->GetTypedOuter<URigVMCollapseNode>()))
+			{
+				TArray<FRigVMExternalVariable> ExternalVariables = OuterFunction->GetContainedGraph()->GetExternalVariables();
+
+				bool bFoundExternalVariable = false;
+				for(const FRigVMExternalVariable& ExternalVariable : ExternalVariables)
+				{
+					if(ExternalVariable.Name == InVariableName)
+					{
+						bFoundExternalVariable = true;
+						break;
+					}
+				}
+
+				if(!bFoundExternalVariable)
+				{
+					if(RequestBulkEditDialogDelegate.IsBound())
+					{
+						if(!RequestBulkEditDialogDelegate.Execute(OuterFunction, ERigVMControllerBulkEditType::AddVariable))
+						{
+							return nullptr;
+						}
+						bSetupUndoRedo = false;
+					}
+				}
+			}
+		}
+	}
+
 	if (InCPPTypeObject == nullptr)
 	{
 		InCPPTypeObject = URigVMCompiler::GetScriptStructForCPPType(InCPPType);
@@ -424,7 +459,7 @@ URigVMVariableNode* URigVMController::AddVariableNode(const FName& InVariableNam
 	{
 		CPPType = ScriptStruct->GetStructCPPName();
 	}
-
+	
 	FString Name = GetValidNodeName(InNodeName.IsEmpty() ? FString(TEXT("VariableNode")) : InNodeName);
 	URigVMVariableNode* Node = NewObject<URigVMVariableNode>(Graph, *Name);
 	Node->Position = InPosition;
@@ -3554,6 +3589,47 @@ bool URigVMController::RemoveNode(URigVMNode* InNode, bool bSetupUndoRedo, bool
 		{
 			return false;
 		}
+
+		// check if the operation will cause to dirty assets
+		if(URigVMVariableNode* VariableNode = Cast<URigVMVariableNode>(InNode))
+		{
+			if(URigVMFunctionLibrary* OuterLibrary = Graph->GetTypedOuter<URigVMFunctionLibrary>())
+			{
+				if(URigVMLibraryNode* OuterFunction = OuterLibrary->FindFunctionForNode(Graph->GetTypedOuter<URigVMCollapseNode>()))
+				{
+					TArray<FRigVMExternalVariable> ExternalVariablesWithoutVariableNode;
+					{
+						URigVMGraph* EditedGraph = InNode->GetGraph();
+						TGuardValue<TArray<URigVMNode*>> TemporaryRemoveNodes(EditedGraph->Nodes, TArray<URigVMNode*>());
+						ExternalVariablesWithoutVariableNode = EditedGraph->GetExternalVariables();
+					}
+
+					const FName VariableToRemove = VariableNode->GetVariableName();
+					bool bFoundExternalVariable = false;
+					for(const FRigVMExternalVariable& ExternalVariable : ExternalVariablesWithoutVariableNode)
+					{
+						if(ExternalVariable.Name == VariableToRemove)
+						{
+							bFoundExternalVariable = true;
+							break;
+						}
+					}
+
+					if(!bFoundExternalVariable)
+					{
+						FRigVMControllerGraphGuard Guard(this, OuterFunction->GetContainedGraph(), false);
+						if(RequestBulkEditDialogDelegate.IsBound())
+						{
+							if(!RequestBulkEditDialogDelegate.Execute(OuterFunction, ERigVMControllerBulkEditType::RemoveVariable))
+							{
+								return nullptr;
+							}
+							bSetupUndoRedo = false;
+						}
+					}
+				}
+			}
+		}
 	}
 
 	TGuardValue<bool> GuardCompactness(bIgnoreRerouteCompactnessChanges, true);
@@ -5945,6 +6021,18 @@ bool URigVMController::RemoveExposedPin(const FName& InPinName, bool bSetupUndoR
 		return false;
 	}
 
+	if(bSetupUndoRedo)
+	{
+		if(RequestBulkEditDialogDelegate.IsBound())
+		{
+			if(!RequestBulkEditDialogDelegate.Execute(LibraryNode, ERigVMControllerBulkEditType::RemoveExposedPin))
+			{
+				return nullptr;
+			}
+			bSetupUndoRedo = false;
+		}
+	}
+
 	FRigVMRemoveExposedPinAction Action(Pin);
 	if (bSetupUndoRedo)
 	{
@@ -6012,6 +6100,18 @@ bool URigVMController::RenameExposedPin(const FName& InOldPinName, const FName&
 		return false;
 	}
 
+	if(bSetupUndoRedo)
+	{
+		if(RequestBulkEditDialogDelegate.IsBound())
+		{
+			if(!RequestBulkEditDialogDelegate.Execute(LibraryNode, ERigVMControllerBulkEditType::RenameExposedPin))
+			{
+				return nullptr;
+			}
+			bSetupUndoRedo = false;
+		}
+	}
+
 	FName PinName = GetUniqueName(InNewPinName, [LibraryNode](const FName& InName) {
 		return LibraryNode->FindPin(InName.ToString()) == nullptr;
 		});
@@ -6135,6 +6235,18 @@ bool URigVMController::ChangeExposedPinType(const FName& InPinName, const FStrin
 		return false;
 	}
 
+	if(bSetupUndoRedo)
+	{
+		if(RequestBulkEditDialogDelegate.IsBound())
+		{
+			if(!RequestBulkEditDialogDelegate.Execute(LibraryNode, ERigVMControllerBulkEditType::ChangeExposedPinType))
+			{
+				return nullptr;
+			}
+			bSetupUndoRedo = false;
+		}
+	}
+
 	FRigVMBaseAction Action;
 	if (bSetupUndoRedo)
 	{
@@ -6502,6 +6614,135 @@ bool URigVMController::RemoveFunctionFromLibrary(const FName& InFunctionName, bo
 	return RemoveNodeByName(InFunctionName, bSetupUndoRedo);
 }
 
+TArray<TSoftObjectPtr<URigVMFunctionReferenceNode>> URigVMController::GetAffectedReferences(ERigVMControllerBulkEditType InEditType, bool bForceLoad, bool bNotify)
+{
+	TArray<TSoftObjectPtr<URigVMFunctionReferenceNode>> FunctionReferencePtrs;
+	
+#if WITH_EDITOR
+
+	check(IsValidGraph());
+	URigVMGraph* Graph = GetGraph();
+	URigVMFunctionLibrary* FunctionLibrary = Graph->GetTypedOuter<URigVMFunctionLibrary>();
+	if(FunctionLibrary == nullptr)
+	{
+		return FunctionReferencePtrs;
+	}
+
+	URigVMLibraryNode* Function = FunctionLibrary->FindFunctionForNode(Graph->GetTypedOuter<URigVMCollapseNode>());
+	if(Function == nullptr)
+	{
+		return FunctionReferencePtrs;
+	}
+
+	// get the immediate references
+	FunctionReferencePtrs = FunctionLibrary->GetReferencesForFunction(Function->GetFName());
+	TMap<FString, int32> VisitedPaths;
+	
+	for(int32 FunctionReferenceIndex = 0; FunctionReferenceIndex < FunctionReferencePtrs.Num(); FunctionReferenceIndex++)
+	{
+		TSoftObjectPtr<URigVMFunctionReferenceNode> FunctionReferencePtr = FunctionReferencePtrs[FunctionReferenceIndex];
+		VisitedPaths.Add(FunctionReferencePtr.ToSoftObjectPath().ToString(), FunctionReferenceIndex);
+	}
+
+	for(int32 FunctionReferenceIndex = 0; FunctionReferenceIndex < FunctionReferencePtrs.Num(); FunctionReferenceIndex++)
+	{
+		TSoftObjectPtr<URigVMFunctionReferenceNode> FunctionReferencePtr = FunctionReferencePtrs[FunctionReferenceIndex];
+
+		if(bForceLoad)
+		{
+			if(OnBulkEditProgressDelegate.IsBound() && bNotify)
+			{
+				OnBulkEditProgressDelegate.Execute(FunctionReferencePtr, InEditType, ERigVMControllerBulkEditProgress::BeginLoad, FunctionReferenceIndex, FunctionReferencePtrs.Num());
+			}
+
+			if(!FunctionReferencePtr.IsValid())
+			{
+				FunctionReferencePtr.LoadSynchronous();
+			}
+
+			if(OnBulkEditProgressDelegate.IsBound() && bNotify)
+			{
+				OnBulkEditProgressDelegate.Execute(FunctionReferencePtr, InEditType, ERigVMControllerBulkEditProgress::FinishedLoad, FunctionReferenceIndex, FunctionReferencePtrs.Num());
+			}
+		}
+
+		// adding pins / renaming doesn't cause any recursion, so we can stop here
+		if((InEditType == ERigVMControllerBulkEditType::AddExposedPin) ||
+			(InEditType == ERigVMControllerBulkEditType::RemoveExposedPin) ||
+			(InEditType == ERigVMControllerBulkEditType::RenameExposedPin) ||
+			(InEditType == ERigVMControllerBulkEditType::ChangeExposedPinType) ||
+            (InEditType == ERigVMControllerBulkEditType::RenameVariable))
+		{
+			continue;
+		}
+
+		// for loaded assets we'll recurse now
+		if(FunctionReferencePtr.IsValid())
+		{
+			if(URigVMFunctionReferenceNode* AffectedFunctionReferenceNode = FunctionReferencePtr.Get())
+			{
+				if(URigVMFunctionLibrary* AffectedFunctionLibrary = AffectedFunctionReferenceNode->GetTypedOuter<URigVMFunctionLibrary>())
+				{
+					if(URigVMLibraryNode* AffectedFunction = AffectedFunctionLibrary->FindFunctionForNode(AffectedFunctionReferenceNode))
+					{
+						FRigVMControllerGraphGuard GraphGuard(this, AffectedFunction->GetContainedGraph(), false);
+						TArray<TSoftObjectPtr<URigVMFunctionReferenceNode>> AffectedFunctionReferencePtrs = GetAffectedReferences(InEditType, bForceLoad, false);
+						for(TSoftObjectPtr<URigVMFunctionReferenceNode> AffectedFunctionReferencePtr : AffectedFunctionReferencePtrs)
+						{
+							const FString Key = AffectedFunctionReferencePtr.ToSoftObjectPath().ToString();
+							if(VisitedPaths.Contains(Key))
+							{
+								continue;
+							}
+							VisitedPaths.Add(Key, FunctionReferencePtrs.Add(AffectedFunctionReferencePtr));
+						}
+					}
+				}
+			}
+		}
+	}
+	
+#endif
+
+	return FunctionReferencePtrs;
+}
+
+TArray<FAssetData> URigVMController::GetAffectedAssets(ERigVMControllerBulkEditType InEditType, bool bForceLoad, bool bNotify)
+{
+	TArray<FAssetData> Assets;
+
+#if WITH_EDITOR
+
+	TArray<TSoftObjectPtr<URigVMFunctionReferenceNode>> FunctionReferencePtrs = GetAffectedReferences(InEditType, bForceLoad, bNotify);
+	TMap<FString, int32> VisitedAssets;
+
+	const FAssetRegistryModule& AssetRegistryModule = FModuleManager::LoadModuleChecked<FAssetRegistryModule>(TEXT("AssetRegistry"));
+
+	for(int32 FunctionReferenceIndex = 0; FunctionReferenceIndex < FunctionReferencePtrs.Num(); FunctionReferenceIndex++)
+	{
+		TSoftObjectPtr<URigVMFunctionReferenceNode> FunctionReferencePtr = FunctionReferencePtrs[FunctionReferenceIndex];
+		const FString AssetPath = FunctionReferencePtr.ToSoftObjectPath().GetAssetPathName().ToString();
+		if(AssetPath.StartsWith(TEXT("/Engine/Transient")))
+		{
+			continue;
+		}
+		if(VisitedAssets.Contains(AssetPath))
+		{
+			continue;
+		}
+					
+		const FAssetData AssetData = AssetRegistryModule.Get().GetAssetByObjectPath(*AssetPath);
+		if(AssetData.IsValid())
+		{
+			VisitedAssets.Add(AssetPath, Assets.Add(AssetData));
+		}
+	}
+	
+#endif
+
+	return Assets;
+}
+
 void URigVMController::ExpandPinRecursively(URigVMPin* InPin, bool bSetupUndoRedo)
 {
 	if (InPin == nullptr)
diff --git a/Engine/Source/Developer/RigVMDeveloper/Private/RigVMModel/RigVMFunctionLibrary.cpp b/Engine/Source/Developer/RigVMDeveloper/Private/RigVMModel/RigVMFunctionLibrary.cpp
--- a/Engine/Source/Developer/RigVMDeveloper/Private/RigVMModel/RigVMFunctionLibrary.cpp
+++ b/Engine/Source/Developer/RigVMDeveloper/Private/RigVMModel/RigVMFunctionLibrary.cpp
@@ -48,16 +48,20 @@ URigVMLibraryNode* URigVMFunctionLibrary::FindFunction(const FName& InFunctionNa
 
 URigVMLibraryNode* URigVMFunctionLibrary::FindFunctionForNode(URigVMNode* InNode) const
 {
+	if(InNode == nullptr)
+	{
+		return nullptr;
+	}
+	
 	UObject* Subject = InNode;
-	do
+	while (Subject->GetOuter() != this)
 	{
+		Subject = Subject->GetOuter();
 		if(Subject == nullptr)
 		{
 			return nullptr;
 		}
-		Subject = Subject->GetOuter();
 	}
-	while (Subject->GetOuter() != this);
 
 	return Cast<URigVMLibraryNode>(Subject);
 }
diff --git a/Engine/Source/Developer/RigVMDeveloper/Public/RigVMModel/RigVMController.h b/Engine/Source/Developer/RigVMDeveloper/Public/RigVMModel/RigVMController.h
--- a/Engine/Source/Developer/RigVMDeveloper/Public/RigVMModel/RigVMController.h
+++ b/Engine/Source/Developer/RigVMDeveloper/Public/RigVMModel/RigVMController.h
@@ -20,12 +20,38 @@
 
 class URigVMActionStack;
 
+UENUM()
+enum class ERigVMControllerBulkEditType : uint8
+{
+	AddExposedPin,
+    RemoveExposedPin,
+    RenameExposedPin,
+    ChangeExposedPinType,
+    AddVariable,
+    RemoveVariable,
+	RenameVariable,
+    ChangeVariableType,
+    Max UMETA(Hidden),
+};
+
+UENUM()
+enum class ERigVMControllerBulkEditProgress : uint8
+{
+	BeginLoad,
+    FinishedLoad,
+	BeginEdit,
+    FinishedEdit,
+    Max UMETA(Hidden),
+};
+
 DECLARE_DELEGATE_RetVal_OneParam(bool, FRigVMController_ShouldStructUnfoldDelegate, const UStruct*)
 DECLARE_DELEGATE_RetVal_OneParam(TArray<FRigVMExternalVariable>, FRigVMController_GetExternalVariablesDelegate, URigVMGraph*)
 DECLARE_DELEGATE_RetVal(const FRigVMByteCode*, FRigVMController_GetByteCodeDelegate)
 DECLARE_DELEGATE_RetVal_OneParam(bool, FRigVMController_IsFunctionAvailableDelegate, URigVMLibraryNode*)
 DECLARE_DELEGATE_RetVal_OneParam(bool, FRigVMController_RequestLocalizeFunctionDelegate, URigVMLibraryNode*)
 DECLARE_DELEGATE_RetVal_TwoParams(bool, FRigVMController_IsDependencyCyclicDelegate, UObject*, UObject*)
+DECLARE_DELEGATE_RetVal_TwoParams(bool, FRigVMController_RequestBulkEditDialogDelegate, URigVMLibraryNode*, ERigVMControllerBulkEditType)
+DECLARE_DELEGATE_FiveParams(FRigVMController_OnBulkEditProgressDelegate, TSoftObjectPtr<URigVMFunctionReferenceNode>, ERigVMControllerBulkEditType, ERigVMControllerBulkEditProgress, int32, int32)
 
 /**
  * The Controller is the sole authority to perform changes
@@ -556,6 +582,12 @@ public:
 	UFUNCTION(BlueprintCallable, Category = RigVMController)
 	bool RemoveFunctionFromLibrary(const FName& InFunctionName, bool bSetupUndoRedo = true);
 
+	// Determine affected function references for a potential bulk edit on a library node
+	TArray<TSoftObjectPtr<URigVMFunctionReferenceNode>> GetAffectedReferences(ERigVMControllerBulkEditType InEditType, bool bForceLoad = false, bool bNotify = true);
+
+	// Determine affected assets for a potential bulk edit on a library node
+	TArray<FAssetData> GetAffectedAssets(ERigVMControllerBulkEditType InEditType, bool bForceLoad = false, bool bNotify = true);
+
 	// Sets the execute context struct type to use
 	void SetExecuteContextStruct(UStruct* InExecuteContextStruct);
 
@@ -575,7 +607,13 @@ public:
 	FRigVMController_RequestLocalizeFunctionDelegate RequestLocalizeFunctionDelegate;
 	
 	// A delegate to validate if we are allowed to introduce a dependency between two objects
-	FRigVMController_IsDependencyCyclicDelegate IsDependencyCyclicDelegate; 
+	FRigVMController_IsDependencyCyclicDelegate IsDependencyCyclicDelegate;
+
+	// A delegate to ask the host / client for a dialog to confirm a bulk edit
+	FRigVMController_RequestBulkEditDialogDelegate RequestBulkEditDialogDelegate;
+
+	// A delegate to inform the host / client about the progress during a bulk edit
+	FRigVMController_OnBulkEditProgressDelegate OnBulkEditProgressDelegate;
 
 	int32 DetachLinksFromPinObjects(const TArray<URigVMLink*>* InLinks = nullptr, bool bNotify = false);
 	int32 ReattachLinksToPinObjects(bool bFollowCoreRedirectors = false, const TArray<URigVMLink*>* InLinks = nullptr, bool bNotify = false, bool bSetupOrphanedPins = false);
