Fix build break - RigVMController.cpp

diff --git a/Engine/Source/Developer/RigVMDeveloper/Private/RigVMModel/RigVMController.cpp b/Engine/Source/Developer/RigVMDeveloper/Private/RigVMModel/RigVMController.cpp
--- a/Engine/Source/Developer/RigVMDeveloper/Private/RigVMModel/RigVMController.cpp
+++ b/Engine/Source/Developer/RigVMDeveloper/Private/RigVMModel/RigVMController.cpp
@@ -3622,7 +3622,7 @@ bool URigVMController::RemoveNode(URigVMNode* InNode, bool bSetupUndoRedo, bool
 						{
 							if(!RequestBulkEditDialogDelegate.Execute(OuterFunction, ERigVMControllerBulkEditType::RemoveVariable))
 							{
-								return nullptr;
+								return false;
 							}
 							bSetupUndoRedo = false;
 						}
@@ -6027,7 +6027,7 @@ bool URigVMController::RemoveExposedPin(const FName& InPinName, bool bSetupUndoR
 		{
 			if(!RequestBulkEditDialogDelegate.Execute(LibraryNode, ERigVMControllerBulkEditType::RemoveExposedPin))
 			{
-				return nullptr;
+				return false;
 			}
 			bSetupUndoRedo = false;
 		}
@@ -6106,7 +6106,7 @@ bool URigVMController::RenameExposedPin(const FName& InOldPinName, const FName&
 		{
 			if(!RequestBulkEditDialogDelegate.Execute(LibraryNode, ERigVMControllerBulkEditType::RenameExposedPin))
 			{
-				return nullptr;
+				return false;
 			}
 			bSetupUndoRedo = false;
 		}
@@ -6241,7 +6241,7 @@ bool URigVMController::ChangeExposedPinType(const FName& InPinName, const FStrin
 		{
 			if(!RequestBulkEditDialogDelegate.Execute(LibraryNode, ERigVMControllerBulkEditType::ChangeExposedPinType))
 			{
-				return nullptr;
+				return false;
 			}
 			bSetupUndoRedo = false;
 		}
