Control Rig: Don't orphan pins during ChangePinType without any links

diff --git a/Engine/Plugins/Experimental/ControlRig/Source/ControlRigDeveloper/Private/ControlRigBlueprint.cpp b/Engine/Plugins/Experimental/ControlRig/Source/ControlRigDeveloper/Private/ControlRigBlueprint.cpp
--- a/Engine/Plugins/Experimental/ControlRig/Source/ControlRigDeveloper/Private/ControlRigBlueprint.cpp
+++ b/Engine/Plugins/Experimental/ControlRig/Source/ControlRigDeveloper/Private/ControlRigBlueprint.cpp
@@ -566,11 +566,6 @@ void UControlRigBlueprint::PostLoad()
 	{
 		URigVMController* Controller = GetOrCreateController(GraphToDetach);
 		Controller->ReattachLinksToPinObjects(true /* follow redirectors */, nullptr, false, true);
-		
-		for(URigVMNode* Node : GraphToDetach->GetNodes())
-		{
-			Controller->RemoveUnusedOrphanedPins(Node, false);
-		}
 	}
 
 	// perform backwards compat value upgrades
@@ -582,6 +577,12 @@ void UControlRigBlueprint::PostLoad()
 		{
 			continue;
 		}
+
+		for(URigVMNode* Node : GraphToValidate->GetNodes())
+		{
+			URigVMController* Controller = GetOrCreateController(GraphToValidate);
+			Controller->RemoveUnusedOrphanedPins(Node, false);
+		}
 			
 		for(URigVMNode* Node : GraphToValidate->GetNodes())
 		{
diff --git a/Engine/Source/Developer/RigVMDeveloper/Private/RigVMModel/RigVMController.cpp b/Engine/Source/Developer/RigVMDeveloper/Private/RigVMModel/RigVMController.cpp
--- a/Engine/Source/Developer/RigVMDeveloper/Private/RigVMModel/RigVMController.cpp
+++ b/Engine/Source/Developer/RigVMDeveloper/Private/RigVMModel/RigVMController.cpp
@@ -9182,19 +9182,22 @@ bool URigVMController::ChangePinType(URigVMPin* InPin, const FString& InCPPType,
 		Links.Append(InPin->GetTargetLinks(true));
 		DetachLinksFromPinObjects(&Links, true);
 
-		const FString OrphanedName = FString::Printf(TEXT("%s%s"), *URigVMPin::OrphanPinPrefix, *InPin->GetName());
-		if(InPin->GetNode()->FindPin(OrphanedName) == nullptr)
+		if(Links.Num() > 0)
 		{
-			URigVMPin* OrphanedPin = NewObject<URigVMPin>(InPin->GetNode(), *OrphanedName);
-			ConfigurePinFromPin(OrphanedPin, InPin);
-			OrphanedPin->DisplayName = InPin->GetFName();
-
-			if(OrphanedPin->IsStruct())
+			const FString OrphanedName = FString::Printf(TEXT("%s%s"), *URigVMPin::OrphanPinPrefix, *InPin->GetName());
+			if(InPin->GetNode()->FindPin(OrphanedName) == nullptr)
 			{
-				AddPinsForStruct(OrphanedPin->GetScriptStruct(), OrphanedPin->GetNode(), OrphanedPin, OrphanedPin->Direction, OrphanedPin->GetDefaultValue(), false, true);
+				URigVMPin* OrphanedPin = NewObject<URigVMPin>(InPin->GetNode(), *OrphanedName);
+				ConfigurePinFromPin(OrphanedPin, InPin);
+				OrphanedPin->DisplayName = InPin->GetFName();
+
+				if(OrphanedPin->IsStruct())
+				{
+					AddPinsForStruct(OrphanedPin->GetScriptStruct(), OrphanedPin->GetNode(), OrphanedPin, OrphanedPin->Direction, OrphanedPin->GetDefaultValue(), false, true);
+				}
+					
+				InPin->GetNode()->OrphanedPins.Add(OrphanedPin);
 			}
-				
-			InPin->GetNode()->OrphanedPins.Add(OrphanedPin);
 		}
 	}
 
