Control Rig VM Compiler - Don't create registers for Variables

diff --git a/Engine/Source/Developer/RigVMDeveloper/Private/RigVMCompiler/RigVMCompiler.cpp b/Engine/Source/Developer/RigVMDeveloper/Private/RigVMCompiler/RigVMCompiler.cpp
--- a/Engine/Source/Developer/RigVMDeveloper/Private/RigVMCompiler/RigVMCompiler.cpp
+++ b/Engine/Source/Developer/RigVMDeveloper/Private/RigVMCompiler/RigVMCompiler.cpp
@@ -1415,6 +1415,25 @@ FRigVMOperand URigVMCompiler::FindOrAddRegister(const FRigVMVarExprAST* InVarExp
 	bool bIsExecutePin = Pin->IsExecuteContext();
 	bool bIsLiteral = InVarExpr->GetType() == FRigVMExprAST::EType::Literal && !bIsDebugValue;
 	bool bIsParameter = InVarExpr->IsGraphParameter();
+	bool bIsVariable = Pin->IsRootPin() && (Pin->GetName() == URigVMVariableNode::ValueName) &&
+		InVarExpr->GetPin()->GetNode()->IsA<URigVMVariableNode>();
+
+	// variables don't require to add any register.
+	if(bIsVariable && !bIsDebugValue)
+	{
+		for(int32 ExternalVariableIndex = 0; ExternalVariableIndex < WorkData.VM->GetExternalVariables().Num(); ExternalVariableIndex++)
+		{
+			const FName& ExternalVariableName = WorkData.VM->GetExternalVariables()[ExternalVariableIndex].Name;
+			const FString ExternalVariableHash = FString::Printf(TEXT("Variable::%s"), *ExternalVariableName.ToString());
+			if(ExternalVariableHash == Hash)
+			{
+				Operand = FRigVMOperand(ERigVMMemoryType::External, ExternalVariableIndex, INDEX_NONE);
+				WorkData.ExprToOperand.Add(InVarExpr, Operand);
+				WorkData.PinPathToOperand->FindOrAdd(Hash) = Operand;
+				return Operand;
+			}
+		}
+	}
 
 	FRigVMMemoryContainer& Memory = 
 	    bIsLiteral ? WorkData.VM->GetLiteralMemory() :
