Control Rig - Follow up on CL 16156500
Control Rig: Don't orphan pins during ChangePinType without any links

diff --git a/Engine/Source/Developer/RigVMDeveloper/Private/RigVMModel/RigVMController.cpp b/Engine/Source/Developer/RigVMDeveloper/Private/RigVMModel/RigVMController.cpp
--- a/Engine/Source/Developer/RigVMDeveloper/Private/RigVMModel/RigVMController.cpp
+++ b/Engine/Source/Developer/RigVMDeveloper/Private/RigVMModel/RigVMController.cpp
@@ -6265,6 +6265,7 @@ bool URigVMController::ChangeExposedPinType(const FName& InPinName, const FStrin
 			}
 			return false;
 		}
+		RemoveUnusedOrphanedPins(LibraryNode, true);
 	}
 
 	if (URigVMFunctionEntryNode* EntryNode = Graph->GetEntryNode())
@@ -6272,6 +6273,7 @@ bool URigVMController::ChangeExposedPinType(const FName& InPinName, const FStrin
 		if (URigVMPin* EntryPin = EntryNode->FindPin(Pin->GetName()))
 		{
 			ChangePinType(EntryPin, InCPPType, InCPPTypeObjectPath, bSetupUndoRedo, bSetupOrphanPins);
+			RemoveUnusedOrphanedPins(EntryNode, true);
 		}
 	}
 	if (URigVMFunctionReturnNode* ReturnNode = Graph->GetReturnNode())
@@ -6279,6 +6281,7 @@ bool URigVMController::ChangeExposedPinType(const FName& InPinName, const FStrin
 		if (URigVMPin* ReturnPin = ReturnNode->FindPin(Pin->GetName()))
 		{
 			ChangePinType(ReturnPin, InCPPType, InCPPTypeObjectPath, bSetupUndoRedo, bSetupOrphanPins);
+			RemoveUnusedOrphanedPins(ReturnNode, true);
 		}
 	}
 
@@ -6290,6 +6293,7 @@ bool URigVMController::ChangeExposedPinType(const FName& InPinName, const FStrin
 			{
 				FRigVMControllerGraphGuard GraphGuard(this, ReferenceNode->GetGraph(), false);
 				ChangePinType(ReferencedNodePin, InCPPType, InCPPTypeObjectPath, bSetupUndoRedo, bSetupOrphanPins);
+				RemoveUnusedOrphanedPins(ReferenceNode, true);
 			}
         });
 	}
@@ -9182,22 +9186,19 @@ bool URigVMController::ChangePinType(URigVMPin* InPin, const FString& InCPPType,
 		Links.Append(InPin->GetTargetLinks(true));
 		DetachLinksFromPinObjects(&Links, true);
 
-		if(Links.Num() > 0)
+		const FString OrphanedName = FString::Printf(TEXT("%s%s"), *URigVMPin::OrphanPinPrefix, *InPin->GetName());
+		if(InPin->GetNode()->FindPin(OrphanedName) == nullptr)
 		{
-			const FString OrphanedName = FString::Printf(TEXT("%s%s"), *URigVMPin::OrphanPinPrefix, *InPin->GetName());
-			if(InPin->GetNode()->FindPin(OrphanedName) == nullptr)
-			{
-				URigVMPin* OrphanedPin = NewObject<URigVMPin>(InPin->GetNode(), *OrphanedName);
-				ConfigurePinFromPin(OrphanedPin, InPin);
-				OrphanedPin->DisplayName = InPin->GetFName();
+			URigVMPin* OrphanedPin = NewObject<URigVMPin>(InPin->GetNode(), *OrphanedName);
+			ConfigurePinFromPin(OrphanedPin, InPin);
+			OrphanedPin->DisplayName = InPin->GetFName();
 
-				if(OrphanedPin->IsStruct())
-				{
-					AddPinsForStruct(OrphanedPin->GetScriptStruct(), OrphanedPin->GetNode(), OrphanedPin, OrphanedPin->Direction, OrphanedPin->GetDefaultValue(), false, true);
-				}
-					
-				InPin->GetNode()->OrphanedPins.Add(OrphanedPin);
+			if(OrphanedPin->IsStruct())
+			{
+				AddPinsForStruct(OrphanedPin->GetScriptStruct(), OrphanedPin->GetNode(), OrphanedPin, OrphanedPin->Direction, OrphanedPin->GetDefaultValue(), false, true);
 			}
+				
+			InPin->GetNode()->OrphanedPins.Add(OrphanedPin);
 		}
 	}
 
