Control Rig: Avoid dialogs for bulk changes on same asset. Catch invalid collapse pins

diff --git a/Engine/Source/Developer/RigVMDeveloper/Private/RigVMModel/RigVMController.cpp b/Engine/Source/Developer/RigVMDeveloper/Private/RigVMModel/RigVMController.cpp
--- a/Engine/Source/Developer/RigVMDeveloper/Private/RigVMModel/RigVMController.cpp
+++ b/Engine/Source/Developer/RigVMDeveloper/Private/RigVMModel/RigVMController.cpp
@@ -6717,9 +6717,18 @@ TArray<FAssetData> URigVMController::GetAffectedAssets(ERigVMControllerBulkEditT
 
 #if WITH_EDITOR
 
+	if(!IsValidGraph())
+	{
+		return Assets;
+	}
+
 	TArray<TSoftObjectPtr<URigVMFunctionReferenceNode>> FunctionReferencePtrs = GetAffectedReferences(InEditType, bForceLoad, bNotify);
 	TMap<FString, int32> VisitedAssets;
 
+	URigVMGraph* Graph = GetGraph();
+	TSoftObjectPtr<URigVMGraph> GraphPtr = Graph;
+	const FString ThisAssetPath = GraphPtr.ToSoftObjectPath().GetAssetPathName().ToString();
+
 	const FAssetRegistryModule& AssetRegistryModule = FModuleManager::LoadModuleChecked<FAssetRegistryModule>(TEXT("AssetRegistry"));
 
 	for(int32 FunctionReferenceIndex = 0; FunctionReferenceIndex < FunctionReferencePtrs.Num(); FunctionReferenceIndex++)
@@ -6734,6 +6743,10 @@ TArray<FAssetData> URigVMController::GetAffectedAssets(ERigVMControllerBulkEditT
 		{
 			continue;
 		}
+		if(AssetPath == ThisAssetPath)
+		{
+			continue;
+		}
 					
 		const FAssetData AssetData = AssetRegistryModule.Get().GetAssetByObjectPath(*AssetPath);
 		if(AssetData.IsValid())
diff --git a/Engine/Source/Developer/RigVMDeveloper/Private/RigVMModel/RigVMPin.cpp b/Engine/Source/Developer/RigVMDeveloper/Private/RigVMModel/RigVMPin.cpp
--- a/Engine/Source/Developer/RigVMDeveloper/Private/RigVMModel/RigVMPin.cpp
+++ b/Engine/Source/Developer/RigVMDeveloper/Private/RigVMModel/RigVMPin.cpp
@@ -234,6 +234,7 @@ void URigVMPin::GetExposedPinChain(TArray<const URigVMPin*>& OutExposedPins) con
 	for (URigVMLink* Link : GetSourceLinks())
 	{
 		URigVMPin* SourcePin = Link->GetSourcePin();
+		check(SourcePin != nullptr);
 		
 		// If the source is on an entry node, add the pin and make a recursive call on the collapse node pin
 		if (URigVMFunctionEntryNode* EntryNode = Cast<URigVMFunctionEntryNode>(SourcePin->GetNode()))
@@ -241,8 +242,10 @@ void URigVMPin::GetExposedPinChain(TArray<const URigVMPin*>& OutExposedPins) con
 			URigVMGraph* Graph = EntryNode->GetGraph();
 			if (URigVMCollapseNode* CollapseNode = Cast<URigVMCollapseNode>(Graph->GetOuter()))
 			{
-				URigVMPin* CollapseNodePin = CollapseNode->FindPin(SourcePin->GetName());
-				CollapseNodePin->GetExposedPinChain(OutExposedPins);				
+				if(URigVMPin* CollapseNodePin = CollapseNode->FindPin(SourcePin->GetName()))
+				{
+					CollapseNodePin->GetExposedPinChain(OutExposedPins);
+				}
 			}
 		}
 		else if (URigVMFunctionReturnNode* ReturnNode = Cast<URigVMFunctionReturnNode>(SourcePin->GetNode()))
@@ -250,8 +253,10 @@ void URigVMPin::GetExposedPinChain(TArray<const URigVMPin*>& OutExposedPins) con
 			URigVMGraph* Graph = ReturnNode->GetGraph();
 			if (URigVMCollapseNode* CollapseNode = Cast<URigVMCollapseNode>(Graph->GetOuter()))
 			{
-				URigVMPin* CollapseNodePin = CollapseNode->FindPin(SourcePin->GetName());
-				CollapseNodePin->GetExposedPinChain(OutExposedPins);				
+				if(URigVMPin* CollapseNodePin = CollapseNode->FindPin(SourcePin->GetName()))
+				{
+					CollapseNodePin->GetExposedPinChain(OutExposedPins);
+				}
 			}
 		}
 		// Variable nodes do not share the operand with their source link
