Control Rig: AST lookup var expresisons by name

diff --git a/Engine/Source/Developer/RigVMDeveloper/Private/RigVMCompiler/RigVMAST.cpp b/Engine/Source/Developer/RigVMDeveloper/Private/RigVMCompiler/RigVMAST.cpp
--- a/Engine/Source/Developer/RigVMDeveloper/Private/RigVMCompiler/RigVMAST.cpp
+++ b/Engine/Source/Developer/RigVMDeveloper/Private/RigVMCompiler/RigVMAST.cpp
@@ -805,15 +805,21 @@ int32 FRigVMSelectExprAST::NumValues() const
 
 const FRigVMVarExprAST* FRigVMCallExternExprAST::FindVarWithPinName(const FName& InPinName) const
 {
-	for (int32 ChildIndex = 0; ChildIndex < NumChildren(); ChildIndex++)
+	if(URigVMUnitNode* UnitNode = Cast<URigVMUnitNode>(GetNode()))
 	{
-		const FRigVMExprAST* Child = ChildAt(ChildIndex);
-		if (Child->IsA(FRigVMExprAST::Var))
+		for(URigVMPin* Pin : UnitNode->GetPins())
 		{
-			const FRigVMVarExprAST* VarExpr = Child->To<FRigVMVarExprAST>();
-			if (VarExpr->GetPin()->GetFName() == InPinName)
+			if(Pin->GetFName() == InPinName)
 			{
-				return VarExpr;
+				const int32 PinIndex = Pin->GetPinIndex();
+				if(PinIndex <= NumChildren())
+				{
+					const FRigVMExprAST* Child = ChildAt(PinIndex);
+					if (Child->IsA(FRigVMExprAST::Var))
+					{
+						return Child->To<FRigVMVarExprAST>();
+					}
+				}
 			}
 		}
 	}
