ControlRig: Adding local variables to CR graphs

diff --git a/Engine/Source/Developer/RigVMDeveloper/Private/RigVMCompiler/RigVMASTProxy.cpp b/Engine/Source/Developer/RigVMDeveloper/Private/RigVMCompiler/RigVMASTProxy.cpp
--- a/Engine/Source/Developer/RigVMDeveloper/Private/RigVMCompiler/RigVMASTProxy.cpp
+++ b/Engine/Source/Developer/RigVMDeveloper/Private/RigVMCompiler/RigVMASTProxy.cpp
@@ -59,7 +59,7 @@ const UObject* FRigVMCallstack::operator[](int32 InIndex) const
 	return Stack[InIndex];
 }
 
-bool FRigVMCallstack::Contains(UObject* InEntry) const
+bool FRigVMCallstack::Contains(const UObject* InEntry) const
 {
 	return Stack.Contains(InEntry);
 }
diff --git a/Engine/Source/Developer/RigVMDeveloper/Private/RigVMCompiler/RigVMCompiler.cpp b/Engine/Source/Developer/RigVMDeveloper/Private/RigVMCompiler/RigVMCompiler.cpp
--- a/Engine/Source/Developer/RigVMDeveloper/Private/RigVMCompiler/RigVMCompiler.cpp
+++ b/Engine/Source/Developer/RigVMDeveloper/Private/RigVMCompiler/RigVMCompiler.cpp
@@ -199,6 +199,36 @@ bool URigVMCompiler::Compile(URigVMGraph* InGraph, URigVMController* InControlle
 
 	UE_LOG_RIGVMMEMORY(TEXT("RigVMCompiler: Begin '%s'..."), *InGraph->GetPathName());
 
+
+	// Look for all local variables to create the register with the default value in the literal memory
+	int32 IndexLocalVariable = 0;
+	for(URigVMGraph* VisitedGraph : VisitedGraphs)
+	{
+		for (const FRigVMGraphVariableDescription& LocalVariable : VisitedGraph->LocalVariables)
+		{
+			// To create the default value in the literal memory, we need to find a pin in a variable node that
+			// uses this local variable
+			for (URigVMNode* Node : VisitedGraph->GetNodes())
+			{
+				if (URigVMVariableNode* VariableNode = Cast<URigVMVariableNode>(Node))
+				{
+					if (URigVMPin* Pin = VariableNode->FindPin(URigVMVariableNode::VariableName))
+					{
+						if (Pin->GetDefaultValue() == LocalVariable.Name.ToString())
+						{
+							URigVMPin* ValuePin = VariableNode->FindPin(URigVMVariableNode::ValueName);
+							FRigVMASTProxy PinProxy = FRigVMASTProxy::MakeFromUObject(ValuePin);
+							FRigVMVarExprAST* TempVarExpr = AST->MakeExpr<FRigVMVarExprAST>(FRigVMExprAST::EType::Literal, PinProxy);
+							FRigVMOperand Operand = FindOrAddRegister(TempVarExpr, WorkData, false);
+							WorkData.VM->GetLiteralMemory().SetRegisterValueFromString(Operand, LocalVariable.CPPType, LocalVariable.CPPTypeObject, {LocalVariable.DefaultValue});
+							break;
+						}
+					}
+				}
+			}
+		}
+	}	
+
 	// define all parameters independent from sorted nodes
 	for (const FRigVMExprAST* Expr : AST->Expressions)
 	{
@@ -298,6 +328,8 @@ void URigVMCompiler::TraverseExpression(const FRigVMExprAST* InExpr, FRigVMCompi
 	}
 	WorkData.ExprComplete.Add(InExpr, true);
 
+	InitializeLocalVariables(InExpr, WorkData);	
+
 	switch (InExpr->GetType())
 	{
 		case FRigVMExprAST::EType::Block:
@@ -449,8 +481,8 @@ void URigVMCompiler::TraverseEntry(const FRigVMEntryExprAST* InExpr, FRigVMCompi
 
 		if (Settings.SetupNodeInstructionIndex)
 		{
-			const FString CallPath = InExpr->GetProxy().GetCallstack().GetCallPath();
-			WorkData.VM->GetByteCode().SetSubject(EntryInstructionIndex, UnitNode, CallPath);
+			const FRigVMCallstack Callstack = InExpr->GetProxy().GetCallstack();
+			WorkData.VM->GetByteCode().SetSubject(EntryInstructionIndex, Callstack.GetCallPath(), Callstack.GetStack());
 		}
 
 		TraverseChildren(InExpr, WorkData);
@@ -533,8 +565,8 @@ int32 URigVMCompiler::TraverseCallExtern(const FRigVMCallExternExprAST* InExpr,
 		InstructionIndex = WorkData.VM->GetByteCode().GetNumInstructions() - 1;
 		if (Settings.SetupNodeInstructionIndex)
 		{
-			const FString CallPath = InExpr->GetProxy().GetCallstack().GetCallPath();
-			WorkData.VM->GetByteCode().SetSubject(InstructionIndex, UnitNode, CallPath);
+			const FRigVMCallstack Callstack = InExpr->GetProxy().GetCallstack();
+			WorkData.VM->GetByteCode().SetSubject(InstructionIndex, Callstack.GetCallPath(), Callstack.GetStack());
 		}
 	}
 
@@ -550,7 +582,7 @@ void URigVMCompiler::TraverseForLoop(const FRigVMCallExternExprAST* InExpr, FRig
 	}
 
 	URigVMUnitNode* UnitNode = Cast<URigVMUnitNode>(InExpr->GetNode());
-	const FString CallPath = InExpr->GetProxy().GetCallstack().GetCallPath();
+	const FRigVMCallstack Callstack = InExpr->GetProxy().GetCallstack();
 
 	const FRigVMVarExprAST* CompletedExpr = InExpr->FindVarWithPinName(FRigVMStruct::ForLoopCompletedPinName);
 	check(CompletedExpr);
@@ -566,14 +598,14 @@ void URigVMCompiler::TraverseForLoop(const FRigVMCallExternExprAST* InExpr, FRig
 	WorkData.VM->GetByteCode().AddZeroOp(IndexOperand);
 	if (Settings.SetupNodeInstructionIndex)
 	{
-		WorkData.VM->GetByteCode().SetSubject(WorkData.VM->GetByteCode().GetNumInstructions() - 1, UnitNode, CallPath);
+		WorkData.VM->GetByteCode().SetSubject(WorkData.VM->GetByteCode().GetNumInstructions() - 1, Callstack.GetCallPath(), Callstack.GetStack());
 	}
 
 	// call the for loop compute
 	int32 ForLoopInstructionIndex = TraverseCallExtern(InExpr, WorkData);
 	if (Settings.SetupNodeInstructionIndex)
 	{
-		WorkData.VM->GetByteCode().SetSubject(ForLoopInstructionIndex, UnitNode, CallPath);
+		WorkData.VM->GetByteCode().SetSubject(ForLoopInstructionIndex, Callstack.GetCallPath(), Callstack.GetStack());
 	}
 
 	// set up the jump forward (jump out of the loop)
@@ -585,7 +617,7 @@ void URigVMCompiler::TraverseForLoop(const FRigVMCallExternExprAST* InExpr, FRig
 	int32 JumpToEndInstruction = WorkData.VM->GetByteCode().GetNumInstructions() - 1;
 	if (Settings.SetupNodeInstructionIndex)
 	{
-		WorkData.VM->GetByteCode().SetSubject(WorkData.VM->GetByteCode().GetNumInstructions() - 1, UnitNode, CallPath);
+		WorkData.VM->GetByteCode().SetSubject(WorkData.VM->GetByteCode().GetNumInstructions() - 1, Callstack.GetCallPath(), Callstack.GetStack());
 	}
 
 	// begin the loop's block
@@ -595,7 +627,7 @@ void URigVMCompiler::TraverseForLoop(const FRigVMCallExternExprAST* InExpr, FRig
 	WorkData.VM->GetByteCode().AddBeginBlockOp(CountOperand, IndexOperand);
 	if (Settings.SetupNodeInstructionIndex)
 	{
-		WorkData.VM->GetByteCode().SetSubject(WorkData.VM->GetByteCode().GetNumInstructions() - 1, UnitNode, CallPath);
+		WorkData.VM->GetByteCode().SetSubject(WorkData.VM->GetByteCode().GetNumInstructions() - 1, Callstack.GetCallPath(), Callstack.GetStack());
 	}
 
 	// traverse the body of the loop
@@ -606,14 +638,14 @@ void URigVMCompiler::TraverseForLoop(const FRigVMCallExternExprAST* InExpr, FRig
 	WorkData.VM->GetByteCode().AddEndBlockOp();
 	if (Settings.SetupNodeInstructionIndex)
 	{
-		WorkData.VM->GetByteCode().SetSubject(WorkData.VM->GetByteCode().GetNumInstructions() - 1, UnitNode, CallPath);
+		WorkData.VM->GetByteCode().SetSubject(WorkData.VM->GetByteCode().GetNumInstructions() - 1, Callstack.GetCallPath(), Callstack.GetStack());
 	}
 
 	// increment the index
 	WorkData.VM->GetByteCode().AddIncrementOp(IndexOperand);
 	if (Settings.SetupNodeInstructionIndex)
 	{
-		WorkData.VM->GetByteCode().SetSubject(WorkData.VM->GetByteCode().GetNumInstructions() - 1, UnitNode, CallPath);
+		WorkData.VM->GetByteCode().SetSubject(WorkData.VM->GetByteCode().GetNumInstructions() - 1, Callstack.GetCallPath(), Callstack.GetStack());
 	}
 
 	// jump to the beginning of the loop
@@ -621,7 +653,7 @@ void URigVMCompiler::TraverseForLoop(const FRigVMCallExternExprAST* InExpr, FRig
 	WorkData.VM->GetByteCode().AddJumpOp(ERigVMOpCode::JumpBackward, JumpToStartInstruction - ForLoopInstructionIndex);
 	if (Settings.SetupNodeInstructionIndex)
 	{
-		WorkData.VM->GetByteCode().SetSubject(WorkData.VM->GetByteCode().GetNumInstructions() - 1, UnitNode, CallPath);
+		WorkData.VM->GetByteCode().SetSubject(WorkData.VM->GetByteCode().GetNumInstructions() - 1, Callstack.GetCallPath(), Callstack.GetStack());
 	}
 
 	// update the jump operator with the right address
@@ -794,8 +826,8 @@ void URigVMCompiler::TraverseAssign(const FRigVMAssignExprAST* InExpr, FRigVMCom
 				{
 					if (URigVMVariableNode* VariableNode = Cast<URigVMVariableNode>(SourcePin->GetNode()))
 					{
-						const FString CallPath = SourceExpr->GetProxy().GetSibling(VariableNode).GetCallstack().GetCallPath();
-						WorkData.VM->GetByteCode().SetSubject(InstructionIndex, VariableNode, CallPath);
+						const FRigVMCallstack Callstack = SourceExpr->GetProxy().GetSibling(VariableNode).GetCallstack();
+						WorkData.VM->GetByteCode().SetSubject(InstructionIndex, Callstack.GetCallPath(), Callstack.GetStack());
 					}
 				}
 			}
@@ -805,8 +837,8 @@ void URigVMCompiler::TraverseAssign(const FRigVMAssignExprAST* InExpr, FRigVMCom
 				{
 					if (URigVMVariableNode* VariableNode = Cast<URigVMVariableNode>(TargetPin->GetNode()))
 					{
-						const FString CallPath = TargetExpr->GetProxy().GetSibling(VariableNode).GetCallstack().GetCallPath();
-						WorkData.VM->GetByteCode().SetSubject(InstructionIndex, VariableNode, CallPath);
+						const FRigVMCallstack Callstack = TargetExpr->GetProxy().GetSibling(VariableNode).GetCallstack();
+						WorkData.VM->GetByteCode().SetSubject(InstructionIndex, Callstack.GetCallPath(), Callstack.GetStack());
 					}
 				}
 			}
@@ -844,7 +876,7 @@ void URigVMCompiler::TraverseBranch(const FRigVMBranchExprAST* InExpr, FRigVMCom
 	}
 
 	URigVMBranchNode* BranchNode = Cast<URigVMBranchNode>(InExpr->GetNode());
-	const FString CallPath = InExpr->GetProxy().GetCallstack().GetCallPath();
+	const FRigVMCallstack Callstack = InExpr->GetProxy().GetCallstack();
 
 	const FRigVMVarExprAST* ExecuteContextExpr = InExpr->ChildAt<FRigVMVarExprAST>(0);
 	const FRigVMVarExprAST* ConditionExpr = InExpr->ChildAt<FRigVMVarExprAST>(1);
@@ -866,7 +898,7 @@ void URigVMCompiler::TraverseBranch(const FRigVMBranchExprAST* InExpr, FRigVMCom
 	int32 JumpToFalseInstruction = WorkData.VM->GetByteCode().GetNumInstructions() - 1;
 	if (Settings.SetupNodeInstructionIndex)
 	{
-		WorkData.VM->GetByteCode().SetSubject(WorkData.VM->GetByteCode().GetNumInstructions() - 1, BranchNode, CallPath);
+		WorkData.VM->GetByteCode().SetSubject(WorkData.VM->GetByteCode().GetNumInstructions() - 1, Callstack.GetCallPath(), Callstack.GetStack());
 	}
 
 	// traverse the true case
@@ -876,7 +908,7 @@ void URigVMCompiler::TraverseBranch(const FRigVMBranchExprAST* InExpr, FRigVMCom
 	int32 JumpToEndInstruction = WorkData.VM->GetByteCode().GetNumInstructions() - 1;
 	if (Settings.SetupNodeInstructionIndex)
 	{
-		WorkData.VM->GetByteCode().SetSubject(WorkData.VM->GetByteCode().GetNumInstructions() - 1, BranchNode, CallPath);
+		WorkData.VM->GetByteCode().SetSubject(WorkData.VM->GetByteCode().GetNumInstructions() - 1, Callstack.GetCallPath(), Callstack.GetStack());
 	}
 
 	// correct the jump to false instruction index
@@ -902,7 +934,7 @@ void URigVMCompiler::TraverseIf(const FRigVMIfExprAST* InExpr, FRigVMCompilerWor
 	}
 
 	URigVMIfNode* IfNode = Cast<URigVMIfNode>(InExpr->GetNode());
-	const FString CallPath = InExpr->GetProxy().GetCallstack().GetCallPath();
+	const FRigVMCallstack Callstack = InExpr->GetProxy().GetCallstack();
 
 	const FRigVMVarExprAST* ConditionExpr = InExpr->ChildAt<FRigVMVarExprAST>(0);
 	const FRigVMVarExprAST* TrueExpr = InExpr->ChildAt<FRigVMVarExprAST>(1);
@@ -925,7 +957,7 @@ void URigVMCompiler::TraverseIf(const FRigVMIfExprAST* InExpr, FRigVMCompilerWor
 	int32 JumpToFalseInstruction = WorkData.VM->GetByteCode().GetNumInstructions() - 1;
 	if (Settings.SetupNodeInstructionIndex)
 	{
-		WorkData.VM->GetByteCode().SetSubject(WorkData.VM->GetByteCode().GetNumInstructions() - 1, IfNode, CallPath);
+		WorkData.VM->GetByteCode().SetSubject(WorkData.VM->GetByteCode().GetNumInstructions() - 1, Callstack.GetCallPath(), Callstack.GetStack());
 	}
 
 	// traverse the true case
@@ -941,14 +973,14 @@ void URigVMCompiler::TraverseIf(const FRigVMIfExprAST* InExpr, FRigVMCompilerWor
 	WorkData.VM->GetByteCode().AddCopyOp(TrueOperand, ResultOperand);
 	if (Settings.SetupNodeInstructionIndex)
 	{
-		WorkData.VM->GetByteCode().SetSubject(WorkData.VM->GetByteCode().GetNumInstructions() - 1, IfNode, CallPath);
+		WorkData.VM->GetByteCode().SetSubject(WorkData.VM->GetByteCode().GetNumInstructions() - 1, Callstack.GetCallPath(), Callstack.GetStack());
 	}
 
 	uint64 JumpToEndByte = WorkData.VM->GetByteCode().AddJumpOp(ERigVMOpCode::JumpForward, 1);
 	int32 JumpToEndInstruction = WorkData.VM->GetByteCode().GetNumInstructions() - 1;
 	if (Settings.SetupNodeInstructionIndex)
 	{
-		WorkData.VM->GetByteCode().SetSubject(WorkData.VM->GetByteCode().GetNumInstructions() - 1, IfNode, CallPath);
+		WorkData.VM->GetByteCode().SetSubject(WorkData.VM->GetByteCode().GetNumInstructions() - 1, Callstack.GetCallPath(), Callstack.GetStack());
 	}
 
 	// correct the jump to false instruction index
@@ -968,7 +1000,7 @@ void URigVMCompiler::TraverseIf(const FRigVMIfExprAST* InExpr, FRigVMCompilerWor
 	WorkData.VM->GetByteCode().AddCopyOp(FalseOperand, ResultOperand);
 	if (Settings.SetupNodeInstructionIndex)
 	{
-		WorkData.VM->GetByteCode().SetSubject(WorkData.VM->GetByteCode().GetNumInstructions() - 1, IfNode, CallPath);
+		WorkData.VM->GetByteCode().SetSubject(WorkData.VM->GetByteCode().GetNumInstructions() - 1, Callstack.GetCallPath(), Callstack.GetStack());
 	}
 
 	// correct the jump to end instruction index
@@ -979,7 +1011,7 @@ void URigVMCompiler::TraverseIf(const FRigVMIfExprAST* InExpr, FRigVMCompilerWor
 void URigVMCompiler::TraverseSelect(const FRigVMSelectExprAST* InExpr, FRigVMCompilerWorkData& WorkData)
 {
 	URigVMSelectNode* SelectNode = Cast<URigVMSelectNode>(InExpr->GetNode());
-	const FString CallPath = InExpr->GetProxy().GetCallstack().GetCallPath();
+	const FRigVMCallstack Callstack = InExpr->GetProxy().GetCallstack();
 
 	int32 NumCases = SelectNode->FindPin(URigVMSelectNode::ValueName)->GetArraySize();
 
@@ -1057,14 +1089,14 @@ void URigVMCompiler::TraverseSelect(const FRigVMSelectExprAST* InExpr, FRigVMCom
 		WorkData.VM->GetByteCode().AddEqualsOp(IndexOperand, WorkData.IntegerLiterals.FindChecked(CaseIndex), WorkData.ComparisonOperand);
 		if (Settings.SetupNodeInstructionIndex)
 		{
-			WorkData.VM->GetByteCode().SetSubject(WorkData.VM->GetByteCode().GetNumInstructions() - 1, SelectNode, CallPath);
+			WorkData.VM->GetByteCode().SetSubject(WorkData.VM->GetByteCode().GetNumInstructions() - 1, Callstack.GetCallPath(), Callstack.GetStack());
 		}
 
 		uint64 JumpByte = WorkData.VM->GetByteCode().AddJumpIfOp(ERigVMOpCode::JumpForwardIf, 1, WorkData.ComparisonOperand, true);
 		int32 JumpInstruction = WorkData.VM->GetByteCode().GetNumInstructions() - 1;
 		if (Settings.SetupNodeInstructionIndex)
 		{
-			WorkData.VM->GetByteCode().SetSubject(WorkData.VM->GetByteCode().GetNumInstructions() - 1, SelectNode, CallPath);
+			WorkData.VM->GetByteCode().SetSubject(WorkData.VM->GetByteCode().GetNumInstructions() - 1, Callstack.GetCallPath(), Callstack.GetStack());
 		}
 
 		JumpToCaseBytes.Add(JumpByte);
@@ -1089,7 +1121,7 @@ void URigVMCompiler::TraverseSelect(const FRigVMSelectExprAST* InExpr, FRigVMCom
 		WorkData.VM->GetByteCode().AddCopyOp(CaseOperand, ResultOperand);
 		if (Settings.SetupNodeInstructionIndex)
 		{
-			WorkData.VM->GetByteCode().SetSubject(WorkData.VM->GetByteCode().GetNumInstructions() - 1, SelectNode, CallPath);
+			WorkData.VM->GetByteCode().SetSubject(WorkData.VM->GetByteCode().GetNumInstructions() - 1, Callstack.GetCallPath(), Callstack.GetStack());
 		}
 		
 		if (CaseIndex < NumCases - 1)
@@ -1098,7 +1130,7 @@ void URigVMCompiler::TraverseSelect(const FRigVMSelectExprAST* InExpr, FRigVMCom
 			int32 JumpInstruction = WorkData.VM->GetByteCode().GetNumInstructions() - 1;
 			if (Settings.SetupNodeInstructionIndex)
 			{
-				WorkData.VM->GetByteCode().SetSubject(WorkData.VM->GetByteCode().GetNumInstructions() - 1, SelectNode, CallPath);
+				WorkData.VM->GetByteCode().SetSubject(WorkData.VM->GetByteCode().GetNumInstructions() - 1, Callstack.GetCallPath(), Callstack.GetStack());
 			}
 
 			JumpToEndBytes.Add(JumpByte);
@@ -1113,6 +1145,164 @@ void URigVMCompiler::TraverseSelect(const FRigVMSelectExprAST* InExpr, FRigVMCom
 	}
 }
 
+void URigVMCompiler::InitializeLocalVariables(const FRigVMExprAST* InExpr, FRigVMCompilerWorkData& WorkData)
+{
+	// Initialize local variables if we are entering a new graph
+	if (!WorkData.bSetupMemory)
+	{
+		FRigVMByteCode& ByteCode = WorkData.VM->GetByteCode();
+		FRigVMCallstack Callstack;
+		switch (InExpr->GetType())
+		{
+			case FRigVMExprAST::EType::CallExtern:
+			{
+				Callstack = InExpr->To<FRigVMCallExternExprAST>()->GetProxy().GetCallstack();
+				break;
+			}
+			case FRigVMExprAST::EType::NoOp:
+			{
+				Callstack = InExpr->To<FRigVMNoOpExprAST>()->GetProxy().GetCallstack();
+				break;
+			}
+			case FRigVMExprAST::EType::Var:
+			{
+				Callstack = InExpr->To<FRigVMVarExprAST>()->GetProxy().GetCallstack();
+				break;
+			}
+			case FRigVMExprAST::EType::Literal:
+			{
+				Callstack = InExpr->To<FRigVMLiteralExprAST>()->GetProxy().GetCallstack();
+				break;
+			}
+			case FRigVMExprAST::EType::ExternalVar:
+			{
+				Callstack = InExpr->To<FRigVMExternalVarExprAST>()->GetProxy().GetCallstack();
+				break;
+			}
+			case FRigVMExprAST::EType::Branch:
+			{
+				Callstack = InExpr->To<FRigVMBranchExprAST>()->GetProxy().GetCallstack();
+				break;
+			}
+			case FRigVMExprAST::EType::If:
+			{
+				Callstack = InExpr->To<FRigVMIfExprAST>()->GetProxy().GetCallstack();
+				break;
+			}
+			case FRigVMExprAST::EType::Select:
+			{
+				Callstack = InExpr->To<FRigVMSelectExprAST>()->GetProxy().GetCallstack();
+				break;
+			}
+		}
+
+		if (Callstack.Num() > 0)
+		{
+			// Look at previous instructions to get the last callstack
+			const TArray<UObject*>* PreviousCallstack = nullptr;
+			int32 PreviousIndex = ByteCode.GetNumInstructions()-1;
+			TArray<FRigVMCopyOp> PreviousCopyInstructions;
+			while (PreviousCallstack == nullptr)
+			{
+				if (PreviousIndex < 0)
+				{
+					break;
+				}
+				PreviousCallstack = ByteCode.GetCallstackForInstruction(PreviousIndex);
+
+				// Other recursive calls might have already initialized local variables, we don't want
+				// to initialize them again
+				uint64 ByteCodeIndex = ByteCode.GetInstructions()[PreviousIndex].ByteCodeIndex;
+				ERigVMOpCode Code = ByteCode.GetOpCodeAt(ByteCodeIndex);
+				if (ByteCode.GetOpCodeAt(ByteCodeIndex) == ERigVMOpCode::Copy)
+				{
+					const FRigVMCopyOp& CopyOp = ByteCode.GetOpAt<FRigVMCopyOp>(ByteCodeIndex);
+					PreviousCopyInstructions.Add(CopyOp);
+				}					
+			
+				PreviousIndex--;
+			}
+
+			int32 FirstDifferenceIndex = 0;
+			if (PreviousCallstack)
+			{
+				// Compare the two callstacks to find the first difference in the path				
+				for (; FirstDifferenceIndex < PreviousCallstack->Num() && FirstDifferenceIndex < Callstack.Num(); ++FirstDifferenceIndex)
+				{
+					if (PreviousCallstack->operator[](FirstDifferenceIndex) != Callstack[FirstDifferenceIndex])
+					{
+						break;
+					}
+				}
+
+				// Starting from the first difference in the path, look for function references or collapse nodes
+				// to initialize their local variables. 
+				for (int32 Index = FirstDifferenceIndex; Index < Callstack.Num(); ++Index)
+				{
+					// If it is a library node (which will handle function references and collapsed nodes), try to initialize local variables
+					const URigVMLibraryNode* Node = Cast<const URigVMLibraryNode>(Callstack[Index]);
+					if (Node)
+					{
+						for (FRigVMGraphVariableDescription Variable : Node->GetContainedGraph()->LocalVariables)
+						{
+							FString TargetPath = FString::Printf(TEXT("LocalVariable::%s|%s"), *Node->GetNodePath(), *Variable.Name.ToString());
+							FString SourcePath = FString::Printf(TEXT("LocalVariable::%s|%s::Const"), *Node->GetContainedGraph()->GetGraphName(), *Variable.Name.ToString());
+							FRigVMOperand* Target = WorkData.PinPathToOperand->Find(TargetPath);
+							FRigVMOperand* Source = WorkData.PinPathToOperand->Find(SourcePath);
+							if (Source && Target) 
+							{
+								bool bAlreadyCopied = false;
+								for (const FRigVMCopyOp& CopyOp : PreviousCopyInstructions)
+								{
+									if (CopyOp.Source == *Source && CopyOp.Target == *Target)
+									{
+										bAlreadyCopied = true;
+										break;
+									}
+								}
+
+								if (!bAlreadyCopied)
+								{
+									ByteCode.AddCopyOp(*Source, *Target);
+								}	
+							}					
+						}
+					}
+				}
+			}
+		}
+
+		// We also need to initialize local variables of the root graph (after the entry instruction is inserted)
+		// If the parent of InExpr is an entry, initialize the root graph's local variables
+		if(InExpr->NumParents() > 0 && InExpr->ParentAt(0)->GetType() == FRigVMExprAST::Entry)
+		{
+			if (UObject* Subject = ByteCode.GetSubjectForInstruction(ByteCode.GetNumInstructions()-1))
+			{
+				URigVMNode* Node = Cast<URigVMNode>(Subject);
+				if (!Node)
+				{
+					URigVMPin* Pin = Cast<URigVMPin>(Subject);
+					Node = Pin->GetNode();
+				}
+				if (Node)
+				{
+					for (FRigVMGraphVariableDescription Variable : Node->GetGraph()->LocalVariables)
+					{
+						FString TargetPath = FString::Printf(TEXT("LocalVariable::%s"), *Variable.Name.ToString());
+						FString SourcePath = FString::Printf(TEXT("LocalVariable::|%s::Const"), *Variable.Name.ToString());
+						FRigVMOperand* Target = WorkData.PinPathToOperand->Find(TargetPath);
+						FRigVMOperand* Source = WorkData.PinPathToOperand->Find(SourcePath);
+						if (Source && Target)
+						{
+							ByteCode.AddCopyOp(*Source, *Target);
+						}					
+					}
+				}
+			}
+		}
+	}
+}
+
 FString URigVMCompiler::GetPinHash(const URigVMPin* InPin, const FRigVMVarExprAST* InVarExpr, bool bIsDebugValue)
 {
 	FString Prefix = bIsDebugValue ? TEXT("DebugWatch:") : TEXT("");
@@ -1167,31 +1357,71 @@ FString URigVMCompiler::GetPinHash(const URigVMPin* InPin, const FRigVMVarExprAS
 	}
 	else if (URigVMVariableNode* VariableNode = Cast<URigVMVariableNode>(Node))
 	{
-		if (InPin->GetName() == TEXT("Value") && !bIsLiteral && !bIsDebugValue)
+		if (InPin->GetName() == TEXT("Value") && !bIsDebugValue)
 		{
 			FName VariableName = VariableNode->GetVariableName();
-			
-			// determine if this variable needs to be remapped
-			if(InVarExpr)
+
+			// Figure out if it is a local variable
 			{
-				FRigVMASTProxy ParentProxy = InVarExpr->GetProxy();
-				while(ParentProxy.GetCallstack().Num() > 1)
+				URigVMGraph* RootGraph = InPin->GetNode()->GetGraph();
+				TArray<FRigVMGraphVariableDescription>& LocalVariables = RootGraph->LocalVariables;
+				for (const FRigVMGraphVariableDescription& LocalVariable : LocalVariables)
 				{
-					ParentProxy = ParentProxy.GetParent();
-
-					if(URigVMFunctionReferenceNode* FunctionReferenceNode = ParentProxy.GetSubject<URigVMFunctionReferenceNode>())
+					if (LocalVariable.Name == VariableName)
 					{
-						const FName RemappedVariableName = FunctionReferenceNode->GetOuterVariableName(VariableName);
-						if(!RemappedVariableName.IsNone())
+						if (bIsLiteral)
+						{
+							// Literal values will be reused for all instance of local variables
+							return FString::Printf(TEXT("%sLocalVariable::%s|%s%s"), *Prefix, *RootGraph->GetGraphName(), *VariableName.ToString(), *Suffix);
+						}
+						else
 						{
-							VariableName = RemappedVariableName;
+							if(InVarExpr)
+							{
+								FRigVMASTProxy ParentProxy = InVarExpr->GetProxy();
+								while(ParentProxy.GetCallstack().Num() > 1)
+								{
+									ParentProxy = ParentProxy.GetParent();
+
+									if(URigVMLibraryNode* LibraryNode = ParentProxy.GetSubject<URigVMLibraryNode>())
+									{
+										// Local variables for non-root graphs are in the format "LocalVariable::PathToGraph|VariableName"
+										return FString::Printf(TEXT("%sLocalVariable::%s|%s%s"), *Prefix, *LibraryNode->GetNodePath(true), *VariableName.ToString(), *Suffix);
+									}
+								}
+
+								// Local variables for root graphs are in the format "LocalVariable::VariableName"
+								return FString::Printf(TEXT("%sLocalVariable::%s%s"), *Prefix, *VariableName.ToString(), *Suffix);
+							}
 						}
 					}
 				}
 			}
+
+			if (!bIsLiteral)
+			{		
+				// determine if this variable needs to be remapped
+				if(InVarExpr)
+				{
+					FRigVMASTProxy ParentProxy = InVarExpr->GetProxy();
+					while(ParentProxy.GetCallstack().Num() > 1)
+					{
+						ParentProxy = ParentProxy.GetParent();
+
+						if(URigVMFunctionReferenceNode* FunctionReferenceNode = ParentProxy.GetSubject<URigVMFunctionReferenceNode>())
+						{
+							const FName RemappedVariableName = FunctionReferenceNode->GetOuterVariableName(VariableName);
+							if(!RemappedVariableName.IsNone())
+							{
+								VariableName = RemappedVariableName;
+							}
+						}
+					}
+				}
 			
-			return FString::Printf(TEXT("%sVariable::%s%s"), *Prefix, *VariableName.ToString(), *Suffix);
-		}
+				return FString::Printf(TEXT("%sVariable::%s%s"), *Prefix, *VariableName.ToString(), *Suffix);
+			}
+		}		
 	}
 	else
 	{
@@ -1418,7 +1648,7 @@ FRigVMOperand URigVMCompiler::FindOrAddRegister(const FRigVMVarExprAST* InVarExp
 	bool bIsVariable = Pin->IsRootPin() && (Pin->GetName() == URigVMVariableNode::ValueName) &&
 		InVarExpr->GetPin()->GetNode()->IsA<URigVMVariableNode>();
 
-	// variables don't require to add any register.
+	// external variables don't require to add any register.
 	if(bIsVariable && !bIsDebugValue)
 	{
 		for(int32 ExternalVariableIndex = 0; ExternalVariableIndex < WorkData.VM->GetExternalVariables().Num(); ExternalVariableIndex++)
diff --git a/Engine/Source/Developer/RigVMDeveloper/Private/RigVMModel/RigVMController.cpp b/Engine/Source/Developer/RigVMDeveloper/Private/RigVMModel/RigVMController.cpp
--- a/Engine/Source/Developer/RigVMDeveloper/Private/RigVMModel/RigVMController.cpp
+++ b/Engine/Source/Developer/RigVMDeveloper/Private/RigVMModel/RigVMController.cpp
@@ -418,27 +418,42 @@ URigVMVariableNode* URigVMController::AddVariableNode(const FName& InVariableNam
 		{
 			if(URigVMLibraryNode* OuterFunction = OuterLibrary->FindFunctionForNode(Graph->GetTypedOuter<URigVMCollapseNode>()))
 			{
-				TArray<FRigVMExternalVariable> ExternalVariables = OuterFunction->GetContainedGraph()->GetExternalVariables();
-
-				bool bFoundExternalVariable = false;
-				for(const FRigVMExternalVariable& ExternalVariable : ExternalVariables)
+				// Make sure there is no local variable with that name
+				bool bFoundLocalVariable = false;
+				for (FRigVMGraphVariableDescription& LocalVariable : OuterFunction->GetContainedGraph()->LocalVariables)
 				{
-					if(ExternalVariable.Name == InVariableName)
+					if (LocalVariable.Name == InVariableName)
 					{
-						bFoundExternalVariable = true;
+						bFoundLocalVariable = true;
 						break;
 					}
 				}
 
-				if(!bFoundExternalVariable)
+				if (!bFoundLocalVariable)
 				{
-					if(RequestBulkEditDialogDelegate.IsBound())
+					// Make sure there is no external variable with that name
+					TArray<FRigVMExternalVariable> ExternalVariables = OuterFunction->GetContainedGraph()->GetExternalVariables();
+					bool bFoundExternalVariable = false;
+					for(const FRigVMExternalVariable& ExternalVariable : ExternalVariables)
 					{
-						if(!RequestBulkEditDialogDelegate.Execute(OuterFunction, ERigVMControllerBulkEditType::AddVariable))
+						if(ExternalVariable.Name == InVariableName)
 						{
-							return nullptr;
+							bFoundExternalVariable = true;
+							break;
+						}
+					}
+
+					if(!bFoundExternalVariable)
+					{
+						// Warn the user the changes are not undoable
+						if(RequestBulkEditDialogDelegate.IsBound())
+						{
+							if(!RequestBulkEditDialogDelegate.Execute(OuterFunction, ERigVMControllerBulkEditType::AddVariable))
+							{
+								return nullptr;
+							}
+							bSetupUndoRedo = false;
 						}
-						bSetupUndoRedo = false;
 					}
 				}
 			}
@@ -6618,6 +6633,103 @@ bool URigVMController::RemoveFunctionFromLibrary(const FName& InFunctionName, bo
 	return RemoveNodeByName(InFunctionName, bSetupUndoRedo);
 }
 
+bool URigVMController::AddLocalVariable(const FName& InVariableName,
+	const FString& InCPPType, UObject* InCPPTypeObject, const FString& InDefaultValue, bool bSetupUndoRedo)
+{
+	if (!IsValidGraph())
+	{
+		return false;
+	}
+	
+	URigVMGraph* Graph = GetGraph();
+	check(Graph);
+
+	if (Graph->GetParentGraph() && !Graph->GetParentGraph()->IsA<URigVMFunctionLibrary>())
+	{
+		ReportError(TEXT("Can only add local variables to root graphs."));
+		return false;
+	}
+
+	TArray<FRigVMGraphVariableDescription>& LocalVariables = Graph->LocalVariables;
+	for (FRigVMGraphVariableDescription& Variable : LocalVariables)
+	{
+		if (Variable.Name == InVariableName)
+		{
+			return false;
+		}
+	}
+
+	FRigVMGraphVariableDescription NewVariable;
+	NewVariable.Name = InVariableName;
+	NewVariable.CPPType = InCPPType;
+	NewVariable.CPPTypeObject = InCPPTypeObject;
+	NewVariable.DefaultValue = InDefaultValue;
+
+	LocalVariables.Add(NewVariable);
+
+	if (bSetupUndoRedo)
+	{
+		FRigVMInverseAction InverseAction;
+		InverseAction.Title = FString::Printf(TEXT("Add Local Variable %s"), *InVariableName.ToString());
+
+		ActionStack->BeginAction(InverseAction);
+		ActionStack->AddAction(FRigVMRemoveLocalVariableAction(NewVariable));
+		ActionStack->EndAction(InverseAction);
+	}
+
+	if (!bSuspendNotifications)
+	{
+		Graph->MarkPackageDirty();
+	}
+
+	return true;
+}
+
+bool URigVMController::RemoveLocalVariable(const FName& InVariableName, bool bSetupUndoRedo)
+{
+	if (!IsValidGraph())
+	{
+		return false;
+	}
+
+	URigVMGraph* Graph = GetGraph();
+	check(Graph);
+
+	TArray<FRigVMGraphVariableDescription>& LocalVariables = Graph->LocalVariables;
+	int32 FoundIndex = INDEX_NONE;
+	for (int32 Index = 0; Index < LocalVariables.Num(); ++Index)
+	{
+		if (LocalVariables[Index].Name == InVariableName)
+		{
+			FoundIndex = Index;
+			break;
+		}
+	}
+
+	if (FoundIndex != INDEX_NONE)
+	{
+		if (!bSuspendNotifications)
+		{
+			Graph->MarkPackageDirty();
+		}
+		
+		if (bSetupUndoRedo)
+		{
+			FRigVMInverseAction InverseAction;
+			InverseAction.Title = FString::Printf(TEXT("Remove Local Variable %s"), *InVariableName.ToString());
+
+			ActionStack->BeginAction(InverseAction);
+			ActionStack->AddAction(FRigVMAddLocalVariableAction(LocalVariables[FoundIndex]));
+			ActionStack->EndAction(InverseAction);
+		}	
+		
+		LocalVariables.RemoveAt(FoundIndex);
+		return true;
+	}
+
+	return false;
+}
+
 TArray<TSoftObjectPtr<URigVMFunctionReferenceNode>> URigVMController::GetAffectedReferences(ERigVMControllerBulkEditType InEditType, bool bForceLoad, bool bNotify)
 {
 	TArray<TSoftObjectPtr<URigVMFunctionReferenceNode>> FunctionReferencePtrs;
diff --git a/Engine/Source/Developer/RigVMDeveloper/Private/RigVMModel/RigVMControllerActions.cpp b/Engine/Source/Developer/RigVMDeveloper/Private/RigVMModel/RigVMControllerActions.cpp
--- a/Engine/Source/Developer/RigVMDeveloper/Private/RigVMModel/RigVMControllerActions.cpp
+++ b/Engine/Source/Developer/RigVMDeveloper/Private/RigVMModel/RigVMControllerActions.cpp
@@ -1728,3 +1728,50 @@ bool FRigVMSetRemappedVariableAction::Redo(URigVMController* InController)
 	}
 	return FRigVMBaseAction::Redo(InController);
 }
+
+FRigVMAddLocalVariableAction::FRigVMAddLocalVariableAction(const FRigVMGraphVariableDescription& InLocalVariable)
+	: LocalVariable(InLocalVariable)
+{
+	
+}
+
+bool FRigVMAddLocalVariableAction::Undo(URigVMController* InController)
+{
+	if (!FRigVMBaseAction::Undo(InController))
+	{
+		return false;
+	}
+	return InController->RemoveLocalVariable(LocalVariable.Name);
+}
+
+bool FRigVMAddLocalVariableAction::Redo(URigVMController* InController)
+{
+	if (!LocalVariable.Name.IsNone())
+	{
+		return InController->AddLocalVariable(LocalVariable.Name, LocalVariable.CPPType, LocalVariable.CPPTypeObject, LocalVariable.DefaultValue);
+	}
+	return FRigVMBaseAction::Redo(InController);
+}
+
+FRigVMRemoveLocalVariableAction::FRigVMRemoveLocalVariableAction(const FRigVMGraphVariableDescription& InLocalVariable)
+	: LocalVariable(InLocalVariable)
+{
+}
+
+bool FRigVMRemoveLocalVariableAction::Undo(URigVMController* InController)
+{
+	if (!FRigVMBaseAction::Undo(InController))
+	{
+		return false;
+	}
+	return InController->AddLocalVariable(LocalVariable.Name, LocalVariable.CPPType, LocalVariable.CPPTypeObject, LocalVariable.DefaultValue, false);
+}
+
+bool FRigVMRemoveLocalVariableAction::Redo(URigVMController* InController)
+{
+	if (!LocalVariable.Name.IsNone())
+	{
+		return InController->RemoveLocalVariable(LocalVariable.Name, false);
+	}
+	return FRigVMBaseAction::Redo(InController);
+}
\ No newline at end of file
diff --git a/Engine/Source/Developer/RigVMDeveloper/Private/RigVMModel/RigVMGraph.cpp b/Engine/Source/Developer/RigVMDeveloper/Private/RigVMModel/RigVMGraph.cpp
--- a/Engine/Source/Developer/RigVMDeveloper/Private/RigVMModel/RigVMGraph.cpp
+++ b/Engine/Source/Developer/RigVMDeveloper/Private/RigVMModel/RigVMGraph.cpp
@@ -242,6 +242,40 @@ void URigVMGraph::SetDefaultFunctionLibrary(URigVMFunctionLibrary* InFunctionLib
 	DefaultFunctionLibraryPtr = InFunctionLibrary;
 }
 
+bool URigVMGraph::AddLocalVariable(const FRigVMGraphVariableDescription& NewVar)
+{
+	for (FRigVMGraphVariableDescription Variable : LocalVariables)
+	{
+		if (Variable.Name == NewVar.Name)
+		{
+			return false;
+		}
+	}
+	
+	LocalVariables.Add(NewVar);
+	return true;
+}
+
+bool URigVMGraph::RemoveLocalVariable(const FName& InVariableName)
+{
+	int32 FoundIndex = INDEX_NONE;
+	for (int32 Index = 0; Index < LocalVariables.Num(); ++Index)
+	{
+		if (LocalVariables[Index].Name == InVariableName)
+		{
+			FoundIndex = Index;
+			break;
+		}
+	}
+
+	if (FoundIndex != INDEX_NONE)
+	{
+		LocalVariables.RemoveAt(FoundIndex);
+		return true;
+	}
+	return false;
+}
+
 TArray<FRigVMExternalVariable> URigVMGraph::GetExternalVariables() const
 {
 	TArray<FRigVMExternalVariable> Variables;
@@ -258,7 +292,21 @@ TArray<FRigVMExternalVariable> URigVMGraph::GetExternalVariables() const
 		}
 		else if(URigVMVariableNode* VariableNode = Cast<URigVMVariableNode>(Node))
 		{
-			FRigVMExternalVariable::MergeExternalVariable(Variables, VariableNode->GetVariableDescription().ToExternalVariable());
+			// Make sure it is not a local variable
+			bool bFoundLocalVariable = false;
+			for (FRigVMGraphVariableDescription& LocalVariable : Node->GetGraph()->LocalVariables)
+			{
+				if (LocalVariable.Name == VariableNode->GetVariableName())
+				{
+					bFoundLocalVariable = true;
+					break;
+				}
+			}
+
+			if (!bFoundLocalVariable)
+			{
+				FRigVMExternalVariable::MergeExternalVariable(Variables, VariableNode->GetVariableDescription().ToExternalVariable());
+			}
 		}
 	}
 	
diff --git a/Engine/Source/Developer/RigVMDeveloper/Public/RigVMCompiler/RigVMASTProxy.h b/Engine/Source/Developer/RigVMDeveloper/Public/RigVMCompiler/RigVMASTProxy.h
--- a/Engine/Source/Developer/RigVMDeveloper/Public/RigVMCompiler/RigVMASTProxy.h
+++ b/Engine/Source/Developer/RigVMDeveloper/Public/RigVMCompiler/RigVMASTProxy.h
@@ -27,7 +27,7 @@ public:
 	const UObject* operator[](int32 InIndex) const;
 
 	// Returns true if the stack contains a given entry
-	bool Contains(UObject* InEntry) const;
+	bool Contains(const UObject* InEntry) const;
 
 	friend FORCEINLINE uint32 GetTypeHash(const FRigVMCallstack& Callstack)
 	{
@@ -95,6 +95,8 @@ public:
 		return operator[](0) > Other[0];
 	}
 
+	FORCEINLINE TArray<UObject*> GetStack() const { return Stack; }
+
 private:
 
 	FORCEINLINE uint32 GetEntryTypeHash(int32 Index) const
diff --git a/Engine/Source/Developer/RigVMDeveloper/Public/RigVMCompiler/RigVMCompiler.h b/Engine/Source/Developer/RigVMDeveloper/Public/RigVMCompiler/RigVMCompiler.h
--- a/Engine/Source/Developer/RigVMDeveloper/Public/RigVMCompiler/RigVMCompiler.h
+++ b/Engine/Source/Developer/RigVMDeveloper/Public/RigVMCompiler/RigVMCompiler.h
@@ -147,6 +147,8 @@ private:
 	void TraverseIf(const FRigVMIfExprAST* InExpr, FRigVMCompilerWorkData& WorkData);
 	void TraverseSelect(const FRigVMSelectExprAST* InExpr, FRigVMCompilerWorkData& WorkData);
 
+	void InitializeLocalVariables(const FRigVMExprAST* InExpr, FRigVMCompilerWorkData& WorkData);
+
 	FRigVMOperand FindOrAddRegister(const FRigVMVarExprAST* InVarExpr, FRigVMCompilerWorkData& WorkData, bool bIsDebugValue = false);
 
 	void ReportInfo(const FString& InMessage);
diff --git a/Engine/Source/Developer/RigVMDeveloper/Public/RigVMModel/RigVMController.h b/Engine/Source/Developer/RigVMDeveloper/Public/RigVMModel/RigVMController.h
--- a/Engine/Source/Developer/RigVMDeveloper/Public/RigVMModel/RigVMController.h
+++ b/Engine/Source/Developer/RigVMDeveloper/Public/RigVMModel/RigVMController.h
@@ -582,6 +582,14 @@ public:
 	UFUNCTION(BlueprintCallable, Category = RigVMController)
 	bool RemoveFunctionFromLibrary(const FName& InFunctionName, bool bSetupUndoRedo = true);
 
+	// Add a local variable to the graph
+	UFUNCTION(BlueprintCallable, Category = RigVMController)
+	bool AddLocalVariable(const FName& InVariableName, const FString& InCPPType, UObject* InCPPTypeObject, const FString& InDefaultValue, bool bSetupUndoRedo = true);
+
+	// Remove a local variable from the graph
+	UFUNCTION(BlueprintCallable, Category = RigVMController)
+	bool RemoveLocalVariable(const FName& InVariableName, bool bSetupUndoRedo = true);
+
 	// Determine affected function references for a potential bulk edit on a library node
 	TArray<TSoftObjectPtr<URigVMFunctionReferenceNode>> GetAffectedReferences(ERigVMControllerBulkEditType InEditType, bool bForceLoad = false, bool bNotify = true);
 
diff --git a/Engine/Source/Developer/RigVMDeveloper/Public/RigVMModel/RigVMControllerActions.h b/Engine/Source/Developer/RigVMDeveloper/Public/RigVMModel/RigVMControllerActions.h
--- a/Engine/Source/Developer/RigVMDeveloper/Public/RigVMModel/RigVMControllerActions.h
+++ b/Engine/Source/Developer/RigVMDeveloper/Public/RigVMModel/RigVMControllerActions.h
@@ -1467,3 +1467,43 @@ public:
 	UPROPERTY()
 	FName NewOuterVariableName;
 };
+
+/**
+* An action to add a local variable.
+*/
+USTRUCT()
+struct FRigVMAddLocalVariableAction : public FRigVMBaseAction
+{
+	GENERATED_BODY()
+
+public:
+
+	FRigVMAddLocalVariableAction() {}
+	FRigVMAddLocalVariableAction(const FRigVMGraphVariableDescription& InLocalVariable);
+	virtual ~FRigVMAddLocalVariableAction() {};
+	virtual bool Undo(URigVMController* InController) override;
+	virtual bool Redo(URigVMController* InController) override;
+
+	UPROPERTY()
+	FRigVMGraphVariableDescription LocalVariable;
+};
+
+/**
+* An action to remove a local variable.
+*/
+USTRUCT()
+struct FRigVMRemoveLocalVariableAction : public FRigVMBaseAction
+{
+	GENERATED_BODY()
+
+public:
+
+	FRigVMRemoveLocalVariableAction() {}
+	FRigVMRemoveLocalVariableAction(const FRigVMGraphVariableDescription& InLocalVariable);
+	virtual ~FRigVMRemoveLocalVariableAction() {};
+	virtual bool Undo(URigVMController* InController) override;
+	virtual bool Redo(URigVMController* InController) override;
+
+	UPROPERTY()
+	FRigVMGraphVariableDescription LocalVariable;
+};
diff --git a/Engine/Source/Developer/RigVMDeveloper/Public/RigVMModel/RigVMGraph.h b/Engine/Source/Developer/RigVMDeveloper/Public/RigVMModel/RigVMGraph.h
--- a/Engine/Source/Developer/RigVMDeveloper/Public/RigVMModel/RigVMGraph.h
+++ b/Engine/Source/Developer/RigVMDeveloper/Public/RigVMModel/RigVMGraph.h
@@ -119,8 +119,18 @@ public:
 	UFUNCTION(BlueprintCallable, Category = RigVMGraph)
 	void SetDefaultFunctionLibrary(URigVMFunctionLibrary* InFunctionLibrary);
 
+	// Returns true if the variable was added to the graph's local variables
+	UFUNCTION(BlueprintCallable, Category = RigVMGraph)
+	bool AddLocalVariable(const FRigVMGraphVariableDescription& InVariable);
+
+	// Returns true if the variable was removed from the graph's local variables
+	UFUNCTION(BlueprintCallable, Category = RigVMGraph)
+	bool RemoveLocalVariable(const FName& InVariableName);
+
 	TArray<FRigVMExternalVariable> GetExternalVariables() const;
 
+	TArray<FRigVMGraphVariableDescription> GetLocalVariables() const { return LocalVariables; }
+
 	// Returns the modified event, which can be used to 
 	// subscribe to changes happening within the Graph.
 	FRigVMGraphModifiedEvent& OnModified();
@@ -157,6 +167,9 @@ private:
 	TArray<TSharedPtr<FString>> ParameterNames;
 #endif
 
+	UPROPERTY()
+	TArray<FRigVMGraphVariableDescription> LocalVariables;
+
 	bool IsNameAvailable(const FString& InName);
 
 	friend class URigVMController;
diff --git a/Engine/Source/Runtime/RigVM/Private/RigVMCore/RigVMByteCode.cpp b/Engine/Source/Runtime/RigVM/Private/RigVMCore/RigVMByteCode.cpp
--- a/Engine/Source/Runtime/RigVM/Private/RigVMCore/RigVMByteCode.cpp
+++ b/Engine/Source/Runtime/RigVM/Private/RigVMCore/RigVMByteCode.cpp
@@ -536,6 +536,7 @@ void FRigVMByteCode::Reset()
 	SubjectToInstructions.Reset();
 	CallPathPerInstruction.Reset();
 	CallPathToInstructions.Reset();
+	CallstackPerInstruction.Reset();
 #endif
 }
 
@@ -551,6 +552,7 @@ void FRigVMByteCode::Empty()
 	SubjectToInstructions.Empty();
 	CallPathPerInstruction.Empty();
 	CallPathToInstructions.Empty();
+	CallstackPerInstruction.Empty();
 #endif
 }
 
@@ -1327,14 +1329,15 @@ TArray<int32> FRigVMByteCode::GetAllInstructionIndicesForCallPath(const FString&
 	return MatchedInstructions;
 }
 
-void FRigVMByteCode::SetSubject(int32 InInstructionIndex, UObject* InSubject, const FString& InCallPath)
+void FRigVMByteCode::SetSubject(int32 InInstructionIndex, const FString& InCallPath, const TArray<UObject*>& InCallstack)
 {
+	UObject* Subject = InCallstack.Last();
 	if (SubjectPerInstruction.Num() <= InInstructionIndex)
 	{
 		SubjectPerInstruction.AddZeroed(1 + InInstructionIndex - SubjectPerInstruction.Num());
 	}
-	SubjectPerInstruction[InInstructionIndex] = InSubject;
-	SubjectToInstructions.FindOrAdd(InSubject).Add(InInstructionIndex);
+	SubjectPerInstruction[InInstructionIndex] = Subject;
+	SubjectToInstructions.FindOrAdd(Subject).Add(InInstructionIndex);
 
 	if (CallPathPerInstruction.Num() <= InInstructionIndex)
 	{
@@ -1342,6 +1345,21 @@ void FRigVMByteCode::SetSubject(int32 InInstructionIndex, UObject* InSubject, co
 	}
 	CallPathPerInstruction[InInstructionIndex] = InCallPath;
 	CallPathToInstructions.FindOrAdd(InCallPath).Add(InInstructionIndex);
+
+	if (CallstackPerInstruction.Num() <= InInstructionIndex)
+	{
+		CallstackPerInstruction.AddZeroed(1 + InInstructionIndex - CallstackPerInstruction.Num());
+	}
+	CallstackPerInstruction[InInstructionIndex] = InCallstack;
+}
+
+const TArray<UObject*>* FRigVMByteCode::GetCallstackForInstruction(int32 InInstructionIndex) const
+{
+	if (CallstackPerInstruction.IsValidIndex(InInstructionIndex))
+	{
+		return &CallstackPerInstruction[InInstructionIndex];
+	}
+	return nullptr;
 }
 
 #endif
diff --git a/Engine/Source/Runtime/RigVM/Public/RigVMCore/RigVMByteCode.h b/Engine/Source/Runtime/RigVM/Public/RigVMCore/RigVMByteCode.h
--- a/Engine/Source/Runtime/RigVM/Public/RigVMCore/RigVMByteCode.h
+++ b/Engine/Source/Runtime/RigVM/Public/RigVMCore/RigVMByteCode.h
@@ -683,6 +683,9 @@ public:
 	// returns all found instruction indices for a given callpath
 	TArray<int32> GetAllInstructionIndicesForCallPath(const FString& InCallPath, bool bStartsWith = false, bool bEndsWith = false) const;
 
+	// returns the callstack which was used to inject a given instruction
+	const TArray<UObject*>* GetCallstackForInstruction(int32 InInstructionIndex) const;
+
 #endif
 
 private:
@@ -713,12 +716,13 @@ private:
 	TMap<UObject*, TArray<int32>> SubjectToInstructions;
 	TArray<FString> CallPathPerInstruction;
 	TMap<FString, TArray<int32>> CallPathToInstructions;
+	TArray<TArray<UObject*>> CallstackPerInstruction;
 
 #endif
 
 #if WITH_EDITOR
 
-	void SetSubject(int32 InInstructionIndex, UObject* InSubject, const FString& InCallPath);
+	void SetSubject(int32 InInstructionIndex, const FString& InCallPath, const TArray<UObject*>& InCallstack);
 
 #endif
 
