Cleaning up references in Control Rig BPs to the transient package.  Prevent adding nodes to the ControlRig Function Library if they're template nodes from the PROTO_ menuing BP function declarations.

diff --git a/Engine/Source/Developer/RigVMDeveloper/Private/RigVMModel/RigVMController.cpp b/Engine/Source/Developer/RigVMDeveloper/Private/RigVMModel/RigVMController.cpp
--- a/Engine/Source/Developer/RigVMDeveloper/Private/RigVMModel/RigVMController.cpp
+++ b/Engine/Source/Developer/RigVMDeveloper/Private/RigVMModel/RigVMController.cpp
@@ -6412,6 +6412,12 @@ URigVMFunctionReferenceNode* URigVMController::AddFunctionReferenceNode(URigVMLi
 		return nullptr;
 	}
 
+	if (Graph->GetOutermost() == GetTransientPackage())
+	{
+		// don't allow linking to transient graphs, which usually mean template nodes
+		return nullptr;
+	}
+
 	if(!CanAddFunctionRefForDefinition(InFunctionDefinition, true))
 	{
 		return nullptr;
