Control Rig: Fix up unit tests due to breakage in 16281286

diff --git a/Engine/Source/Developer/RigVMDeveloper/Private/RigVMModel/RigVMController.cpp b/Engine/Source/Developer/RigVMDeveloper/Private/RigVMModel/RigVMController.cpp
--- a/Engine/Source/Developer/RigVMDeveloper/Private/RigVMModel/RigVMController.cpp
+++ b/Engine/Source/Developer/RigVMDeveloper/Private/RigVMModel/RigVMController.cpp
@@ -32,6 +32,7 @@ URigVMController::URigVMController()
 	, bSuspendNotifications(false)
 	, bReportWarningsAndErrors(true)
 	, bIgnoreRerouteCompactnessChanges(false)
+	, bIsRunningUnitTest(false)
 {
 	SetExecuteContextStruct(FRigVMExecuteContext::StaticStruct());
 }
@@ -6412,7 +6413,7 @@ URigVMFunctionReferenceNode* URigVMController::AddFunctionReferenceNode(URigVMLi
 		return nullptr;
 	}
 
-	if (Graph->GetOutermost() == GetTransientPackage())
+	if (Graph->GetOutermost() == GetTransientPackage() && !bIsRunningUnitTest)
 	{
 		// don't allow linking to transient graphs, which usually mean template nodes
 		return nullptr;
diff --git a/Engine/Source/Developer/RigVMDeveloper/Public/RigVMModel/RigVMController.h b/Engine/Source/Developer/RigVMDeveloper/Public/RigVMModel/RigVMController.h
--- a/Engine/Source/Developer/RigVMDeveloper/Public/RigVMModel/RigVMController.h
+++ b/Engine/Source/Developer/RigVMDeveloper/Public/RigVMModel/RigVMController.h
@@ -683,6 +683,11 @@ public:
 	 */
 	void SuspendNotifications(bool bSuspend) { bSuspendNotifications = bSuspend; }
 
+	/**
+	 * Helper function to disable a series of checks that can be ignored during a unit test
+	 */
+	void SetIsRunningUnitTest(bool bIsRunning) { bIsRunningUnitTest = bIsRunning; }
+
 private:
 
 	UPROPERTY(BlueprintReadOnly, Category = RigVMController, meta = (ScriptName = "ModifiedEvent", AllowPrivateAccess = "true"))
@@ -824,6 +829,8 @@ private:
 
 	FRigVMUnitNodeCreatedContext UnitNodeCreatedContext;
 
+	bool bIsRunningUnitTest;
+
 	friend class URigVMGraph;
 	friend class URigVMActionStack;
 	friend class URigVMCompiler;
