Control Rig: Improving debug watches

* watches now no longer crash with Mannequin Control Rig
* watches no longer disappear after recompile
* watches work on formerly optimized out nodes
* watches work within collapse nodes and functions

diff --git a/Engine/Plugins/Experimental/ControlRig/Source/ControlRigDeveloper/Private/ControlRigBlueprint.cpp b/Engine/Plugins/Experimental/ControlRig/Source/ControlRigDeveloper/Private/ControlRigBlueprint.cpp
--- a/Engine/Plugins/Experimental/ControlRig/Source/ControlRigDeveloper/Private/ControlRigBlueprint.cpp
+++ b/Engine/Plugins/Experimental/ControlRig/Source/ControlRigDeveloper/Private/ControlRigBlueprint.cpp
@@ -2480,7 +2480,16 @@ void UControlRigBlueprint::HandleModifiedEvent(ERigVMGraphNotifType InNotifType,
 					
 					if(Pin->RequiresWatch())
 					{
-						Compiler->CreateDebugRegister(Pin, CR->GetVM(), &PinToOperandMap, RuntimeAST);
+						// check if the node is optimized out - in that case we need to recompile
+						if(CR->GetVM()->GetByteCode().GetFirstInstructionIndexForSubject(Pin->GetNode()) == INDEX_NONE)
+						{
+							RequestAutoVMRecompilation();
+							MarkPackageDirty();
+						}
+						else
+						{
+							Compiler->CreateDebugRegister(Pin, CR->GetVM(), &PinToOperandMap, RuntimeAST);
+						}
 					}
 					else
 					{
diff --git a/Engine/Plugins/Experimental/ControlRig/Source/ControlRigEditor/Private/ControlRigEditorModule.cpp b/Engine/Plugins/Experimental/ControlRig/Source/ControlRigEditor/Private/ControlRigEditorModule.cpp
--- a/Engine/Plugins/Experimental/ControlRig/Source/ControlRigEditor/Private/ControlRigEditorModule.cpp
+++ b/Engine/Plugins/Experimental/ControlRig/Source/ControlRigEditor/Private/ControlRigEditorModule.cpp
@@ -1238,24 +1238,19 @@ void FControlRigEditorModule::GetContextMenuActions(const UControlRigGraphSchema
 			
 			if(UControlRigGraphNode* RigNode = Cast<UControlRigGraphNode>(InGraphPin->GetOwningNode()))
 			{
-				// Check if we can watch this value (not inside any function)
-				bool bIsInsideFunction = RigNode->GetModel()->GetRootGraph()->IsA<URigVMFunctionLibrary>();
-				if(!bIsInsideFunction)
+				// Add the watch pin / unwatch pin menu items
+				FToolMenuSection& Section = Menu->AddSection("EdGraphSchemaWatches", LOCTEXT("WatchesHeader", "Watches"));
+				UBlueprint* OwnerBlueprint = FBlueprintEditorUtils::FindBlueprintForGraphChecked(Context->Graph);
 				{
-					// Add the watch pin / unwatch pin menu items
-					FToolMenuSection& Section = Menu->AddSection("EdGraphSchemaWatches", LOCTEXT("WatchesHeader", "Watches"));
-					UBlueprint* OwnerBlueprint = FBlueprintEditorUtils::FindBlueprintForGraphChecked(Context->Graph);
+					if (FKismetDebugUtilities::IsPinBeingWatched(OwnerBlueprint, InGraphPin))
 					{
-						if (FKismetDebugUtilities::IsPinBeingWatched(OwnerBlueprint, InGraphPin))
-						{
-							Section.AddMenuEntry(FGraphEditorCommands::Get().StopWatchingPin);
-						}
-						else
-						{
-							Section.AddMenuEntry(FGraphEditorCommands::Get().StartWatchingPin);
-						}
-					}				
-				}
+						Section.AddMenuEntry(FGraphEditorCommands::Get().StopWatchingPin);
+					}
+					else
+					{
+						Section.AddMenuEntry(FGraphEditorCommands::Get().StartWatchingPin);
+					}
+				}				
 			}
 
 			// Add alphainterp menu entries
diff --git a/Engine/Source/Developer/RigVMDeveloper/Private/RigVMCompiler/RigVMAST.cpp b/Engine/Source/Developer/RigVMDeveloper/Private/RigVMCompiler/RigVMAST.cpp
--- a/Engine/Source/Developer/RigVMDeveloper/Private/RigVMCompiler/RigVMAST.cpp
+++ b/Engine/Source/Developer/RigVMDeveloper/Private/RigVMCompiler/RigVMAST.cpp
@@ -559,6 +559,12 @@ bool FRigVMNodeExprAST::IsConstant() const
 		TArray<URigVMPin*> AllPins = CurrentNode->GetAllPinsRecursively();
 		for (URigVMPin* Pin : AllPins)
 		{
+			// don't flatten pins which have a watch
+			if(Pin->RequiresWatch(false))
+			{
+				return false;
+			}
+
 			URigVMPin::FPinOverride PinOverride(GetProxy(), GetParser()->GetPinOverrides());
 			if (Pin->IsBoundToVariable(PinOverride))
 			{
@@ -1845,6 +1851,9 @@ bool FRigVMParserAST::FoldConstantValuesToLiterals(URigVMGraph* InGraph, URigVMC
 	URigVMCompiler* TempCompiler = NewObject<URigVMCompiler>(GetTransientPackage());
 	TempCompiler->Settings.ConsolidateWorkRegisters = false;
 	TempCompiler->Settings.SetupNodeInstructionIndex = false;
+	TempCompiler->Settings.IsPreprocessorPhase = true;
+	TempCompiler->Settings.EnablePinWatches = false;
+	TempCompiler->Settings.ASTSettings = FRigVMParserASTSettings::Fast();
 
 	TempCompiler->Compile(InGraph, InController, TempVM, InExternalVariables, InRigVMUserData, &Operands, TempAST);
 
diff --git a/Engine/Source/Developer/RigVMDeveloper/Private/RigVMCompiler/RigVMCompiler.cpp b/Engine/Source/Developer/RigVMDeveloper/Private/RigVMCompiler/RigVMCompiler.cpp
--- a/Engine/Source/Developer/RigVMDeveloper/Private/RigVMCompiler/RigVMCompiler.cpp
+++ b/Engine/Source/Developer/RigVMDeveloper/Private/RigVMCompiler/RigVMCompiler.cpp
@@ -52,6 +52,7 @@ FRigVMCompileSettings::FRigVMCompileSettings()
 	, SurpressErrors(false)
 	, EnablePinWatches(true)
 	, ConsolidateWorkRegisters(false)
+	, IsPreprocessorPhase(false)
 	, ASTSettings(FRigVMParserASTSettings::Optimized())
 	, SetupNodeInstructionIndex(true)
 {
@@ -294,17 +295,20 @@ bool URigVMCompiler::Compile(URigVMGraph* InGraph, URigVMController* InControlle
 
 	WorkData.VM->GetByteCode().AlignByteCode();
 
-	// setup debug registers
-	for(int32 GraphIndex=0; GraphIndex<VisitedGraphs.Num(); GraphIndex++)
+	// setup debug registers after all other registers have been created
+	if(Settings.EnablePinWatches)
 	{
-		URigVMGraph* VisitedGraph = VisitedGraphs[GraphIndex];
-		for(URigVMNode* ModelNode : VisitedGraph->GetNodes())
+		for(int32 GraphIndex=0; GraphIndex<VisitedGraphs.Num(); GraphIndex++)
 		{
-			for(URigVMPin* ModelPin : ModelNode->GetPins())
+			URigVMGraph* VisitedGraph = VisitedGraphs[GraphIndex];
+			for(URigVMNode* ModelNode : VisitedGraph->GetNodes())
 			{
-				if(ModelPin->RequiresWatch(false))
+				for(URigVMPin* ModelPin : ModelNode->GetPins())
 				{
-					CreateDebugRegister(ModelPin, WorkData.VM, WorkData.PinPathToOperand, AST);
+					if(ModelPin->RequiresWatch(true))
+					{
+						CreateDebugRegister(ModelPin, WorkData.VM, WorkData.PinPathToOperand, AST);
+					}
 				}
 			}
 		}
@@ -502,20 +506,6 @@ int32 URigVMCompiler::TraverseCallExtern(const FRigVMCallExternExprAST* InExpr,
 		TraverseChildren(InExpr, WorkData);
 		WorkData.DefaultStructs.Pop();
 
-		if (Settings.EnablePinWatches)
-		{
-			for (URigVMPin* Pin : UnitNode->Pins)
-			{
-				if (Pin->RequiresWatch(true))
-				{
-					FRigVMASTProxy PinProxy = FRigVMASTProxy::MakeFromUObject(Pin);
-					FRigVMVarExprAST TempVarExpr(FRigVMExprAST::EType::Var, PinProxy);
-					TempVarExpr.ParserPtr = InExpr->GetParser();
-					FindOrAddRegister(&TempVarExpr, WorkData, true);
-				}
-			}
-		}
-
 		if (Settings.ConsolidateWorkRegisters)
 		{
 			for (const FRigVMExprAST* Child : *InExpr)
@@ -1348,6 +1338,7 @@ FString URigVMCompiler::GetPinHash(const URigVMPin* InPin, const FRigVMVarExprAS
 		}
 	}
 
+	bool bUseFullNodePath = true;
 	if (URigVMParameterNode* ParameterNode = Cast<URigVMParameterNode>(Node))
 	{
 		if (InPin->GetName() == TEXT("Value") && !bIsLiteral)
@@ -1427,20 +1418,21 @@ FString URigVMCompiler::GetPinHash(const URigVMPin* InPin, const FRigVMVarExprAS
 	{
 		if (InVarExpr)
 		{
-			FRigVMASTProxy NodeProxy = InVarExpr->GetProxy().GetSibling(Node);
+			const FRigVMASTProxy NodeProxy = InVarExpr->GetProxy().GetSibling(Node);
 			if (const FRigVMExprAST* NodeExpr = InVarExpr->GetParser()->GetExprForSubject(NodeProxy))
 			{
 				// rely on the proxy callstack to differentiate registers
-				FString CallStackPath = NodeProxy.GetCallstack().GetCallPath(false /* include last */);
+				const FString CallStackPath = NodeProxy.GetCallstack().GetCallPath(false /* include last */);
 				if (!CallStackPath.IsEmpty())
 				{
 					Prefix += CallStackPath + TEXT("|");
+					bUseFullNodePath = false;
 				}
 			}
 		}
 	}
 
-	return FString::Printf(TEXT("%s%s%s"), *Prefix, *InPin->GetPinPath(true /*  full node path */), *Suffix);
+	return FString::Printf(TEXT("%s%s%s"), *Prefix, *InPin->GetPinPath(bUseFullNodePath), *Suffix);
 }
 
 void URigVMCompiler::CreateDebugRegister(URigVMPin* InPin, URigVM* OutVM,
diff --git a/Engine/Source/Developer/RigVMDeveloper/Public/RigVMCompiler/RigVMCompiler.h b/Engine/Source/Developer/RigVMDeveloper/Public/RigVMCompiler/RigVMCompiler.h
--- a/Engine/Source/Developer/RigVMDeveloper/Public/RigVMCompiler/RigVMCompiler.h
+++ b/Engine/Source/Developer/RigVMDeveloper/Public/RigVMCompiler/RigVMCompiler.h
@@ -41,6 +41,9 @@ public:
 	UPROPERTY(VisibleAnywhere, BlueprintReadOnly, Category = FRigVMCompileSettings)
 	bool ConsolidateWorkRegisters;
 
+	UPROPERTY(Transient)
+	bool IsPreprocessorPhase;
+
 	UPROPERTY(EditAnywhere, Transient, BlueprintReadWrite, Category = FRigVMCompileSettings)
 	FRigVMParserASTSettings ASTSettings;
 
@@ -52,6 +55,7 @@ public:
 		FRigVMCompileSettings Settings;
 		Settings.ConsolidateWorkRegisters = false;
 		Settings.EnablePinWatches = true;
+		Settings.IsPreprocessorPhase = false;
 		Settings.ASTSettings = FRigVMParserASTSettings::Fast();
 		return Settings;
 	}
@@ -61,6 +65,7 @@ public:
 		FRigVMCompileSettings Settings;
 		Settings.ConsolidateWorkRegisters = false;
 		Settings.EnablePinWatches = false;
+		Settings.IsPreprocessorPhase = false;
 		Settings.ASTSettings = FRigVMParserASTSettings::Optimized();
 		return Settings;
 	}
diff --git a/Engine/Source/Runtime/RigVM/Private/RigVMCore/RigVM.cpp b/Engine/Source/Runtime/RigVM/Private/RigVMCore/RigVM.cpp
--- a/Engine/Source/Runtime/RigVM/Private/RigVMCore/RigVM.cpp
+++ b/Engine/Source/Runtime/RigVM/Private/RigVMCore/RigVM.cpp
@@ -289,6 +289,7 @@ void URigVM::CopyFrom(URigVM* InVM, bool bDeferCopy, bool bReferenceLiteralMemor
 	Instructions = InVM->Instructions;
 	Parameters = InVM->Parameters;
 	ParametersNameMap = InVM->ParametersNameMap;
+	OperandToDebugRegisters = InVM->OperandToDebugRegisters;
 
 	if (bCopyExternalVariables)
 	{
@@ -1161,21 +1162,24 @@ bool URigVM::Execute(FRigVMMemoryContainerPtrArray Memory, FRigVMFixedArray<void
 			{
 				const FRigVMExecuteOp& Op = ByteCode.GetOpAt<FRigVMExecuteOp>(Instruction);
 				const int32 OperandCount = FirstHandleForInstruction[Context.InstructionIndex + 1] - FirstHandleForInstruction[Context.InstructionIndex];
-				FRigVMMemoryHandleArray OpHandles(&CachedMemoryHandles[FirstHandleForInstruction[Context.InstructionIndex]], OperandCount);
+				FRigVMMemoryHandleArray Handles(&CachedMemoryHandles[FirstHandleForInstruction[Context.InstructionIndex]], OperandCount);
 #if WITH_EDITOR
 				Context.FunctionName = FunctionNames[Op.FunctionIndex];
 #endif
-				(*Functions[Op.FunctionIndex])(Context, OpHandles);
+				(*Functions[Op.FunctionIndex])(Context, Handles);
 
 #if WITH_EDITOR
 				if(DebugMemoryPtr->Num() > 0)
 				{
 					const FRigVMOperandArray Operands = ByteCode.GetOperandsForExecuteOp(Instruction);
-					check(Operands.Num() == OperandCount);
-					
-					for(int32 OperandIndex=0;OperandIndex<OperandCount;OperandIndex++)
+					for(int32 OperandIndex = 0, HandleIndex = 0; OperandIndex < Operands.Num() && HandleIndex < Handles.Num(); HandleIndex++)
 					{
-						CopyOperandForDebuggingIfNeeded(Operands[OperandIndex], OpHandles[OperandIndex]);
+						// skip array sizes
+						if(Handles[HandleIndex].GetType() == FRigVMMemoryHandle::FType::ArraySize)
+						{
+							continue;
+						}
+						CopyOperandForDebuggingIfNeeded(Operands[OperandIndex++], Handles[HandleIndex]);
 					}
 				}
 #endif
@@ -1894,7 +1898,7 @@ void URigVM::CacheSingleMemoryHandle(const FRigVMOperand& InArg, bool bForExecut
 		if (Register.IsArray() && !Register.IsDynamic())
 	{
 			void* ElementsForArray = reinterpret_cast<void*>(Register.ElementCount);
-			CachedMemoryHandles.Add(FRigVMMemoryHandle((uint8*)ElementsForArray));
+			CachedMemoryHandles.Add(FRigVMMemoryHandle((uint8*)ElementsForArray, sizeof(uint16), FRigVMMemoryHandle::FType::ArraySize));
 		}
 	}
 }
diff --git a/Engine/Source/Runtime/RigVM/Public/RigVMCore/RigVM.h b/Engine/Source/Runtime/RigVM/Public/RigVMCore/RigVM.h
--- a/Engine/Source/Runtime/RigVM/Public/RigVMCore/RigVM.h
+++ b/Engine/Source/Runtime/RigVM/Public/RigVMCore/RigVM.h
@@ -586,7 +586,7 @@ private:
 	
 	void CacheSingleMemoryHandle(const FRigVMOperand& InArg, bool bForExecute = false);
 
-	FORCEINLINE void CopyOperandForDebuggingIfNeeded(const FRigVMOperand& InArg, const FRigVMMemoryHandle& InHandle)
+	FORCEINLINE_DEBUGGABLE void CopyOperandForDebuggingIfNeeded(const FRigVMOperand& InArg, const FRigVMMemoryHandle& InHandle)
 	{
 #if WITH_EDITOR
 		const FRigVMOperand KeyOperand(InArg.GetMemoryType(), InArg.GetRegisterIndex()); // no register offset
diff --git a/Engine/Source/Runtime/RigVM/Public/RigVMCore/RigVMMemory.h b/Engine/Source/Runtime/RigVM/Public/RigVMCore/RigVMMemory.h
--- a/Engine/Source/Runtime/RigVM/Public/RigVMCore/RigVMMemory.h
+++ b/Engine/Source/Runtime/RigVM/Public/RigVMCore/RigVMMemory.h
@@ -407,7 +407,8 @@ public:
 	{
 		Plain, // Shallow memory with no indirection (for example FTransform.Translation.X)
 		Dynamic, // Memory representing an array
-		NestedDynamic // Memory representing an array or single element with slices (for example within loop)
+		NestedDynamic, // Memory representing an array or single element with slices (for example within loop)
+		ArraySize // a count represented as a pointer address
 	};
 
 	FORCEINLINE_DEBUGGABLE FRigVMMemoryHandle()
