Control Rig: remove notion of slice context from graph and improve compilation performance

diff --git a/Engine/Source/Developer/RigVMDeveloper/Private/RigVMCompiler/RigVMCompiler.cpp b/Engine/Source/Developer/RigVMDeveloper/Private/RigVMCompiler/RigVMCompiler.cpp
--- a/Engine/Source/Developer/RigVMDeveloper/Private/RigVMCompiler/RigVMCompiler.cpp
+++ b/Engine/Source/Developer/RigVMDeveloper/Private/RigVMCompiler/RigVMCompiler.cpp
@@ -1892,13 +1892,7 @@ FRigVMOperand URigVMCompiler::FindOrAddRegister(const FRigVMVarExprAST* InVarExp
 
 		if (!Operand.IsValid())
 		{
-			int32 NumSlices = 1;
-			if (Pin->GetDirection() == ERigVMPinDirection::Hidden)
-			{
-				NumSlices = Pin->GetNumSlices(WorkData.RigVMUserData);
-				ensure(NumSlices >= 1);
-			}
-
+			const int32 NumSlices = 1;
 			int32 Register = INDEX_NONE;
 
 			// debug watch register might already exists - look for them by name
diff --git a/Engine/Source/Developer/RigVMDeveloper/Private/RigVMModel/Nodes/RigVMUnitNode.cpp b/Engine/Source/Developer/RigVMDeveloper/Private/RigVMModel/Nodes/RigVMUnitNode.cpp
--- a/Engine/Source/Developer/RigVMDeveloper/Private/RigVMModel/Nodes/RigVMUnitNode.cpp
+++ b/Engine/Source/Developer/RigVMDeveloper/Private/RigVMModel/Nodes/RigVMUnitNode.cpp
@@ -3,8 +3,6 @@
 #include "RigVMModel/Nodes/RigVMUnitNode.h"
 #include "RigVMCore/RigVMStruct.h"
 
-const FName URigVMUnitNode::LoopSliceContextName = TEXT("LoopContext");
-
 FString URigVMUnitNode::GetNodeTitle() const
 {
 	if (UScriptStruct* Struct = GetScriptStruct())
@@ -52,64 +50,6 @@ FName URigVMUnitNode::GetEventName() const
 	return NAME_None;
 }
 
-FName URigVMUnitNode::GetSliceContextForPin(URigVMPin* InRootPin, const FRigVMUserDataArray& InUserData)
-{
-	TSharedPtr<FStructOnScope> StructOnScope = ConstructStructInstance(false);
-	if (StructOnScope.IsValid())
-	{
-		const FRigVMStruct* StructMemory = (FRigVMStruct*)StructOnScope->GetStructMemory();
-		if (StructMemory->IsForLoop())
-		{
-			// if we are on any of the pins returning something from the loop
-			if (InRootPin->GetFName() == FRigVMStruct::ExecuteContextName ||
-				InRootPin->GetFName() == FRigVMStruct::ForLoopIndexPinName)
-			{
-				return LoopSliceContextName;
-			}
-			else if (InRootPin->GetFName() == FRigVMStruct::ForLoopCompletedPinName||
-				InRootPin->GetFName() == FRigVMStruct::ForLoopCountPinName ||
-				InRootPin->GetFName() == FRigVMStruct::ForLoopContinuePinName)
-			{
-				return NAME_None;
-			}
-		}
-	}
-
-	if (UScriptStruct* Struct = GetScriptStruct())
-	{
-		if (FProperty* Property = Struct->FindPropertyByName(InRootPin->GetFName()))
-		{
-			FString SliceContextMetaData = Property->GetMetaData(FRigVMStruct::SliceContextMetaName);
-			if (!SliceContextMetaData.IsEmpty())
-			{
-				return *SliceContextMetaData;
-			}
-		}
-	}
-
-	return Super::GetSliceContextForPin(InRootPin, InUserData);
-}
-
-int32 URigVMUnitNode::GetNumSlicesForContext(const FName& InContextName, const FRigVMUserDataArray& InUserData)
-{
-	int32 NumSlices = Super::GetNumSlicesForContext(InContextName, InUserData);
-
-	if (InContextName == LoopSliceContextName)
-	{
-		TSharedPtr<FStructOnScope> StructOnScope = ConstructStructInstance(false);
-		if (StructOnScope.IsValid())
-		{
-			const FRigVMStruct* StructMemory = (FRigVMStruct*)StructOnScope->GetStructMemory();
-			if (StructMemory->IsForLoop())
-			{
-				NumSlices = NumSlices * FMath::Max<int32>(StructMemory->GetNumSlices(), 1);
-			}
-		}
-	}
-
-	return NumSlices;
-}
-
 FText URigVMUnitNode::GetToolTipTextForPin(const URigVMPin* InPin) const
 {
 	if(UScriptStruct* Struct = GetScriptStruct())
diff --git a/Engine/Source/Developer/RigVMDeveloper/Private/RigVMModel/RigVMNode.cpp b/Engine/Source/Developer/RigVMDeveloper/Private/RigVMModel/RigVMNode.cpp
--- a/Engine/Source/Developer/RigVMDeveloper/Private/RigVMModel/RigVMNode.cpp
+++ b/Engine/Source/Developer/RigVMDeveloper/Private/RigVMModel/RigVMNode.cpp
@@ -12,7 +12,6 @@ URigVMNode::URigVMNode()
 , Position(FVector2D::ZeroVector)
 , Size(FVector2D::ZeroVector)
 , NodeColor(FLinearColor::White)
-, GetSliceContextBracket(0)
 {
 
 }
@@ -442,49 +441,6 @@ void URigVMNode::GetLinkedNodesRecursive(URigVMPin* InPin, bool bLookForSources,
 	}
 }
 
-FName URigVMNode::GetSliceContextForPin(URigVMPin* InRootPin, const FRigVMUserDataArray& InUserData)
-{
-	return NAME_None;
-}
-
-int32 URigVMNode::GetNumSlices(const FRigVMUserDataArray& InUserData)
-{
-	return GetNumSlicesForContext(NAME_None, InUserData);
-}
-
-int32 URigVMNode::GetNumSlicesForContext(const FName& InContextName, const FRigVMUserDataArray& InUserData)
-{
-	for (URigVMPin* RootPin : GetPins())
-	{
-		if (RootPin->GetFName() == InContextName)
-		{
-			return RootPin->GetNumSlices(InUserData);
-		}
-	}
-
-	int32 MaxSlices = 1;
-
-	if (GetSliceContextBracket == 0)
-	{
-		TGuardValue<int32> ReentrantGuard(GetSliceContextBracket, GetSliceContextBracket + 1);
-
-		for (URigVMPin* Pin : GetPins())
-		{
-			TArray<URigVMPin*> SourcePins = Pin->GetLinkedSourcePins(true /* recursive */);
-			if (SourcePins.Num() > 0)
-			{
-				for (URigVMPin* SourcePin : SourcePins)
-				{
-					int32 NumSlices = SourcePin->GetNumSlices(InUserData);
-					MaxSlices = FMath::Max<int32>(NumSlices, MaxSlices);
-				}
-			}
-		}
-	}
-
-	return MaxSlices;
-}
-
 TArray<int32> URigVMNode::GetInstructionsForVM(URigVM* InVM, const FRigVMASTProxy& InProxy) const
 {
 	TArray<int32> Instructions;
diff --git a/Engine/Source/Developer/RigVMDeveloper/Private/RigVMModel/RigVMPin.cpp b/Engine/Source/Developer/RigVMDeveloper/Private/RigVMModel/RigVMPin.cpp
--- a/Engine/Source/Developer/RigVMDeveloper/Private/RigVMModel/RigVMPin.cpp
+++ b/Engine/Source/Developer/RigVMDeveloper/Private/RigVMModel/RigVMPin.cpp
@@ -790,28 +790,6 @@ bool URigVMPin::CanBeBoundToVariable(const FRigVMExternalVariable& InExternalVar
 	return true;
 }
 
-FName URigVMPin::GetSliceContext(const FRigVMUserDataArray& InUserData)
-{
-	URigVMPin* RootPin = GetRootPin();
-	if (RootPin != this)
-	{
-		return RootPin->GetSliceContext(InUserData);
-	}
-	return GetNode()->GetSliceContextForPin(this, InUserData);
-}
-
-int32 URigVMPin::GetNumSlices(const FRigVMUserDataArray& InUserData)
-{
-	URigVMPin* RootPin = GetRootPin();
-	if(RootPin != this)
-	{
-		return RootPin->GetNumSlices(InUserData);
-	}
-
-	FName SliceContext = GetSliceContext(InUserData);
-	return GetNode()->GetNumSlicesForContext(SliceContext, InUserData);
-}
-
 bool URigVMPin::ShowInDetailsPanelOnly() const
 {
 #if WITH_EDITOR
diff --git a/Engine/Source/Developer/RigVMDeveloper/Public/RigVMModel/Nodes/RigVMUnitNode.h b/Engine/Source/Developer/RigVMDeveloper/Public/RigVMModel/Nodes/RigVMUnitNode.h
--- a/Engine/Source/Developer/RigVMDeveloper/Public/RigVMModel/Nodes/RigVMUnitNode.h
+++ b/Engine/Source/Developer/RigVMDeveloper/Public/RigVMModel/Nodes/RigVMUnitNode.h
@@ -25,8 +25,6 @@ public:
 	virtual bool IsDefinedAsConstant() const override;
 	virtual bool IsDefinedAsVarying() const override;
 	virtual FName GetEventName() const override;
-	virtual FName GetSliceContextForPin(URigVMPin* InRootPin, const FRigVMUserDataArray& InUserData) override;
-	virtual int32 GetNumSlicesForContext(const FName& InContextName, const FRigVMUserDataArray& InUserData) override;
 
 	bool IsDeprecated() const;
 	FString GetDeprecatedMetadata() const;
@@ -51,8 +49,6 @@ public:
 	// @param bUseDefault If set to true the default struct will be created - otherwise the struct will contains the values from the node
 	TSharedPtr<FStructOnScope> ConstructStructInstance(bool bUseDefault = false) const;
 
-	static const FName LoopSliceContextName;
-
 protected:
 
 	virtual FText GetToolTipTextForPin(const URigVMPin* InPin) const override;
diff --git a/Engine/Source/Developer/RigVMDeveloper/Public/RigVMModel/RigVMNode.h b/Engine/Source/Developer/RigVMDeveloper/Public/RigVMModel/RigVMNode.h
--- a/Engine/Source/Developer/RigVMDeveloper/Public/RigVMModel/RigVMNode.h
+++ b/Engine/Source/Developer/RigVMDeveloper/Public/RigVMModel/RigVMNode.h
@@ -184,15 +184,6 @@ public:
 	UFUNCTION(BlueprintCallable, Category = RigVMNode)
 	TArray<URigVMNode*> GetLinkedTargetNodes() const;
 
-	// Returns the name of the slice context for a pin
-	virtual FName GetSliceContextForPin(URigVMPin* InRootPin, const FRigVMUserDataArray& InUserData);
-
-	// returns the number of slices on this node
-	int32 GetNumSlices(const FRigVMUserDataArray& InUserData);
-
-	// Returns the number of slices for a given context
-	virtual int32 GetNumSlicesForContext(const FName& InContextName, const FRigVMUserDataArray& InUserData);
-
 	// Returns the name of the node prior to the renaming
 	UFUNCTION(BlueprintCallable, Category = RigVMNode)
 	FName GetPreviousFName() const { return PreviousName; }
@@ -255,8 +246,6 @@ private:
 	UPROPERTY()
 	TArray<URigVMPin*> OrphanedPins;
 
-	int32 GetSliceContextBracket;
-
 	friend class URigVMController;
 	friend class URigVMGraph;
 	friend class URigVMPin;
diff --git a/Engine/Source/Developer/RigVMDeveloper/Public/RigVMModel/RigVMPin.h b/Engine/Source/Developer/RigVMDeveloper/Public/RigVMModel/RigVMPin.h
--- a/Engine/Source/Developer/RigVMDeveloper/Public/RigVMModel/RigVMPin.h
+++ b/Engine/Source/Developer/RigVMDeveloper/Public/RigVMModel/RigVMPin.h
@@ -347,12 +347,6 @@ public:
 		return Cast<T>(FindObjectFromCPPTypeObjectPath(InObjectPath));
 	}
 
-	// Returns the name of the context this pin belongs to
-	FName GetSliceContext(const FRigVMUserDataArray& InUserData);
-
-	// Returns the number of slices in memory exist for this pin
-	int32 GetNumSlices(const FRigVMUserDataArray& InUserData);
-
 	// Returns true if the pin should not show up on a node, but in the details panel
 	bool ShowInDetailsPanelOnly() const;
 
