Control Rig: Avoid nullptr access for intermediate compiles

diff --git a/Engine/Source/Developer/RigVMDeveloper/Private/RigVMModel/RigVMNode.cpp b/Engine/Source/Developer/RigVMDeveloper/Private/RigVMModel/RigVMNode.cpp
--- a/Engine/Source/Developer/RigVMDeveloper/Private/RigVMModel/RigVMNode.cpp
+++ b/Engine/Source/Developer/RigVMDeveloper/Private/RigVMModel/RigVMNode.cpp
@@ -24,10 +24,13 @@ FString URigVMNode::GetNodePath(bool bRecursive) const
 {
 	if (bRecursive)
 	{
-		FString ParentNodePath = GetGraph()->GetNodePath();
-		if (!ParentNodePath.IsEmpty())
+		if(URigVMGraph* Graph = GetGraph())
 		{
-			return JoinNodePath(ParentNodePath, GetName());
+			const FString ParentNodePath = Graph->GetNodePath();
+			if (!ParentNodePath.IsEmpty())
+			{
+				return JoinNodePath(ParentNodePath, GetName());
+			}
 		}
 	}
 	return GetName();
