ControlRig: Fix change function pin type in details panel not working

diff --git a/Engine/Source/Developer/RigVMDeveloper/Private/RigVMModel/RigVMController.cpp b/Engine/Source/Developer/RigVMDeveloper/Private/RigVMModel/RigVMController.cpp
--- a/Engine/Source/Developer/RigVMDeveloper/Private/RigVMModel/RigVMController.cpp
+++ b/Engine/Source/Developer/RigVMDeveloper/Private/RigVMModel/RigVMController.cpp
@@ -6270,7 +6270,6 @@ bool URigVMController::ChangeExposedPinType(const FName& InPinName, const FStrin
 		ActionStack->BeginAction(Action);
 	}
 
-	if (bSetupUndoRedo)
 	{
 		FRigVMControllerGraphGuard GraphGuard(this, LibraryNode->GetGraph(), bSetupUndoRedo);
 		if (!ChangePinType(Pin, InCPPType, InCPPTypeObjectPath, bSetupUndoRedo, bSetupOrphanPins))
@@ -6283,6 +6282,7 @@ bool URigVMController::ChangeExposedPinType(const FName& InPinName, const FStrin
 		}
 		RemoveUnusedOrphanedPins(LibraryNode, true);
 	}
+	
 
 	if (URigVMFunctionEntryNode* EntryNode = Graph->GetEntryNode())
 	{
