Control Rig: Make sure we don't reset the loaders when doing a UObject rename in PostLoad, otherwise bad things happen to the async loader.

Edigrate of 16413543

diff --git a/Engine/Source/Developer/RigVMDeveloper/Private/RigVMModel/RigVMController.cpp b/Engine/Source/Developer/RigVMDeveloper/Private/RigVMModel/RigVMController.cpp
--- a/Engine/Source/Developer/RigVMDeveloper/Private/RigVMModel/RigVMController.cpp
+++ b/Engine/Source/Developer/RigVMDeveloper/Private/RigVMModel/RigVMController.cpp
@@ -9002,7 +9002,7 @@ void URigVMController::ApplyPinStates(URigVMNode* InNode, const TMap<FString, UR
 		{
 			for (URigVMInjectionInfo* InjectionInfo : PinState.InjectionInfos)
 			{
-				InjectionInfo->UnitNode->Rename(nullptr, InNode->GetGraph());
+				InjectionInfo->UnitNode->Rename(nullptr, InNode->GetGraph(), REN_ForceNoResetLoaders | REN_DoNotDirty | REN_DontCreateRedirectors | REN_NonTransactional);
 				DestroyObject(InjectionInfo);
 			}
 		}
