ControlRig: Fix crash when undoing a change exposed pin type in function

diff --git a/Engine/Plugins/Experimental/ControlRig/Source/ControlRigEditor/Private/ControlRigGraphDetails.cpp b/Engine/Plugins/Experimental/ControlRig/Source/ControlRigEditor/Private/ControlRigGraphDetails.cpp
--- a/Engine/Plugins/Experimental/ControlRig/Source/ControlRigEditor/Private/ControlRigGraphDetails.cpp
+++ b/Engine/Plugins/Experimental/ControlRig/Source/ControlRigEditor/Private/ControlRigGraphDetails.cpp
@@ -358,7 +358,15 @@ void FControlRigArgumentLayout::PinInfoChanged(const FEdGraphPinType& PinType)
 					CPPType = FString::Printf(TEXT("TArray<%s>"), *CPPType);
 				}
 
-				Controller->ChangeExposedPinType(Pin->GetFName(), CPPType, CPPTypeObjectName, true);
+				bool bSetupUndoRedo = true;
+				Controller->ChangeExposedPinType(Pin->GetFName(), CPPType, CPPTypeObjectName, bSetupUndoRedo);
+
+				// If the controller has identified this as a bulk change, it has not added the actions to the action stack
+				// We need to disable the transaction from the UI as well to keep them synced
+				if (!bSetupUndoRedo)
+				{
+					GEditor->CancelTransaction(0);
+				}
 			}
 		}
 	}
diff --git a/Engine/Source/Developer/RigVMDeveloper/Private/RigVMModel/RigVMController.cpp b/Engine/Source/Developer/RigVMDeveloper/Private/RigVMModel/RigVMController.cpp
--- a/Engine/Source/Developer/RigVMDeveloper/Private/RigVMModel/RigVMController.cpp
+++ b/Engine/Source/Developer/RigVMDeveloper/Private/RigVMModel/RigVMController.cpp
@@ -5970,7 +5970,7 @@ FName URigVMController::AddExposedPin(const FName& InPinName, ERigVMPinDirection
 
 	if (Pin->IsStruct())
 	{
-		FRigVMControllerGraphGuard GraphGuard(this, LibraryNode->GetGraph(), false);
+		FRigVMControllerGraphGuard GraphGuard(this, LibraryNode->GetGraph(), bSetupUndoRedo);
 
 		FString DefaultValue = InDefaultValue;
 		CreateDefaultValueForStructIfRequired(Pin->GetScriptStruct(), DefaultValue);
@@ -5984,19 +5984,19 @@ FName URigVMController::AddExposedPin(const FName& InPinName, ERigVMPinDirection
 	}
 
 	{
-		FRigVMControllerGraphGuard GraphGuard(this, LibraryNode->GetGraph(), false);
+		FRigVMControllerGraphGuard GraphGuard(this, LibraryNode->GetGraph(), bSetupUndoRedo);
 		Notify(ERigVMGraphNotifType::PinAdded, Pin);
 	}
 
 	if (!InDefaultValue.IsEmpty())
 	{
-		FRigVMControllerGraphGuard GraphGuard(this, Pin->GetGraph(), false);
+		FRigVMControllerGraphGuard GraphGuard(this, Pin->GetGraph(), bSetupUndoRedo);
 		SetPinDefaultValue(Pin, InDefaultValue, true, bSetupUndoRedo, false);
 	}
 
 	RefreshFunctionPins(Graph->GetEntryNode(), true);
 	RefreshFunctionPins(Graph->GetReturnNode(), true);
-	RefreshFunctionReferences(LibraryNode, false);
+	RefreshFunctionReferences(LibraryNode, bSetupUndoRedo);
 
 	if (bSetupUndoRedo)
 	{
@@ -6220,7 +6220,7 @@ bool URigVMController::RenameExposedPin(const FName& InOldPinName, const FName&
 	return true;
 }
 
-bool URigVMController::ChangeExposedPinType(const FName& InPinName, const FString& InCPPType, const FName& InCPPTypeObjectPath, bool bSetupUndoRedo, bool bSetupOrphanPins)
+bool URigVMController::ChangeExposedPinType(const FName& InPinName, const FString& InCPPType, const FName& InCPPTypeObjectPath, bool& bSetupUndoRedo, bool bSetupOrphanPins)
 {
 	if (!IsValidGraph())
 	{
diff --git a/Engine/Source/Developer/RigVMDeveloper/Public/RigVMModel/RigVMController.h b/Engine/Source/Developer/RigVMDeveloper/Public/RigVMModel/RigVMController.h
--- a/Engine/Source/Developer/RigVMDeveloper/Public/RigVMModel/RigVMController.h
+++ b/Engine/Source/Developer/RigVMDeveloper/Public/RigVMModel/RigVMController.h
@@ -560,7 +560,7 @@ public:
 
 	// Changes the type of an exposed pin in the graph controlled by this
 	UFUNCTION(BlueprintCallable, Category = RigVMController)
-	bool ChangeExposedPinType(const FName& InPinName, const FString& InCPPType, const FName& InCPPTypeObjectPath, bool bSetupUndoRedo = true, bool bSetupOrphanPins = true);
+	bool ChangeExposedPinType(const FName& InPinName, const FString& InCPPType, const FName& InCPPTypeObjectPath, UPARAM(ref) bool& bSetupUndoRedo, bool bSetupOrphanPins = true);
 
 	// Sets the index for an exposed pin. This can be used to move the pin up and down on the node.
 	UFUNCTION(BlueprintCallable, Category = RigVMController)
