Control Rig: Disallow pasting Entry / Return nodes

diff --git a/Engine/Source/Developer/RigVMDeveloper/Private/RigVMModel/RigVMController.cpp b/Engine/Source/Developer/RigVMDeveloper/Private/RigVMModel/RigVMController.cpp
--- a/Engine/Source/Developer/RigVMDeveloper/Private/RigVMModel/RigVMController.cpp
+++ b/Engine/Source/Developer/RigVMDeveloper/Private/RigVMModel/RigVMController.cpp
@@ -7579,7 +7579,7 @@ bool URigVMController::IsValidLinkForGraph(URigVMLink* InLink)
 	return true;
 }
 
-bool URigVMController::CanAddNode(URigVMNode* InNode, bool bReportErrors)
+bool URigVMController::CanAddNode(URigVMNode* InNode, bool bReportErrors, bool bIgnoreFunctionEntryReturnNodes)
 {
 	if(!IsValidGraph())
 	{
@@ -7627,6 +7627,33 @@ bool URigVMController::CanAddNode(URigVMNode* InNode, bool bReportErrors)
 			}
 		}			
 	}
+	else if(!bIgnoreFunctionEntryReturnNodes &&
+		(InNode->IsA<URigVMFunctionEntryNode>() ||
+		InNode->IsA<URigVMFunctionReturnNode>()))
+	{
+		// only allow entry / return nodes on sub graphs
+		if(Graph->IsRootGraph())
+		{
+			return false;
+		}
+
+		// only allow one function entry node
+		if(InNode->IsA<URigVMFunctionEntryNode>())
+		{
+			if(Graph->GetEntryNode() != nullptr)
+			{
+				return false;
+			}
+		}
+		// only allow one function return node
+		else if(InNode->IsA<URigVMFunctionReturnNode>())
+		{
+			if(Graph->GetReturnNode() != nullptr)
+			{
+				return false;
+			}
+		}
+	}
 	else if(URigVMCollapseNode* CollapseNode = Cast<URigVMCollapseNode>(InNode))
 	{
 		FRigVMControllerGraphGuard GraphGuard(this, CollapseNode->GetContainedGraph(), false);
@@ -7634,7 +7661,7 @@ bool URigVMController::CanAddNode(URigVMNode* InNode, bool bReportErrors)
 		TArray<URigVMNode*> ContainedNodes = CollapseNode->GetContainedNodes();
 		for(URigVMNode* ContainedNode : ContainedNodes)
 		{
-			if(!CanAddNode(ContainedNode, bReportErrors))
+			if(!CanAddNode(ContainedNode, bReportErrors, true))
 			{
 				return false;
 			}
