Control Rig: Avoid crash when expanding empty function

diff --git a/Engine/Source/Developer/RigVMDeveloper/Private/RigVMModel/RigVMController.cpp b/Engine/Source/Developer/RigVMDeveloper/Private/RigVMModel/RigVMController.cpp
--- a/Engine/Source/Developer/RigVMDeveloper/Private/RigVMModel/RigVMController.cpp
+++ b/Engine/Source/Developer/RigVMDeveloper/Private/RigVMModel/RigVMController.cpp
@@ -2906,6 +2906,10 @@ TArray<URigVMNode*> URigVMController::ExpandLibraryNode(URigVMLibraryNode* InNod
 
 	if (ContainedNodes.Num() == 0)
 	{
+		if (bSetupUndoRedo)
+		{
+			ActionStack->CancelAction(ExpandAction);
+		}
 		return TArray<URigVMNode*>();
 	}
 
