Control Rig: Introducing sub titles for variables showing default values, adding tooltips to pins on function reference nodes to show tooltip

diff --git a/Engine/Plugins/Experimental/ControlRig/Source/ControlRigDeveloper/Private/ControlRigBlueprint.cpp b/Engine/Plugins/Experimental/ControlRig/Source/ControlRigDeveloper/Private/ControlRigBlueprint.cpp
--- a/Engine/Plugins/Experimental/ControlRig/Source/ControlRigDeveloper/Private/ControlRigBlueprint.cpp
+++ b/Engine/Plugins/Experimental/ControlRig/Source/ControlRigDeveloper/Private/ControlRigBlueprint.cpp
@@ -1019,6 +1019,22 @@ URigVMGraph* UControlRigBlueprint::GetModel(const FString& InNodePath) const
 		{
 			return LibraryNode->GetContainedGraph();
 		}
+
+		if(FunctionLibrary)
+		{
+			FString Left, Right;
+			if(URigVMNode::SplitNodePathAtStart(InNodePath, Left, Right))
+			{
+				if(Left == FunctionLibrary->GetNodePath())
+				{
+					if (URigVMLibraryNode* LibraryNode = Cast<URigVMLibraryNode>(FunctionLibrary->FindNode(Right)))
+					{
+						return LibraryNode->GetContainedGraph();
+					}
+				}
+			}
+		}
+		
 		return nullptr;
 	}
 	return Model;
diff --git a/Engine/Plugins/Experimental/ControlRig/Source/ControlRigDeveloper/Private/Graph/ControlRigGraphNode.cpp b/Engine/Plugins/Experimental/ControlRig/Source/ControlRigDeveloper/Private/Graph/ControlRigGraphNode.cpp
--- a/Engine/Plugins/Experimental/ControlRig/Source/ControlRigDeveloper/Private/Graph/ControlRigGraphNode.cpp
+++ b/Engine/Plugins/Experimental/ControlRig/Source/ControlRigDeveloper/Private/Graph/ControlRigGraphNode.cpp
@@ -88,6 +88,47 @@ FText UControlRigGraphNode::GetNodeTitle(ENodeTitleType::Type TitleType) const
 				SubTitle = CollapseNodeString;
 			}
 
+			else if(URigVMVariableNode* VariableNode = Cast<URigVMVariableNode>(ModelNode))
+			{
+				if(VariableNode->IsLocalVariable())
+				{
+					static const FString LocalVariableString = TEXT("Local Variable");
+					const FString DefaultValue = VariableNode->GetVariableDescription().DefaultValue;
+					if(DefaultValue.IsEmpty())
+					{
+						SubTitle = LocalVariableString;
+					}
+					else
+					{
+						SubTitle = FString::Printf(TEXT("%s\nDefault %s"), *LocalVariableString, *DefaultValue);
+					}
+				}
+				else
+				{
+					if(UBlueprint* Blueprint = GetBlueprint())
+					{
+						const FName VariableName = VariableNode->GetVariableName();
+						for(const FBPVariableDescription& NewVariable : Blueprint->NewVariables)
+						{
+							if(NewVariable.VarName == VariableName)
+							{
+								const FString DefaultValue = NewVariable.DefaultValue;
+								if(DefaultValue.IsEmpty())
+								{
+									static const FString VariableString = TEXT("Variable");
+									SubTitle = VariableString;
+								}
+								else
+								{
+									SubTitle = FString::Printf(TEXT("Default %s"), *DefaultValue);
+								}
+								break;
+							}
+						}
+					}
+				}
+			}
+
 			if (NodeTitle.IsEmpty())
 			{
 				NodeTitle = FText::FromString(ModelNode->GetNodeTitle());
diff --git a/Engine/Plugins/Experimental/ControlRig/Source/ControlRigEditor/Private/Editor/ControlRigEditor.cpp b/Engine/Plugins/Experimental/ControlRig/Source/ControlRigEditor/Private/Editor/ControlRigEditor.cpp
--- a/Engine/Plugins/Experimental/ControlRig/Source/ControlRigEditor/Private/Editor/ControlRigEditor.cpp
+++ b/Engine/Plugins/Experimental/ControlRig/Source/ControlRigEditor/Private/Editor/ControlRigEditor.cpp
@@ -319,10 +319,9 @@ void FControlRigEditor::InitControlRigEditor(const EToolkitMode::Type Mode, cons
 
 	CommonInitialization(ControlRigBlueprints, false);
 
-	for (UBlueprint* Blueprint : ControlRigBlueprints)
 	{
 		TArray<UEdGraph*> EdGraphs;
-		Blueprint->GetAllGraphs(EdGraphs);
+		InControlRigBlueprint->GetAllGraphs(EdGraphs);
 
 		for (UEdGraph* Graph : EdGraphs)
 		{
@@ -428,6 +427,11 @@ void FControlRigEditor::InitControlRigEditor(const EToolkitMode::Type Mode, cons
 						if (TabsForEdGraph.Num() > 0)
 						{
 							OpenedTabNodePaths.Add(RigGraph->ModelNodePath);
+
+							if(RigGraph->bIsFunctionDefinition)
+							{
+								CloseDocumentTab(RigGraph);
+							}
 						}
 					}
 				}
@@ -1670,6 +1674,49 @@ void FControlRigEditor::Compile()
 			}
 		}
 
+		// invalidate all node titles
+		TArray<UEdGraph*> EdGraphs;
+		RigBlueprint->GetAllGraphs(EdGraphs);
+		for (UEdGraph* EdGraph : EdGraphs)
+		{
+			UControlRigGraph* RigGraph = Cast<UControlRigGraph>(EdGraph);
+			if (RigGraph == nullptr)
+			{
+				continue;
+			}
+
+			for (UEdGraphNode* EdNode : RigGraph->Nodes)
+			{
+				if (UControlRigGraphNode* RigNode = Cast<UControlRigGraphNode>(EdNode))
+				{
+					RigNode->InvalidateNodeTitle();
+				}
+			}
+		}
+
+		// store the defaults from the CDO back on the new variables list
+		for(FBPVariableDescription& NewVariable : RigBlueprint->NewVariables)
+		{
+			UClass* GeneratedClass = RigBlueprint->GeneratedClass;
+			UObject* GeneratedCDO = GeneratedClass->GetDefaultObject();
+			FProperty* TargetProperty = FindFProperty<FProperty>(GeneratedClass, NewVariable.VarName);
+
+			if (TargetProperty)
+			{
+				// Grab the address of where the property is actually stored (UObject* base, plus the offset defined in the property)
+				void* OldPropertyAddr = TargetProperty->ContainerPtrToValuePtr<void>(GeneratedCDO);
+				if (OldPropertyAddr)
+				{
+					FString NewDefaultValue;
+					TargetProperty->ExportTextItem(NewDefaultValue, OldPropertyAddr, OldPropertyAddr, nullptr, PPF_SerializedAsImportText);
+					if(NewDefaultValue.Len() < 80)
+					{
+						NewVariable.DefaultValue = NewDefaultValue;
+					}
+				}
+			}
+		}
+
 		if(SelectedKey.IsValid())
 		{
 			SetDetailStruct(SelectedKey);
diff --git a/Engine/Source/Developer/RigVMDeveloper/Private/RigVMCompiler/RigVMCompiler.cpp b/Engine/Source/Developer/RigVMDeveloper/Private/RigVMCompiler/RigVMCompiler.cpp
--- a/Engine/Source/Developer/RigVMDeveloper/Private/RigVMCompiler/RigVMCompiler.cpp
+++ b/Engine/Source/Developer/RigVMDeveloper/Private/RigVMCompiler/RigVMCompiler.cpp
@@ -1357,39 +1357,31 @@ FString URigVMCompiler::GetPinHash(const URigVMPin* InPin, const FRigVMVarExprAS
 		{
 			FName VariableName = VariableNode->GetVariableName();
 
-			// Figure out if it is a local variable
+			if(VariableNode->IsLocalVariable())
 			{
-				URigVMGraph* RootGraph = InPin->GetNode()->GetGraph();
-				TArray<FRigVMGraphVariableDescription>& LocalVariables = RootGraph->LocalVariables;
-				for (const FRigVMGraphVariableDescription& LocalVariable : LocalVariables)
+				if (bIsLiteral)
 				{
-					if (LocalVariable.Name == VariableName)
+					// Literal values will be reused for all instance of local variables
+					return FString::Printf(TEXT("%sLocalVariable::%s|%s%s"), *Prefix, *Node->GetGraph()->GetGraphName(), *VariableName.ToString(), *Suffix);
+				}
+				else
+				{
+					if(InVarExpr)
 					{
-						if (bIsLiteral)
+						FRigVMASTProxy ParentProxy = InVarExpr->GetProxy();
+						while(ParentProxy.GetCallstack().Num() > 1)
 						{
-							// Literal values will be reused for all instance of local variables
-							return FString::Printf(TEXT("%sLocalVariable::%s|%s%s"), *Prefix, *RootGraph->GetGraphName(), *VariableName.ToString(), *Suffix);
-						}
-						else
-						{
-							if(InVarExpr)
-							{
-								FRigVMASTProxy ParentProxy = InVarExpr->GetProxy();
-								while(ParentProxy.GetCallstack().Num() > 1)
-								{
-									ParentProxy = ParentProxy.GetParent();
+							ParentProxy = ParentProxy.GetParent();
 
-									if(URigVMLibraryNode* LibraryNode = ParentProxy.GetSubject<URigVMLibraryNode>())
-									{
-										// Local variables for non-root graphs are in the format "LocalVariable::PathToGraph|VariableName"
-										return FString::Printf(TEXT("%sLocalVariable::%s|%s%s"), *Prefix, *LibraryNode->GetNodePath(true), *VariableName.ToString(), *Suffix);
-									}
-								}
-
-								// Local variables for root graphs are in the format "LocalVariable::VariableName"
-								return FString::Printf(TEXT("%sLocalVariable::%s%s"), *Prefix, *VariableName.ToString(), *Suffix);
+							if(URigVMLibraryNode* LibraryNode = ParentProxy.GetSubject<URigVMLibraryNode>())
+							{
+								// Local variables for non-root graphs are in the format "LocalVariable::PathToGraph|VariableName"
+								return FString::Printf(TEXT("%sLocalVariable::%s|%s%s"), *Prefix, *LibraryNode->GetNodePath(true), *VariableName.ToString(), *Suffix);
 							}
 						}
+
+						// Local variables for root graphs are in the format "LocalVariable::VariableName"
+						return FString::Printf(TEXT("%sLocalVariable::%s%s"), *Prefix, *VariableName.ToString(), *Suffix);
 					}
 				}
 			}
diff --git a/Engine/Source/Developer/RigVMDeveloper/Private/RigVMModel/Nodes/RigVMFunctionReferenceNode.cpp b/Engine/Source/Developer/RigVMDeveloper/Private/RigVMModel/Nodes/RigVMFunctionReferenceNode.cpp
--- a/Engine/Source/Developer/RigVMDeveloper/Private/RigVMModel/Nodes/RigVMFunctionReferenceNode.cpp
+++ b/Engine/Source/Developer/RigVMDeveloper/Private/RigVMModel/Nodes/RigVMFunctionReferenceNode.cpp
@@ -194,6 +194,27 @@ FName URigVMFunctionReferenceNode::GetOuterVariableName(const FName& InInnerVari
 	return NAME_None;
 }
 
+FText URigVMFunctionReferenceNode::GetToolTipTextForPin(const URigVMPin* InPin) const
+{
+	check(InPin);
+	
+	if(URigVMLibraryNode* ReferencedNode = GetReferencedNode())
+	{
+		if(URigVMPin* ReferencedPin = ReferencedNode->FindPin(InPin->GetSegmentPath(true)))
+		{
+			const FString DefaultValue = ReferencedPin->GetDefaultValue();
+			if(!DefaultValue.IsEmpty())
+			{
+				return FText::FromString(FString::Printf(TEXT("%s\nDefault %s"),
+					*ReferencedPin->GetToolTipText().ToString(),
+					*DefaultValue));
+			}
+		}
+	}
+
+	return Super::GetToolTipTextForPin(InPin);
+}
+
 URigVMLibraryNode* URigVMFunctionReferenceNode::GetReferencedNode() const
 {
 	if (!ReferencedNodePtr.IsValid())
diff --git a/Engine/Source/Developer/RigVMDeveloper/Private/RigVMModel/Nodes/RigVMVariableNode.cpp b/Engine/Source/Developer/RigVMDeveloper/Private/RigVMModel/Nodes/RigVMVariableNode.cpp
--- a/Engine/Source/Developer/RigVMDeveloper/Private/RigVMModel/Nodes/RigVMVariableNode.cpp
+++ b/Engine/Source/Developer/RigVMDeveloper/Private/RigVMModel/Nodes/RigVMVariableNode.cpp
@@ -2,6 +2,8 @@
 
 #include "RigVMModel/Nodes/RigVMVariableNode.h"
 
+#include "RigVMModel/RigVMGraph.h"
+
 const FString URigVMVariableNode::VariableName = TEXT("Variable");
 const FString URigVMVariableNode::ValueName = TEXT("Value");
 
@@ -38,6 +40,24 @@ bool URigVMVariableNode::IsGetter() const
 	return ValuePin->GetDirection() == ERigVMPinDirection::Output;
 }
 
+bool URigVMVariableNode::IsLocalVariable() const
+{
+	const FName CurrentVariableName = GetVariableName();
+	
+	if(URigVMGraph* Graph = GetGraph())
+	{
+		const TArray<FRigVMGraphVariableDescription>& LocalVariables = Graph->GetLocalVariables();
+		for (const FRigVMGraphVariableDescription& LocalVariable : LocalVariables)
+		{
+			if(LocalVariable.Name == CurrentVariableName)
+			{
+				return true;
+			}
+		}
+	}
+	return false;
+}
+
 FString URigVMVariableNode::GetCPPType() const
 {
 	URigVMPin* ValuePin = FindPin(ValueName);
diff --git a/Engine/Source/Developer/RigVMDeveloper/Private/RigVMModel/RigVMPin.cpp b/Engine/Source/Developer/RigVMDeveloper/Private/RigVMModel/RigVMPin.cpp
--- a/Engine/Source/Developer/RigVMDeveloper/Private/RigVMModel/RigVMPin.cpp
+++ b/Engine/Source/Developer/RigVMDeveloper/Private/RigVMModel/RigVMPin.cpp
@@ -206,18 +206,24 @@ FString URigVMPin::GetPinPath(bool bUseNodePath) const
 	return PinPath;
 }
 
-FString URigVMPin::GetSegmentPath() const
+FString URigVMPin::GetSegmentPath(bool bIncludeRootPin) const
 {
 	URigVMPin* ParentPin = GetParentPin();
 	if (ParentPin)
 	{
-		FString ParentSegmentPath = ParentPin->GetSegmentPath();
+		FString ParentSegmentPath = ParentPin->GetSegmentPath(bIncludeRootPin);
 		if (ParentSegmentPath.IsEmpty())
 		{
 			return GetName();
 		}
 		return FString::Printf(TEXT("%s.%s"), *ParentSegmentPath, *GetName());
 	}
+
+	if(bIncludeRootPin)
+	{
+		return GetName();
+	}
+	
 	return FString();
 }
 
diff --git a/Engine/Source/Developer/RigVMDeveloper/Public/RigVMModel/Nodes/RigVMFunctionReferenceNode.h b/Engine/Source/Developer/RigVMDeveloper/Public/RigVMModel/Nodes/RigVMFunctionReferenceNode.h
--- a/Engine/Source/Developer/RigVMDeveloper/Public/RigVMModel/Nodes/RigVMFunctionReferenceNode.h
+++ b/Engine/Source/Developer/RigVMDeveloper/Public/RigVMModel/Nodes/RigVMFunctionReferenceNode.h
@@ -44,6 +44,7 @@ public:
 
 private:
 
+	virtual FText GetToolTipTextForPin(const URigVMPin* InPin) const override;
 	bool RequiresVariableRemappingInternal(TArray<FRigVMExternalVariable>& InnerVariables) const;
 
 	void SetReferencedNode(URigVMLibraryNode* InReferenceNode);
diff --git a/Engine/Source/Developer/RigVMDeveloper/Public/RigVMModel/Nodes/RigVMVariableNode.h b/Engine/Source/Developer/RigVMDeveloper/Public/RigVMModel/Nodes/RigVMVariableNode.h
--- a/Engine/Source/Developer/RigVMDeveloper/Public/RigVMModel/Nodes/RigVMVariableNode.h
+++ b/Engine/Source/Developer/RigVMDeveloper/Public/RigVMModel/Nodes/RigVMVariableNode.h
@@ -93,6 +93,10 @@ public:
 	UFUNCTION(BlueprintCallable, Category = RigVMVariableNode)
 	bool IsGetter() const;
 
+	// Returns true if this variable is a local variable
+	UFUNCTION(BlueprintCallable, Category = RigVMVariableNode)
+	bool IsLocalVariable() const;
+
 	// Returns the C++ data type of the variable
 	UFUNCTION(BlueprintCallable, Category = RigVMVariableNode)
 	FString GetCPPType() const;
diff --git a/Engine/Source/Developer/RigVMDeveloper/Public/RigVMModel/RigVMGraph.h b/Engine/Source/Developer/RigVMDeveloper/Public/RigVMModel/RigVMGraph.h
--- a/Engine/Source/Developer/RigVMDeveloper/Public/RigVMModel/RigVMGraph.h
+++ b/Engine/Source/Developer/RigVMDeveloper/Public/RigVMModel/RigVMGraph.h
@@ -129,7 +129,7 @@ public:
 
 	TArray<FRigVMExternalVariable> GetExternalVariables() const;
 
-	TArray<FRigVMGraphVariableDescription> GetLocalVariables() const { return LocalVariables; }
+	const TArray<FRigVMGraphVariableDescription>& GetLocalVariables() const { return LocalVariables; }
 
 	// Returns the modified event, which can be used to 
 	// subscribe to changes happening within the Graph.
diff --git a/Engine/Source/Developer/RigVMDeveloper/Public/RigVMModel/RigVMPin.h b/Engine/Source/Developer/RigVMDeveloper/Public/RigVMModel/RigVMPin.h
--- a/Engine/Source/Developer/RigVMDeveloper/Public/RigVMModel/RigVMPin.h
+++ b/Engine/Source/Developer/RigVMDeveloper/Public/RigVMModel/RigVMPin.h
@@ -128,7 +128,7 @@ public:
 	// corresponding SegmentPath is "Translation.X", since the transform is the
 	// storage / memory.
 	UFUNCTION(BlueprintCallable, Category = RigVMPin)
-	FString GetSegmentPath() const;
+	FString GetSegmentPath(bool bIncludeRootPin = false) const;
 
 	// Populates an array of pins which will be reduced to the same operand in the
 	// VM. This includes Source-Target pins in different nodes, pins in collapse and
