Control Rig: RigVM: Fix constant folding for more complex cases

diff --git a/Engine/Source/Developer/RigVMDeveloper/Private/RigVMCompiler/RigVMAST.cpp b/Engine/Source/Developer/RigVMDeveloper/Private/RigVMCompiler/RigVMAST.cpp
--- a/Engine/Source/Developer/RigVMDeveloper/Private/RigVMCompiler/RigVMAST.cpp
+++ b/Engine/Source/Developer/RigVMDeveloper/Private/RigVMCompiler/RigVMAST.cpp
@@ -1840,6 +1840,21 @@ bool FRigVMParserAST::FoldConstantValuesToLiterals(URigVMGraph* InGraph, URigVMC
 		return false;
 	}
 
+	// add all of the additional nodes driving these
+	for(int32 NodeToComputeIndex = 0; NodeToComputeIndex < NodesToCompute.Num(); NodeToComputeIndex++)
+	{
+		FRigVMASTProxy& ProxyToCompute = NodesToCompute[NodeToComputeIndex];
+		if(URigVMNode* NodeToCompute = ProxyToCompute.GetSubject<URigVMNode>())
+		{
+			TArray<URigVMNode*> SourceNodes = NodeToCompute->GetLinkedSourceNodes();
+			for(URigVMNode* SourceNode : SourceNodes)
+			{
+				FRigVMASTProxy SourceProxy = ProxyToCompute.GetSibling(SourceNode);
+				NodesToCompute.AddUnique(SourceProxy);
+			}
+		}
+	}
+
 	// we now know the node we need to run.
 	// let's build an temporary AST which has only those nodes
 	TSharedPtr<FRigVMParserAST> TempAST = MakeShareable(new FRigVMParserAST(InGraph, NodesToCompute));
@@ -1887,6 +1902,10 @@ bool FRigVMParserAST::FoldConstantValuesToLiterals(URigVMGraph* InGraph, URigVMC
 		FString PinHash = URigVMCompiler::GetPinHash(RootPin, RootVarExpr, false);
 		const FRigVMOperand& Operand = Operands.FindChecked(PinHash);
 		TArray<FString> DefaultValues = TempVM->GetWorkMemory().GetRegisterValueAsString(Operand, RootPin->GetCPPType(), RootPin->GetCPPTypeObject());
+
+		// FString TempDefaultValue = FString::Printf(TEXT("(%s)"), *FString::Join(DefaultValues, TEXT(",")));
+		// UE_LOG(LogRigVMDeveloper, Display, TEXT("Computed constant value '%s' = '%s'"), *PinHash, *TempDefaultValue);
+
 		if (DefaultValues.Num() == 0)
 		{
 			continue;
