Control Rig: Collapse / expand node action to filter injected nodes correctly

diff --git a/Engine/Source/Developer/RigVMDeveloper/Private/RigVMModel/RigVMController.cpp b/Engine/Source/Developer/RigVMDeveloper/Private/RigVMModel/RigVMController.cpp
--- a/Engine/Source/Developer/RigVMDeveloper/Private/RigVMModel/RigVMController.cpp
+++ b/Engine/Source/Developer/RigVMDeveloper/Private/RigVMModel/RigVMController.cpp
@@ -2364,6 +2364,26 @@ URigVMCollapseNode* URigVMController::CollapseNodes(const TArray<URigVMNode*>& I
 	TArray<URigVMLink*> LinksToRewire;
 	TArray<URigVMLink*> AllLinks = Graph->GetLinks();
 
+	auto NodeToBeCollapsed = [&Nodes](URigVMNode* InNode) -> bool
+	{
+		check(InNode);
+		
+		if(Nodes.Contains(InNode))
+		{
+			return true;
+		}
+		
+		if(InNode->IsInjected()) 
+		{
+			InNode = InNode->GetTypedOuter<URigVMNode>();
+			if(Nodes.Contains(InNode))
+			{
+				return true;
+			}
+		}
+
+		return false;
+	};
 	// find all pins to collapse. we need this to find out if
 	// we might have a parent pin of a given linked pin already 
 	// collapsed.
@@ -2371,8 +2391,8 @@ URigVMCollapseNode* URigVMController::CollapseNodes(const TArray<URigVMNode*>& I
 	{
 		URigVMPin* SourcePin = Link->GetSourcePin();
 		URigVMPin* TargetPin = Link->GetTargetPin();
-		bool bSourceToBeCollapsed = Nodes.Contains(SourcePin->GetNode());
-		bool bTargetToBeCollapsed = Nodes.Contains(TargetPin->GetNode());
+		bool bSourceToBeCollapsed = NodeToBeCollapsed(SourcePin->GetNode());
+		bool bTargetToBeCollapsed = NodeToBeCollapsed(TargetPin->GetNode());
 		if (bSourceToBeCollapsed == bTargetToBeCollapsed)
 		{
 			continue;
@@ -2396,7 +2416,7 @@ URigVMCollapseNode* URigVMController::CollapseNodes(const TArray<URigVMNode*>& I
 		{
 			continue;
 		}
-		if (!Nodes.Contains(ExecutePin->GetNode()))
+		if (!NodeToBeCollapsed(ExecutePin->GetNode()))
 		{
 			continue;
 		}
@@ -2444,7 +2464,7 @@ URigVMCollapseNode* URigVMController::CollapseNodes(const TArray<URigVMNode*>& I
 				if (SourceLinks.Num() > 0)
 				{
 					URigVMPin* PreviousExecutePin = SourceLinks[0]->GetSourcePin();
-					if (Nodes.Contains(PreviousExecutePin->GetNode()))
+					if (NodeToBeCollapsed(PreviousExecutePin->GetNode()))
 					{
 						if (Pin != IntermediateExecutePins.Last())
 						{
@@ -2464,7 +2484,7 @@ URigVMCollapseNode* URigVMController::CollapseNodes(const TArray<URigVMNode*>& I
 		{
 			continue;
 		}
-		if (!Nodes.Contains(ExecutePin->GetNode()))
+		if (!NodeToBeCollapsed(ExecutePin->GetNode()))
 		{
 			continue;
 		}
@@ -2509,8 +2529,8 @@ URigVMCollapseNode* URigVMController::CollapseNodes(const TArray<URigVMNode*>& I
 	// now looper over the links to be rewired
 	for (URigVMLink* Link : LinksToRewire)
 	{
-		bool bSourceToBeCollapsed = Nodes.Contains(Link->GetSourcePin()->GetNode());
-		bool bTargetToBeCollapsed = Nodes.Contains(Link->GetTargetPin()->GetNode());
+		bool bSourceToBeCollapsed = NodeToBeCollapsed(Link->GetSourcePin()->GetNode());
+		bool bTargetToBeCollapsed = NodeToBeCollapsed(Link->GetTargetPin()->GetNode());
 
 		URigVMPin* PinToCollapse = bSourceToBeCollapsed ? Link->GetSourcePin() : Link->GetTargetPin();
 		if (CollapsedPins.Contains(PinToCollapse))
@@ -2629,7 +2649,10 @@ URigVMCollapseNode* URigVMController::CollapseNodes(const TArray<URigVMNode*>& I
 		{
 			if (URigVMNode* ContainedNode = CollapseNode->GetContainedGraph()->FindNodeByName(ContainedNodeName))
 			{
-				SetNodePosition(ContainedNode, ContainedNode->Position - Center, false, false);
+				if(!ContainedNode->IsInjected())
+				{
+					SetNodePosition(ContainedNode, ContainedNode->Position - Center, false, false);
+				}
 			}
 		}
 
@@ -2638,7 +2661,7 @@ URigVMCollapseNode* URigVMController::CollapseNodes(const TArray<URigVMNode*>& I
 			URigVMPin* SourcePin = LinkToRewire->GetSourcePin();
 			URigVMPin* TargetPin = LinkToRewire->GetTargetPin();
 
-			if (Nodes.Contains(SourcePin->GetNode()))
+			if (NodeToBeCollapsed(SourcePin->GetNode()))
 			{
 				// if the parent pin of this was collapsed
 				// it's possible that the child pin wasn't.
@@ -2679,7 +2702,7 @@ URigVMCollapseNode* URigVMController::CollapseNodes(const TArray<URigVMNode*>& I
 		URigVMPin* SourcePin = LinkToRewire->GetSourcePin();
 		URigVMPin* TargetPin = LinkToRewire->GetTargetPin();
 
-		if (Nodes.Contains(SourcePin->GetNode()))
+		if (NodeToBeCollapsed(SourcePin->GetNode()))
 		{
 			FString SegmentPath;
 			URigVMPin* PinToCheck = SourcePin;
@@ -2897,6 +2920,11 @@ TArray<URigVMNode*> URigVMController::ExpandLibraryNode(URigVMLibraryNode* InNod
 				continue;
 			}
 
+			if(Node->IsInjected())
+			{
+				continue;
+			}
+			
 			NodeNames.Add(Node->GetFName());
 			FilteredNodes.Add(Node);
 			Bounds += Node->GetPosition();
@@ -2931,7 +2959,7 @@ TArray<URigVMNode*> URigVMController::ExpandLibraryNode(URigVMLibraryNode* InNod
 		ExpandedNodes.Add(ExpandedNode);
 	}
 
-	check(ExpandedNodeNames.Num() == NodeNames.Num());
+	check(ExpandedNodeNames.Num() >= NodeNames.Num());
 
 	TMap<FName, FName> NodeNameMap;
 	for (int32 NodeNameIndex = 0; NodeNameIndex < NodeNames.Num(); NodeNameIndex++)
@@ -3690,7 +3718,7 @@ bool URigVMController::RemoveNode(URigVMNode* InNode, bool bSetupUndoRedo, bool
 			TArray<URigVMInjectionInfo*> InjectedNodes = Pin->GetInjectedNodes();
 			for (URigVMInjectionInfo* InjectedNode : InjectedNodes)
 			{
-				RemoveNode(InjectedNode->UnitNode, bSetupUndoRedo);
+				RemoveNode(InjectedNode->UnitNode, bSetupUndoRedo, bRecursive);
 			}
 
 			BreakAllLinks(Pin, true, bSetupUndoRedo);
