Control Rig: Catch nullptrs during rename of pins

diff --git a/Engine/Source/Developer/RigVMDeveloper/Private/RigVMModel/RigVMPin.cpp b/Engine/Source/Developer/RigVMDeveloper/Private/RigVMModel/RigVMPin.cpp
--- a/Engine/Source/Developer/RigVMDeveloper/Private/RigVMModel/RigVMPin.cpp
+++ b/Engine/Source/Developer/RigVMDeveloper/Private/RigVMModel/RigVMPin.cpp
@@ -322,7 +322,10 @@ void URigVMPin::GetExposedPinChain(TArray<const URigVMPin*>& OutExposedPins) con
 			if (URigVMCollapseNode* ParentNode = Cast<URigVMCollapseNode>(Graph->GetOuter()))
 			{
 				URigVMPin* CollapseNodePin = ParentNode->FindPin(Current->GetName());
-				ToProcess.Push(CollapseNodePin);
+				if(CollapseNodePin)
+				{
+					ToProcess.Push(CollapseNodePin);
+				}
 			}
 		}
 	}		
