Control Rig: RigVM AST parser to traverse into function nodes for sub pins

diff --git a/Engine/Source/Developer/RigVMDeveloper/Private/RigVMCompiler/RigVMAST.cpp b/Engine/Source/Developer/RigVMDeveloper/Private/RigVMCompiler/RigVMAST.cpp
--- a/Engine/Source/Developer/RigVMDeveloper/Private/RigVMCompiler/RigVMAST.cpp
+++ b/Engine/Source/Developer/RigVMDeveloper/Private/RigVMCompiler/RigVMAST.cpp
@@ -2827,7 +2827,7 @@ void FRigVMParserAST::Inline(URigVMGraph* InGraph, const TArray<FRigVMASTProxy>&
 			{
 				if (ChildPin->GetDirection() == ERigVMPinDirection::Output && ChildPin->GetParentPin() == nullptr)
 				{
-					if (URigVMFunctionEntryNode* EntryNode = Cast<URigVMFunctionEntryNode>(Pin->GetNode()))
+					if (URigVMFunctionEntryNode* EntryNode = Cast<URigVMFunctionEntryNode>(ChildPin->GetNode()))
 					{
 						if (URigVMCollapseNode* OuterNode = Cast<URigVMCollapseNode>(EntryNode->GetGraph()->GetOuter()))
 						{
@@ -2840,6 +2840,18 @@ void FRigVMParserAST::Inline(URigVMGraph* InGraph, const TArray<FRigVMASTProxy>&
 							}
 						}
 					}
+					else if(URigVMLibraryNode* LibraryNode = Cast<URigVMLibraryNode>(ChildPin->GetNode()))
+					{
+						const FRigVMASTProxy ChildPinProxy = InPinProxy.GetSibling(ChildPin); 
+						if (ShouldRecursePin(ChildPinProxy))
+						{
+							FRigVMASTProxy SourceSourcePinProxy = FindSourcePin(ChildPinProxy, OutTraversalInfo);
+							if(SourceSourcePinProxy.IsValid())
+							{
+								SourcePinProxy = SourceSourcePinProxy;
+							}
+						}
+					}
 				}
 
 				TArray<URigVMLink*> SourceLinks = ChildPin->GetSourceLinks(false /* recursive */);
