Control Rig: Correct ActionIndex in certain scenarios in UndoStack

diff --git a/Engine/Source/Developer/RigVMDeveloper/Private/RigVMModel/RigVMControllerActions.cpp b/Engine/Source/Developer/RigVMDeveloper/Private/RigVMModel/RigVMControllerActions.cpp
--- a/Engine/Source/Developer/RigVMDeveloper/Private/RigVMModel/RigVMControllerActions.cpp
+++ b/Engine/Source/Developer/RigVMDeveloper/Private/RigVMModel/RigVMControllerActions.cpp
@@ -75,7 +75,8 @@ bool URigVMActionStack::Undo(URigVMController* InController)
 	}
 
 	FRigVMActionKey KeyToUndo = UndoActions.Pop();
-	ActionIndex--;
+	ActionIndex = UndoActions.Num();
+	
 	FRigVMActionWrapper Wrapper(KeyToUndo);
 	if (Wrapper.GetAction()->Undo(InController))
 	{
@@ -96,11 +97,12 @@ bool URigVMActionStack::Redo(URigVMController* InController)
 	}
 
 	FRigVMActionKey KeyToRedo = RedoActions.Pop();
-	ActionIndex++;
+
 	FRigVMActionWrapper Wrapper(KeyToRedo);
 	if (Wrapper.GetAction()->Redo(InController))
 	{
 		UndoActions.Add(KeyToRedo);
+		ActionIndex = UndoActions.Num();
 		return true;
 	}
 	return false;
