Control Rig: RigVM to treat entry, return and library nodes as NoOp

diff --git a/Engine/Source/Developer/RigVMDeveloper/Private/RigVMCompiler/RigVMAST.cpp b/Engine/Source/Developer/RigVMDeveloper/Private/RigVMCompiler/RigVMAST.cpp
--- a/Engine/Source/Developer/RigVMDeveloper/Private/RigVMCompiler/RigVMAST.cpp
+++ b/Engine/Source/Developer/RigVMDeveloper/Private/RigVMCompiler/RigVMAST.cpp
@@ -1097,7 +1097,10 @@ FRigVMExprAST* FRigVMParserAST::CreateExpressionForNode(const FRigVMASTProxy& In
 		if (InNodeProxy.IsA<URigVMRerouteNode>() ||
 			InNodeProxy.IsA<URigVMParameterNode>() ||
 			InNodeProxy.IsA<URigVMVariableNode>() ||
-			InNodeProxy.IsA<URigVMEnumNode>())
+			InNodeProxy.IsA<URigVMEnumNode>() ||
+			InNodeProxy.IsA<URigVMLibraryNode>() ||
+			InNodeProxy.IsA<URigVMFunctionEntryNode>() ||
+			InNodeProxy.IsA<URigVMFunctionReturnNode>())
 		{
 			NodeExpr = MakeExpr<FRigVMNoOpExprAST>(InNodeProxy);
 		}
