Control Rig: Add button to copy python commands to replicate a control rig blueprint

diff --git a/Engine/Plugins/Animation/ControlRig/Source/ControlRig/Private/Rigs/RigHierarchyController.cpp b/Engine/Plugins/Animation/ControlRig/Source/ControlRig/Private/Rigs/RigHierarchyController.cpp
--- a/Engine/Plugins/Animation/ControlRig/Source/ControlRig/Private/Rigs/RigHierarchyController.cpp
+++ b/Engine/Plugins/Animation/ControlRig/Source/ControlRig/Private/Rigs/RigHierarchyController.cpp
@@ -1,6 +1,7 @@
 // Copyright Epic Games, Inc. All Rights Reserved.
 
 #include "Rigs/RigHierarchyController.h"
+#include "RigVMCore/RigVMPythonUtils.h"
 
 #if WITH_EDITOR
 #include "Framework/Notifications/NotificationManager.h"
@@ -1115,6 +1116,271 @@ TArray<FRigElementKey> URigHierarchyController::ImportFromHierarchyContainer(con
 	return AddedKeys;
 }
 
+TArray<FString> URigHierarchyController::GeneratePythonCommands()
+{
+	TArray<FString> Commands;
+	Hierarchy->ForEach([&](FRigBaseElement* Element) -> bool
+	{
+		Commands.Append(GetAddElementPythonCommands(Element));
+		
+		return true;
+	});
+
+	return Commands;
+}
+
+TArray<FString> URigHierarchyController::GetAddElementPythonCommands(FRigBaseElement* Element) const
+{
+	if(FRigBoneElement* BoneElement = Cast<FRigBoneElement>(Element))
+	{
+		return GetAddBonePythonCommands(BoneElement);		
+	}
+	else if(FRigNullElement* NullElement = Cast<FRigNullElement>(Element))
+	{
+		return GetAddNullPythonCommands(NullElement);
+	}
+	else if(FRigControlElement* ControlElement = Cast<FRigControlElement>(Element))
+	{
+		return GetAddControlPythonCommands(ControlElement);
+	}
+	else if(FRigCurveElement* CurveElement = Cast<FRigCurveElement>(Element))
+	{
+		return GetAddCurvePythonCommands(CurveElement);
+	}
+	else if(FRigRigidBodyElement* RigidBodyElement = Cast<FRigRigidBodyElement>(Element))
+	{
+		return GetAddRigidBodyPythonCommands(RigidBodyElement);
+	}
+	else if(FRigSocketElement* SocketElement = Cast<FRigSocketElement>(Element))
+	{
+		ensure(false);
+	}
+	return TArray<FString>();
+}
+
+TArray<FString> URigHierarchyController::GetAddBonePythonCommands(FRigBoneElement* Bone) const
+{
+	TArray<FString> Commands;
+	if (!Bone)
+	{
+		return Commands;
+	}
+	
+	FString TransformStr = RigVMPythonUtils::TransformToPythonString(Bone->Pose.Initial.Local.Transform);
+	FString ParentKeyStr = "''";
+	if (Bone->ParentElement)
+	{
+		ParentKeyStr = Bone->ParentElement->GetKey().ToPythonString();
+	}
+	
+	if (Bone->BoneType == ERigBoneType::Imported)
+	{
+		// AddBone(FName InName, FRigElementKey InParent, FTransform InTransform, bool bTransformInGlobal = true, ERigBoneType InBoneType = ERigBoneType::User, bool bSetupUndo = false);
+		Commands.Add(FString::Printf(TEXT("hierarchy_controller.add_bone('%s', %s, %s, False, unreal.RigBoneType.%s)"),
+			*Bone->GetName().ToString(),
+			*ParentKeyStr,
+			*TransformStr,
+			Bone->BoneType == ERigBoneType::Imported ? TEXT("IMPORTED") : TEXT("USER")));
+	}
+
+	return Commands;
+}
+
+TArray<FString> URigHierarchyController::GetAddNullPythonCommands(FRigNullElement* Null) const
+{
+	FString TransformStr = RigVMPythonUtils::TransformToPythonString(Null->Pose.Initial.Local.Transform);
+
+	FString ParentKeyStr = "''";
+	if (Null->ParentElements.Num() > 0)
+	{
+		ParentKeyStr = Null->ParentElements[0]->GetKey().ToPythonString();		
+	}
+		
+	// AddNull(FName InName, FRigElementKey InParent, FTransform InTransform, bool bTransformInGlobal = true, bool bSetupUndo = false);
+	return {FString::Printf(TEXT("hierarchy_controller.add_null('%s', %s, %s, False)"),
+		*Null->GetName().ToString(),
+		*ParentKeyStr,
+		*TransformStr)};
+}
+
+TArray<FString> URigHierarchyController::GetAddControlPythonCommands(FRigControlElement* Control) const
+{
+	TArray<FString> Commands;
+	FString TransformStr = RigVMPythonUtils::TransformToPythonString(Control->Pose.Initial.Local.Transform);
+
+	FString ParentKeyStr = "''";
+	if (Control->ParentElements.Num() > 0)
+	{
+		ParentKeyStr = Control->ParentElements[0]->GetKey().ToPythonString();
+	}
+
+	FRigControlSettings& Settings = Control->Settings;
+	FString SettingsStr;
+	{
+		FString ControlNamePythonized = RigVMPythonUtils::NameToPep8(Control->GetName().ToString());
+		SettingsStr = FString::Printf(TEXT("control_settings_%s"),
+			*ControlNamePythonized);
+			
+		Commands.Add(FString::Printf(TEXT("control_settings_%s = unreal.RigControlSettings()"),
+			*ControlNamePythonized));
+		FString TypeStr;
+		switch (Settings.ControlType)
+		{
+			case ERigControlType::Bool: TypeStr = TEXT("BOOL"); break;							
+			case ERigControlType::Float: TypeStr = TEXT("FLOAT"); break;
+			case ERigControlType::Integer: TypeStr = TEXT("INTEGER"); break;
+			case ERigControlType::Position: TypeStr = TEXT("POSITION"); break;
+			case ERigControlType::Rotator: TypeStr = TEXT("POSITION"); break;
+			case ERigControlType::Scale: TypeStr = TEXT("SCALE"); break;
+			case ERigControlType::Transform: TypeStr = TEXT("TRANSFORM"); break;
+			case ERigControlType::EulerTransform: TypeStr = TEXT("EULER_TRANSFORM"); break;
+			case ERigControlType::Vector2D: TypeStr = TEXT("VECTOR2D"); break;
+			case ERigControlType::TransformNoScale: TypeStr = TEXT("TRANSFORM_NO_SCALE"); break;
+			default: ensure(false);
+		}
+		Commands.Add(FString::Printf(TEXT("control_settings_%s.control_type = unreal.RigControlType.%s"),
+										*ControlNamePythonized,
+										*TypeStr));
+		Commands.Add(FString::Printf(TEXT("control_settings_%s.animatable = %s"),
+			*ControlNamePythonized,
+			Settings.bAnimatable ? TEXT("True") : TEXT("False")));
+		Commands.Add(FString::Printf(TEXT("control_settings_%s.display_name = '%s'"),
+			*ControlNamePythonized,
+			*Settings.DisplayName.ToString()));
+		Commands.Add(FString::Printf(TEXT("control_settings_%s.draw_limits = %s"),
+			*ControlNamePythonized,
+			Settings.bDrawLimits ? TEXT("True") : TEXT("False")));
+		Commands.Add(FString::Printf(TEXT("control_settings_%s.gizmo_color = %s"),
+			*ControlNamePythonized,
+			*RigVMPythonUtils::LinearColorToPythonString(Settings.GizmoColor)));
+		Commands.Add(FString::Printf(TEXT("control_settings_%s.gizmo_enabled = %s"),
+			*ControlNamePythonized,
+			Settings.bGizmoEnabled ? TEXT("True") : TEXT("False")));
+		Commands.Add(FString::Printf(TEXT("control_settings_%s.gizmo_name = '%s'"),
+			*ControlNamePythonized,
+			*Settings.GizmoName.ToString()));
+		Commands.Add(FString::Printf(TEXT("control_settings_%s.gizmo_visible = %s"),
+			*ControlNamePythonized,
+			Settings.bGizmoVisible ? TEXT("True") : TEXT("False")));
+		Commands.Add(FString::Printf(TEXT("control_settings_%s.is_transient_control = %s"),
+			*ControlNamePythonized,
+			Settings.bIsTransientControl ? TEXT("True") : TEXT("False")));
+		Commands.Add(FString::Printf(TEXT("control_settings_%s.limit_rotation = %s"),
+				*ControlNamePythonized,
+				Settings.bLimitRotation ? TEXT("True") : TEXT("False")));
+		Commands.Add(FString::Printf(TEXT("control_settings_%s.limit_translation = %s"),
+				*ControlNamePythonized,
+				Settings.bLimitTranslation ? TEXT("True") : TEXT("False")));
+		Commands.Add(FString::Printf(TEXT("control_settings_%s.limit_scale = %s"),
+				*ControlNamePythonized,
+				Settings.bLimitScale ? TEXT("True") : TEXT("False")));
+		Commands.Add(FString::Printf(TEXT("control_settings_%s.primary_axis = unreal.RigControlAxis.%s"),
+			*ControlNamePythonized,
+			Settings.PrimaryAxis == ERigControlAxis::X ? TEXT("X") : Settings.PrimaryAxis == ERigControlAxis::Y ? TEXT("Y") : TEXT("Z")));				
+	}
+		
+	FRigControlValue Value = Hierarchy->GetControlValue(Control->GetKey(), ERigControlValueType::Initial);
+	FString ValueStr = Value.ToPythonString(Settings.ControlType);
+	
+	// AddControl(FName InName, FRigElementKey InParent, FRigControlSettings InSettings, FRigControlValue InValue, bool bSetupUndo = true);
+	Commands.Add(FString::Printf(TEXT("hierarchy_controller.add_control('%s', %s, %s, %s)"),
+		*Control->GetName().ToString(),
+		*ParentKeyStr,
+		*SettingsStr,
+		*ValueStr));
+
+	Commands.Append(GetSetControlValuePythonCommands(Control, Value, ERigControlValueType::Initial));
+	Commands.Append(GetSetControlValuePythonCommands(Control, Settings.MinimumValue, ERigControlValueType::Minimum));
+	Commands.Append(GetSetControlValuePythonCommands(Control, Settings.MaximumValue, ERigControlValueType::Maximum));
+	Commands.Append(GetSetControlOffsetTransformPythonCommands(Control, Control->Offset.Initial.Local.Transform, true, true));
+	Commands.Append(GetSetControlGizmoTransformPythonCommands(Control, Control->Gizmo.Initial.Local.Transform, true));
+
+	return Commands;
+}
+
+TArray<FString> URigHierarchyController::GetAddCurvePythonCommands(FRigCurveElement* Curve) const
+{
+	// FRigElementKey AddCurve(FName InName, float InValue = 0.f, bool bSetupUndo = true);
+	return {FString::Printf(TEXT("hierarchy_controller.add_curve('%s', %f)"),
+		*Curve->GetName().ToString(),
+		Hierarchy->GetCurveValue(Curve))};
+}
+
+TArray<FString> URigHierarchyController::GetAddRigidBodyPythonCommands(FRigRigidBodyElement* RigidBody) const
+{
+	TArray<FString> Commands;
+	FString TransformStr = RigVMPythonUtils::TransformToPythonString(RigidBody->Pose.Initial.Local.Transform);
+	
+	FString ParentKeyStr = "''";
+	if (RigidBody->ParentElement)
+	{
+		ParentKeyStr = RigidBody->ParentElement->GetKey().ToPythonString();
+	}
+
+	FString RigidBodyNamePythonized = RigVMPythonUtils::NameToPep8(RigidBody->GetName().ToString());
+	FRigRigidBodySettings& Settings = RigidBody->Settings;
+	FString SettingsStr;
+	{
+		SettingsStr =FString::Printf(TEXT("rigid_body_settings_%s"),
+			*RigidBodyNamePythonized);
+			
+		Commands.Add(FString::Printf(TEXT("rigid_body_settings_%s = unreal.RigRigidBodySettings()"),
+			*RigidBodyNamePythonized));
+
+		Commands.Add(FString::Printf(TEXT("control_settings_%s.mass = %f"),
+			*RigidBodyNamePythonized,
+			Settings.Mass));	
+	}
+	
+	// FRigElementKey AddRigidBody(FName InName, FRigElementKey InParent, FRigRigidBodySettings InSettings, FTransform InLocalTransform, bool bSetupUndo = false);
+	Commands.Add(FString::Printf(TEXT("hierarchy_controller.add_rigid_body('%s', %s, %s, %s)"),
+		*RigidBody->GetName().ToString(),
+		*ParentKeyStr,
+		*SettingsStr,
+		*TransformStr));
+
+	return Commands;
+}
+
+TArray<FString> URigHierarchyController::GetSetControlValuePythonCommands(const FRigControlElement* Control, const FRigControlValue& Value,
+                                                                          const ERigControlValueType& Type) const
+{
+	FString TypeStr;
+	switch (Type)
+	{
+		case ERigControlValueType::Initial: TypeStr = TEXT("INITIAL"); break;
+		case ERigControlValueType::Current: TypeStr = TEXT("CURRENT"); break;
+		case ERigControlValueType::Minimum: TypeStr = TEXT("MINIMUM"); break;
+		case ERigControlValueType::Maximum: TypeStr = TEXT("MAXIMUM"); break;
+		default: ensure(false);
+	}
+	return {FString::Printf(TEXT("hierarchy.set_control_value(%s, %s, unreal.RigControlValueType.%s)"),
+		*Control->GetKey().ToPythonString(),
+		*Value.ToPythonString(Control->Settings.ControlType),
+		*TypeStr)};
+}
+
+TArray<FString> URigHierarchyController::GetSetControlOffsetTransformPythonCommands(const FRigControlElement* Control,
+                                                                                    const FTransform& Offset, bool bInitial, bool bAffectChildren) const
+{
+	//SetControlOffsetTransform(FRigElementKey InKey, FTransform InTransform, bool bInitial = false, bool bAffectChildren = true, bool bSetupUndo = false)
+	return {FString::Printf(TEXT("hierarchy.set_control_offset_transform(%s, %s, %s, %s)"),
+		*Control->GetKey().ToPythonString(),
+		*RigVMPythonUtils::TransformToPythonString(Offset),
+		bInitial ? TEXT("True") : TEXT("False"),
+		bAffectChildren ? TEXT("True") : TEXT("False"))};
+}
+
+TArray<FString> URigHierarchyController::GetSetControlGizmoTransformPythonCommands(const FRigControlElement* Control,
+	const FTransform& InTransform, bool bInitial) const
+{
+	//SetControlGizmoTransform(FRigElementKey InKey, FTransform InTransform, bool bInitial = false, bool bSetupUndo = false)
+	return {FString::Printf(TEXT("hierarchy.set_control_gizmo_transform(%s, %s, %s)"),
+		*Control->GetKey().ToPythonString(),
+		*RigVMPythonUtils::TransformToPythonString(InTransform),
+		bInitial ? TEXT("True") : TEXT("False"))};
+}
+
 void URigHierarchyController::Notify(ERigHierarchyNotification InNotifType, const FRigBaseElement* InElement)
 {
 	if(!IsValid())
diff --git a/Engine/Plugins/Animation/ControlRig/Source/ControlRig/Public/Rigs/RigHierarchy.h b/Engine/Plugins/Animation/ControlRig/Source/ControlRig/Public/Rigs/RigHierarchy.h
--- a/Engine/Plugins/Animation/ControlRig/Source/ControlRig/Public/Rigs/RigHierarchy.h
+++ b/Engine/Plugins/Animation/ControlRig/Source/ControlRig/Public/Rigs/RigHierarchy.h
@@ -648,20 +648,20 @@ public:
 	}
 
 	/**
-	 * Returns all Space elements
+	 * Returns all Null elements
 	 * @param bTraverse Returns the elements in order of a depth first traversal
 	 */
-	FORCEINLINE TArray<FRigNullElement*> GetSpaces(bool bTraverse = false) const
+	FORCEINLINE TArray<FRigNullElement*> GetNulls(bool bTraverse = false) const
 	{
 		return GetElementsOfType<FRigNullElement>(bTraverse);
 	}
 
 	/**
-	 * Returns all Space elements
-	 * @param bTraverse Returns the elements in order of a depth first traversal
-	 */
-	UFUNCTION(BlueprintCallable, Category = URigHierarchy, meta = (DisplayName = "Get Spaces", ScriptName = "GetSpaces"))
-	FORCEINLINE TArray<FRigElementKey> GetSpaceKeys(bool bTraverse = true) const
+	* Returns all Null elements
+	* @param bTraverse Returns the elements in order of a depth first traversal
+	*/
+	UFUNCTION(BlueprintCallable, Category = URigHierarchy, meta = (DisplayName = "Get Nulls", ScriptName = "GetNulls"))
+	FORCEINLINE TArray<FRigElementKey> GetNullKeys(bool bTraverse = true) const
 	{
 		return GetKeysOfType<FRigNullElement>(bTraverse);
 	}
diff --git a/Engine/Plugins/Animation/ControlRig/Source/ControlRig/Public/Rigs/RigHierarchyController.h b/Engine/Plugins/Animation/ControlRig/Source/ControlRig/Public/Rigs/RigHierarchyController.h
--- a/Engine/Plugins/Animation/ControlRig/Source/ControlRig/Public/Rigs/RigHierarchyController.h
+++ b/Engine/Plugins/Animation/ControlRig/Source/ControlRig/Public/Rigs/RigHierarchyController.h
@@ -464,7 +464,29 @@ public:
 	/**
 	 * Returns a reference to the suspend notifications flag
 	 */
-	FORCEINLINE bool& GetSuspendNotificationsFlag() { return bSuspendNotifications; } 
+	FORCEINLINE bool& GetSuspendNotificationsFlag() { return bSuspendNotifications; }
+
+
+	UFUNCTION(BlueprintCallable, Category = URigHierarchyController)
+	TArray<FString> GeneratePythonCommands();
+
+	TArray<FString> GetAddElementPythonCommands(FRigBaseElement* Element) const;
+
+	TArray<FString> GetAddBonePythonCommands(FRigBoneElement* Bone) const;
+
+	TArray<FString> GetAddNullPythonCommands(FRigNullElement* Null) const;
+
+	TArray<FString> GetAddControlPythonCommands(FRigControlElement* Control) const;
+
+	TArray<FString> GetAddCurvePythonCommands(FRigCurveElement* Curve) const;
+
+	TArray<FString> GetAddRigidBodyPythonCommands(FRigRigidBodyElement* RigidBody) const;
+
+	TArray<FString> GetSetControlValuePythonCommands(const FRigControlElement* Control, const FRigControlValue& Value, const ERigControlValueType& Type) const;
+	
+	TArray<FString> GetSetControlOffsetTransformPythonCommands(const FRigControlElement* Control, const FTransform& Offset, bool bInitial = false, bool bAffectChildren = true) const;
+	
+	TArray<FString> GetSetControlGizmoTransformPythonCommands(const FRigControlElement* Control, const FTransform& Transform, bool bInitial = false) const;
 
 private:
 
diff --git a/Engine/Plugins/Animation/ControlRig/Source/ControlRig/Public/Rigs/RigHierarchyDefines.h b/Engine/Plugins/Animation/ControlRig/Source/ControlRig/Public/Rigs/RigHierarchyDefines.h
--- a/Engine/Plugins/Animation/ControlRig/Source/ControlRig/Public/Rigs/RigHierarchyDefines.h
+++ b/Engine/Plugins/Animation/ControlRig/Source/ControlRig/Public/Rigs/RigHierarchyDefines.h
@@ -231,6 +231,56 @@ public:
 		return Value;
 	}
 
+	FORCEINLINE_DEBUGGABLE FString ToPythonString(ERigControlType InControlType) const
+	{
+		FString ValueStr;
+		
+		switch (InControlType)
+		{
+			case ERigControlType::Bool: ValueStr = FString::Printf(TEXT("unreal.RigHierarchy.make_control_value_from_bool(%s)"), Get<bool>() ? TEXT("True") : TEXT("False")); break;							
+			case ERigControlType::Float: ValueStr = FString::Printf(TEXT("unreal.RigHierarchy.make_control_value_from_float(%.6f)"), Get<float>()); break;
+			case ERigControlType::Integer: ValueStr = FString::Printf(TEXT("unreal.RigHierarchy.make_control_value_from_int(%d)"), Get<int>()); break;
+			case ERigControlType::Position: ValueStr = FString::Printf(TEXT("unreal.RigHierarchy.make_control_value_from_vector(unreal.Vector(%.6f, %.6f, %.6f))"),
+				Get<FVector>().X, Get<FVector>().Y, Get<FVector>().Z); break;
+			case ERigControlType::Rotator: ValueStr = FString::Printf(TEXT("unreal.RigHierarchy.make_control_value_from_rotator(unreal.Rotator(pitch=%.6f, roll=%.6f, yaw=%.6f))"),
+				Get<FRotator>().Pitch, Get<FRotator>().Roll, Get<FRotator>().Yaw); break;
+			case ERigControlType::Scale: ValueStr = FString::Printf(TEXT("unreal.RigHierarchy.make_control_value_from_vector(unreal.Vector(%.6f, %.6f, %.6f))"),
+				Get<FVector>().X, Get<FVector>().Y, Get<FVector>().Z); break;
+			case ERigControlType::Transform: ValueStr = FString::Printf(TEXT("unreal.RigHierarchy.make_control_value_from_transform(unreal.Transform(location=[%.6f,%.6f,%.6f],rotation=[%.6f,%.6f,%.6f],scale=[%.6f,%.6f,%.6f]))"),
+				Get<FTransform>().GetLocation().X,
+				Get<FTransform>().GetLocation().Y,
+				Get<FTransform>().GetLocation().Z,
+				Get<FTransform>().Rotator().Pitch,
+				Get<FTransform>().Rotator().Yaw,
+				Get<FTransform>().Rotator().Roll,
+				Get<FTransform>().GetScale3D().X,
+				Get<FTransform>().GetScale3D().Y,
+				Get<FTransform>().GetScale3D().Z); break;
+			case ERigControlType::EulerTransform: ValueStr = FString::Printf(TEXT("unreal.RigHierarchy.make_control_value_from_euler_transform(unreal.EulerTransform(location=[%.6f,%.6f,%.6f],rotation=[%.6f,%.6f,%.6f],scale=[%.6f,%.6f,%.6f]))"),
+				Get<FEulerTransform>().Location.X,
+				Get<FEulerTransform>().Location.Y,
+				Get<FEulerTransform>().Location.Z,
+				Get<FEulerTransform>().Rotation.Euler().X,
+				Get<FEulerTransform>().Rotation.Euler().Y,
+				Get<FEulerTransform>().Rotation.Euler().Z,
+				Get<FEulerTransform>().Scale.X,
+				Get<FEulerTransform>().Scale.Y,
+				Get<FEulerTransform>().Scale.Z); break;
+			case ERigControlType::Vector2D: ValueStr = FString::Printf(TEXT("unreal.RigHierarchy.make_control_value_from_vector2d(unreal.Vector2D(%.6f, %.6f))"),
+				Get<FVector2D>().X, Get<FVector2D>().Y); break;
+			case ERigControlType::TransformNoScale: ValueStr = FString::Printf(TEXT("unreal.RigHierarchy.make_control_value_from_transform_no_scale(unreal.TransformNoScale(location=[%.6f,%.6f,%.6f],rotation=[%.6f,%.6f,%.6f]))"),
+				Get<FTransform>().GetLocation().X,
+				Get<FTransform>().GetLocation().Y,
+				Get<FTransform>().GetLocation().Z,
+				Get<FTransform>().Rotator().Pitch,
+				Get<FTransform>().Rotator().Yaw,
+				Get<FTransform>().Rotator().Roll); break;
+			default: ensure(false);
+		}
+
+		return ValueStr;
+	}
+
 	template<class T>
 	FORCEINLINE_DEBUGGABLE static FRigControlValue Make(T InValue)
 	{
@@ -910,6 +960,14 @@ public:
 		}
 		return FString();
 	}
+
+	FORCEINLINE_DEBUGGABLE FString ToPythonString() const
+	{
+		FString TypeStr = StaticEnum<ERigElementType>()->GetDisplayNameTextByIndex((int32)Type).ToUpper().ToString();
+		return FString::Printf(TEXT("unreal.RigElementKey(type=unreal.RigElementType.%s, name='%s')"),
+			*TypeStr,
+			*Name.ToString());
+	}
 };
 
 USTRUCT(BlueprintType)
diff --git a/Engine/Plugins/Animation/ControlRig/Source/ControlRigDeveloper/Private/ControlRigBlueprint.cpp b/Engine/Plugins/Animation/ControlRig/Source/ControlRigDeveloper/Private/ControlRigBlueprint.cpp
--- a/Engine/Plugins/Animation/ControlRig/Source/ControlRigDeveloper/Private/ControlRigBlueprint.cpp
+++ b/Engine/Plugins/Animation/ControlRig/Source/ControlRigDeveloper/Private/ControlRigBlueprint.cpp
@@ -25,6 +25,7 @@
 #include "Async/TaskGraphInterfaces.h"
 #include "Misc/CoreDelegates.h"
 #include "AssetRegistryModule.h"
+#include "RigVMCore/RigVMPythonUtils.h"
 
 #if WITH_EDITOR
 #include "IControlRigEditorModule.h"
@@ -1071,6 +1072,19 @@ URigVMController* UControlRigBlueprint::GetController(URigVMGraph* InGraph) cons
 	return nullptr;
 }
 
+URigVMController* UControlRigBlueprint::GetControllerByName(const FString InGraphName) const
+{
+	for (URigVMGraph* Graph : GetAllModels())
+	{
+		if (Graph->GetGraphName() == InGraphName)
+		{
+			return GetController(Graph);
+		}
+	}
+	
+	return nullptr;
+}
+
 URigVMController* UControlRigBlueprint::GetOrCreateController(URigVMGraph* InGraph)
 {
 	if (URigVMController* ExistingController = GetController(InGraph))
@@ -1254,6 +1268,168 @@ URigVMController* UControlRigBlueprint::GetOrCreateController(const UEdGraph* In
 	return GetOrCreateController(GetModel(InEdGraph));
 }
 
+TArray<FString> UControlRigBlueprint::GeneratePythonCommands(const FString InNewBlueprintName)
+{
+	TArray<FString> Commands;
+	Commands.Add(FString::Printf(TEXT("import unreal\n"
+			"unreal.load_module('ControlRigDeveloper')\n"
+			"factory = unreal.ControlRigBlueprintFactory\n"
+			"blueprint = factory.create_new_control_rig_asset(desired_package_path = '%s')\n"
+			"library = blueprint.get_local_function_library()\n"
+			"library_controller = blueprint.get_controller(library)\n"
+			"hierarchy = blueprint.hierarchy\n"
+			"hierarchy_controller = hierarchy.get_controller()\n")
+			, *InNewBlueprintName));
+
+	// Hierarchy
+	Commands.Append(HierarchyController->GeneratePythonCommands());
+		
+	// Add variables
+	for (const FBPVariableDescription& Variable : NewVariables)
+	{
+		const FRigVMExternalVariable ExternalVariable = UControlRig::GetExternalVariableFromDescription(Variable);
+
+		// FName AddMemberVariable(const FName& InName, const FString& InCPPType, bool bIsPublic = false, bool bIsReadOnly = false, FString InDefaultValue = TEXT(""));
+		Commands.Add(FString::Printf(TEXT("blueprint.add_member_variable('%s', '%s', %s, %s)"),
+					*ExternalVariable.Name.ToString(),
+					ExternalVariable.TypeObject ? *ExternalVariable.TypeObject->GetPathName() : *ExternalVariable.TypeName.ToString(),
+					ExternalVariable.bIsPublic ? TEXT("True") : TEXT("False"),
+					ExternalVariable.bIsReadOnly ? TEXT("True") : TEXT("False")));	
+	}
+	
+	// Create graphs
+	{
+		// Find all graphs to process and sort them by dependencies
+		TArray<URigVMGraph*> ProcessedGraphs;
+		while (ProcessedGraphs.Num() < GetAllModels().Num())
+		{
+			for (URigVMGraph* Graph : GetAllModels())
+			{
+				if (ProcessedGraphs.Contains(Graph))
+				{
+					continue;
+				}
+
+				bool bFoundUnprocessedReference = false;
+				for (auto Node : Graph->GetNodes())
+				{
+					if (URigVMFunctionReferenceNode* Reference = Cast<URigVMFunctionReferenceNode>(Node))
+					{
+						if (Reference->GetContainedGraph()->GetPackage() != GetPackage())
+						{
+							continue;
+						}
+				
+						if (!ProcessedGraphs.Contains(Reference->GetContainedGraph()))
+						{
+							bFoundUnprocessedReference = true;
+							break;
+						}
+					}
+					else if (URigVMCollapseNode* CollapseNode = Cast<URigVMCollapseNode>(Node))
+					{
+						if (!ProcessedGraphs.Contains(CollapseNode->GetContainedGraph()))
+						{
+							bFoundUnprocessedReference = true;
+							break;
+						}
+					}
+				}
+
+				if (!bFoundUnprocessedReference)
+				{
+					ProcessedGraphs.Add(Graph);
+				}
+			}
+		}	
+
+		// Dump python commands for each graph
+		for (URigVMGraph* Graph : ProcessedGraphs)
+		{
+			if (Graph->IsA<URigVMFunctionLibrary>())
+			{
+				continue;
+			}
+
+			URigVMController* Controller = GetController(Graph);
+			if (Graph->GetParentGraph()) 
+			{
+				// Add them all as functions (even collapsed graphs)
+				// The controller will deal with deleting collapsed graph function when it creates the collapse node
+				{						
+					// Add Function
+					Commands.Add(FString::Printf(TEXT("function_%s = library_controller.add_function_to_library('%s', mutable=%s)\ngraph = function_%s.get_contained_graph()"),
+							*RigVMPythonUtils::NameToPep8(Graph->GetGraphName()),
+							*Graph->GetGraphName(),
+							Graph->GetEntryNode()->IsMutable() ? TEXT("True") : TEXT("False"),
+							*RigVMPythonUtils::NameToPep8(Graph->GetGraphName())));
+					
+					URigVMFunctionEntryNode* EntryNode = Graph->GetEntryNode();
+					URigVMFunctionReturnNode* ReturnNode = Graph->GetReturnNode();
+					
+					// Set Entry and Return nodes in the correct position
+					{
+						//bool SetNodePositionByName(const FName& InNodeName, const FVector2D& InPosition, bool bSetupUndoRedo = true, bool bMergeUndoAction = false);
+						Commands.Add(FString::Printf(TEXT("blueprint.get_controller_by_name('%s').set_node_position_by_name('Entry', unreal.Vector2D(%f, %f))"),
+								*Graph->GetGraphName(),
+								EntryNode->GetPosition().X, 
+								EntryNode->GetPosition().Y));
+
+						Commands.Add(FString::Printf(TEXT("blueprint.get_controller_by_name('%s').set_node_position_by_name('Return', unreal.Vector2D(%f, %f))"),
+								*Graph->GetGraphName(),
+								ReturnNode->GetPosition().X, 
+								ReturnNode->GetPosition().Y));
+					}
+					
+					// Add Exposed Pins
+					{
+						for (auto Pin : EntryNode->GetPins())
+						{
+							if (Pin->GetDirection() != ERigVMPinDirection::Output)
+							{
+								continue;
+							}
+
+							// FName AddExposedPin(const FName& InPinName, ERigVMPinDirection InDirection, const FString& InCPPType, const FName& InCPPTypeObjectPath, const FString& InDefaultValue, bool bSetupUndoRedo = true);
+							Commands.Add(FString::Printf(TEXT("blueprint.get_controller_by_name('%s').add_exposed_pin('%s', unreal.RigVMPinDirection.INPUT, '%s', '%s', '%s')"),
+									*Graph->GetGraphName(),
+									*Pin->GetName(),
+									*Pin->GetCPPType(),
+									Pin->GetCPPTypeObject() ? *Pin->GetCPPTypeObject()->GetPathName() : TEXT(""),
+									*Pin->GetDefaultValue()));
+						}
+
+						for (auto Pin : ReturnNode->GetPins())
+						{
+							if (Pin->GetDirection() != ERigVMPinDirection::Input)
+							{
+								continue;
+							}
+
+							// FName AddExposedPin(const FName& InPinName, ERigVMPinDirection InDirection, const FString& InCPPType, const FName& InCPPTypeObjectPath, const FString& InDefaultValue, bool bSetupUndoRedo = true);
+							Commands.Add(FString::Printf(TEXT("blueprint.get_controller_by_name('%s').add_exposed_pin('%s', unreal.RigVMPinDirection.OUTPUT, '%s', '%s', '%s')"),
+									*Graph->GetGraphName(),
+									*Pin->GetName(), 
+									*Pin->GetCPPType(),
+									Pin->GetCPPTypeObject() ? *Pin->GetCPPTypeObject()->GetPathName() : TEXT("''"),
+									*Pin->GetDefaultValue()));
+						}
+					}
+				}
+			}
+								
+			Commands.Append(Controller->GeneratePythonCommands());
+		}
+	}
+
+	FString PreviewMeshPath = GetPreviewMesh()->GetPathName();
+	Commands.Add(FString::Printf(TEXT("blueprint.set_preview_mesh(unreal.load_object(name='%s', outer=None))"),
+		*PreviewMeshPath));
+
+	return Commands;
+}
+
+
 URigVMGraph* UControlRigBlueprint::GetTemplateModel()
 {
 #if WITH_EDITORONLY_DATA
diff --git a/Engine/Plugins/Animation/ControlRig/Source/ControlRigDeveloper/Public/ControlRigBlueprint.h b/Engine/Plugins/Animation/ControlRig/Source/ControlRigDeveloper/Public/ControlRigBlueprint.h
--- a/Engine/Plugins/Animation/ControlRig/Source/ControlRigDeveloper/Public/ControlRigBlueprint.h
+++ b/Engine/Plugins/Animation/ControlRig/Source/ControlRigDeveloper/Public/ControlRigBlueprint.h
@@ -229,12 +229,18 @@ public:
 	UFUNCTION(BlueprintCallable, Category = "Control Rig Blueprint")
 	URigVMController* GetController(URigVMGraph* InGraph = nullptr) const;
 
+	UFUNCTION(BlueprintCallable, Category = "Control Rig Blueprint")
+	URigVMController* GetControllerByName(const FString InGraphName) const;
+
 	UFUNCTION(BlueprintCallable, Category = "Control Rig Blueprint")
 	URigVMController* GetOrCreateController(URigVMGraph* InGraph = nullptr);
 
 	URigVMController* GetController(const UEdGraph* InEdGraph) const;
 	URigVMController* GetOrCreateController(const UEdGraph* InGraph);
 
+	UFUNCTION(BlueprintCallable, Category = "Control Rig Blueprint")
+	TArray<FString> GeneratePythonCommands(const FString InNewBlueprintName);
+
 	URigVMGraph* GetTemplateModel();
 	URigVMController* GetTemplateController();
 
diff --git a/Engine/Plugins/Animation/ControlRig/Source/ControlRigEditor/Private/ControlRigCompilerDetails.cpp b/Engine/Plugins/Animation/ControlRig/Source/ControlRigEditor/Private/ControlRigCompilerDetails.cpp
--- a/Engine/Plugins/Animation/ControlRig/Source/ControlRigEditor/Private/ControlRigCompilerDetails.cpp
+++ b/Engine/Plugins/Animation/ControlRig/Source/ControlRigEditor/Private/ControlRigCompilerDetails.cpp
@@ -84,6 +84,18 @@ void FRigVMCompileSettingsDetails::CustomizeChildren(TSharedRef<IPropertyHandle>
 						.Text(LOCTEXT("CopyByteCodeToClipboard", "Copy ByteCode"))
 					]
 				]
+				+ SVerticalBox::Slot()
+				[
+					SNew(SButton)
+					.OnClicked(this, &FRigVMCompileSettingsDetails::OnCopyPythonScriptClicked)
+					.ContentPadding(FMargin(2))
+					.Content()
+					[
+						SNew(STextBlock)
+						.Justification(ETextJustify::Center)
+						.Text(LOCTEXT("CopyPythonScript", "Copy Python Script"))
+					]
+				]
 			];
 	}
 }
@@ -117,4 +129,22 @@ FReply FRigVMCompileSettingsDetails::OnCopyByteCodeClicked()
 	return FReply::Handled();
 }
 
+FReply FRigVMCompileSettingsDetails::OnCopyPythonScriptClicked()
+{
+	if (BlueprintBeingCustomized)
+	{
+		FString NewName = BlueprintBeingCustomized->GetPathName();
+		int32 DotIndex = NewName.Find(TEXT("."), ESearchCase::IgnoreCase, ESearchDir::FromEnd);
+		if (DotIndex != INDEX_NONE)
+		{
+			NewName = NewName.Left(DotIndex);
+		}
+		
+		TArray<FString> Commands = BlueprintBeingCustomized->GeneratePythonCommands(NewName);
+		FString FullScript = FString::Join(Commands, TEXT("\n"));
+		FPlatformApplicationMisc::ClipboardCopy(*FullScript);
+	}
+	return FReply::Handled();
+}
+
 #undef LOCTEXT_NAMESPACE
diff --git a/Engine/Plugins/Animation/ControlRig/Source/ControlRigEditor/Private/ControlRigCompilerDetails.h b/Engine/Plugins/Animation/ControlRig/Source/ControlRigEditor/Private/ControlRigCompilerDetails.h
--- a/Engine/Plugins/Animation/ControlRig/Source/ControlRigEditor/Private/ControlRigCompilerDetails.h
+++ b/Engine/Plugins/Animation/ControlRig/Source/ControlRigEditor/Private/ControlRigCompilerDetails.h
@@ -26,6 +26,7 @@ protected:
 
 	FReply OnCopyASTClicked();
 	FReply OnCopyByteCodeClicked();
+	FReply OnCopyPythonScriptClicked();
 
 	UControlRigBlueprint* BlueprintBeingCustomized;
 };
diff --git a/Engine/Source/Developer/RigVMDeveloper/Private/RigVMModel/RigVMController.cpp b/Engine/Source/Developer/RigVMDeveloper/Private/RigVMModel/RigVMController.cpp
--- a/Engine/Source/Developer/RigVMDeveloper/Private/RigVMModel/RigVMController.cpp
+++ b/Engine/Source/Developer/RigVMDeveloper/Private/RigVMModel/RigVMController.cpp
@@ -13,6 +13,7 @@
 #include "UObject/Package.h"
 #include "Misc/CoreMisc.h"
 #include "Algo/Sort.h"
+#include "RigVMCore/RigVMPythonUtils.h"
 
 #if WITH_EDITOR
 #include "Exporters/Exporter.h"
@@ -266,6 +267,352 @@ void URigVMController::HandleModifiedEvent(ERigVMGraphNotifType InNotifType, URi
 	}
 }
 
+TArray<FString> URigVMController::GeneratePythonCommands() 
+{
+	TArray<FString> Commands;
+
+	// Add local variables
+	for (const FRigVMGraphVariableDescription& Variable : GetGraph()->LocalVariables)
+	{
+		if (Variable.CPPTypeObject)
+		{
+			// FRigVMGraphVariableDescription AddLocalVariable(const FName& InVariableName, const FString& InCPPType, UObject* InCPPTypeObject, const FString& InDefaultValue, bool bSetupUndoRedo = true);
+			Commands.Add(FString::Printf(TEXT("blueprint.get_controller_by_name('%s').add_local_variable_from_object_path('%s', '%s', '%s', '%s')"),
+						*GetGraph()->GetGraphName(),
+						*Variable.Name.ToString(),
+						*Variable.CPPType,
+						Variable.CPPTypeObject ? *Variable.CPPTypeObject->GetPathName() : TEXT(""),
+						*Variable.DefaultValue));
+		}
+		else
+		{
+			// FRigVMGraphVariableDescription AddLocalVariable(const FName& InVariableName, const FString& InCPPType, UObject* InCPPTypeObject, const FString& InDefaultValue, bool bSetupUndoRedo = true);
+			Commands.Add(FString::Printf(TEXT("blueprint.get_controller_by_name('%s').add_local_variable('%s', '%s', None, '%s')"),
+						*GetGraph()->GetGraphName(),
+						*Variable.Name.ToString(),
+						*Variable.CPPType,
+						*Variable.DefaultValue));
+		}
+	}
+	
+	
+	// All nodes (except reroutes)
+	for (URigVMNode* Node : GetGraph()->GetNodes())
+	{
+		Commands.Append(GetAddNodePythonCommands(Node));
+	}
+
+	// All links
+	for (URigVMLink* Link : GetGraph()->GetLinks())
+	{
+		URigVMPin* SourcePin = Link->GetSourcePin();
+		URigVMPin* TargetPin = Link->GetTargetPin();
+		
+		if (TargetPin->GetNode()->IsA<URigVMRerouteNode>())
+		{
+			continue;
+		}
+
+		if (SourcePin->GetInjectedNodes().Num() > 0 || TargetPin->GetInjectedNodes().Num() > 0)
+		{
+			continue;
+		}
+
+		// Iterate upstream until the source pin is not on a reroute node
+		while (SourcePin && SourcePin->GetNode()->IsA<URigVMRerouteNode>())
+		{
+			const URigVMNode* Node = SourcePin->GetNode();
+			if (Node->Pins.Num() > 0)
+			{
+				const TArray<URigVMLink*>& Links = Node->Pins[0]->GetSourceLinks();
+				if (Links.Num() > 0)
+				{
+					SourcePin = Links[0]->GetSourcePin();
+				}
+			}
+		}
+
+		//bool AddLink(const FString& InOutputPinPath, const FString& InInputPinPath, bool bSetupUndoRedo = true);
+		Commands.Add(FString::Printf(TEXT("blueprint.get_controller_by_name('%s').add_link('%s', '%s')"),
+					*GetGraph()->GetGraphName(),
+					*SourcePin->GetPinPath(),
+					*TargetPin->GetPinPath()));
+	}
+
+	// Reroutes
+	{
+		TArray<URigVMRerouteNode*> Reroutes;
+		for (URigVMNode* Node : GetGraph()->GetNodes())
+		{
+			if (URigVMRerouteNode* Reroute = Cast<URigVMRerouteNode>(Node))
+			{
+				Reroutes.Add(Reroute);
+			}
+		}
+
+		TArray<FString> ReroutesAdded;
+		while (ReroutesAdded.Num() < Reroutes.Num())
+		{
+			for (URigVMRerouteNode* Reroute : Reroutes)
+			{
+				if (Reroute->Pins.IsEmpty() || Reroute->Pins[0]->GetTargetLinks().IsEmpty())
+				{
+					continue;
+				}
+				
+				URigVMPin* TargetPin = Reroute->Pins[0]->GetTargetLinks()[0]->GetTargetPin();
+				if (!TargetPin->GetNode()->IsA<URigVMRerouteNode>() ||
+					ReroutesAdded.Contains(TargetPin->GetNode()->GetName()))
+				{
+					// Iterate upstream until the source pin is not on a reroute node
+					URigVMPin* SourcePin = Reroute->Pins[0]->GetSourceLinks()[0]->GetSourcePin();
+					while (SourcePin && SourcePin->GetNode()->IsA<URigVMRerouteNode>())
+					{
+						const URigVMNode* Node = SourcePin->GetNode();
+						if (Node->Pins.Num() > 0)
+						{
+							const TArray<URigVMLink*>& Links = Node->Pins[0]->GetSourceLinks();
+							if (Links.Num() > 0)
+							{
+								SourcePin = Links[0]->GetSourcePin();
+							}
+						}
+					}
+
+					if(SourcePin)
+					{
+						// AddRerouteNodeOnLinkPath(const FString& InLinkPinPathRepresentation, bool bShowAsFullNode, const FVector2D& InPosition = FVector2D::ZeroVector, const FString& InNodeName = TEXT(""), bool bSetupUndoRedo = true);
+						Commands.Add(FString::Printf(TEXT("blueprint.get_controller_by_name('%s').add_reroute_node_on_link_path('%s', %s, unreal.Vector2D(%f, %f), '%s')"),
+								*GetGraph()->GetGraphName(),
+								*FString::Printf(TEXT("%s -> %s"), *SourcePin->GetPinPath(), *TargetPin->GetPinPath()),
+								Reroute->GetShowsAsFullNode() ? TEXT("True") : TEXT("False"),
+								Reroute->GetPosition().X, 
+								Reroute->GetPosition().Y,
+								*Reroute->GetName()));
+
+						ReroutesAdded.Add(Reroute->GetName());
+					}
+				}
+			}
+		}
+	}
+
+	return Commands;
+}
+
+TArray<FString> URigVMController::GetAddNodePythonCommands(URigVMNode* Node) const
+{
+	TArray<FString> Commands;
+	if (URigVMUnitNode* UnitNode = Cast<URigVMUnitNode>(Node))
+	{
+		if (URigVMInjectionInfo* InjectionInfo = Cast<URigVMInjectionInfo>(UnitNode->GetOuter()))
+		{
+			//URigVMInjectionInfo* AddInjectedNodeFromStructPath(const FString& InPinPath, bool bAsInput, const FString& InScriptStructPath, const FName& InMethodName, const FName& InInputPinName, const FName& InOutputPinName, const FString& InNodeName = TEXT(""), bool bSetupUndoRedo = true);
+			Commands.Add(FString::Printf(TEXT("%s_info = blueprint.get_controller_by_name('%s').add_injected_node_from_struct_path('%s', %s, '%s', '%s', '%s', '%s', '%s')"),
+					*UnitNode->GetName(), 
+					*GetGraph()->GetGraphName(), 
+					*InjectionInfo->GetPin()->GetPinPath(), 
+					InjectionInfo->GetPin()->GetDirection() == ERigVMPinDirection::Input ? TEXT("True") : TEXT("False"), 
+					*UnitNode->GetScriptStruct()->GetPathName(), 
+					*UnitNode->GetMethodName().ToString(), 
+					*InjectionInfo->InputPin->GetName(), 
+					*InjectionInfo->OutputPin->GetName(), 
+					*UnitNode->GetName()));
+		}
+		else
+		{
+			// add_struct_node_from_struct_path(script_struct_path, method_name, position=[0.0, 0.0], node_name='', undo=True)
+			Commands.Add(FString::Printf(TEXT("blueprint.get_controller_by_name('%s').add_unit_node_from_struct_path('%s', 'Execute', unreal.Vector2D(%f, %f), '%s')"),
+					*GetGraph()->GetGraphName(),
+					*UnitNode->GetScriptStruct()->GetPathName(),
+					UnitNode->GetPosition().X, 
+					UnitNode->GetPosition().Y,
+					*UnitNode->GetName()));
+		}
+	}
+	else if (URigVMVariableNode* VariableNode = Cast<URigVMVariableNode>(Node))
+	{
+		// add_variable_node(variable_name, cpp_type, cpp_type_object, is_getter, default_value, position=[0.0, 0.0], node_name='', undo=True)
+		if (VariableNode->GetVariableDescription().CPPTypeObject)
+		{
+			Commands.Add(FString::Printf(TEXT("blueprint.get_controller_by_name('%s').add_variable_node_from_object_path('%s', '%s', '%s', %s, '%s', unreal.Vector2D(%f, %f), '%s')"),
+					*GetGraph()->GetGraphName(),
+					*VariableNode->GetVariableName().ToString(),
+					*VariableNode->GetVariableDescription().CPPType,
+					*VariableNode->GetVariableDescription().CPPTypeObject->GetPathName(),
+					VariableNode->IsGetter() ? TEXT("True") : TEXT("False"),
+					*VariableNode->GetVariableDescription().DefaultValue,
+					VariableNode->GetPosition().X, 
+					VariableNode->GetPosition().Y,
+					*VariableNode->GetName()));	
+		}
+		else
+		{
+			Commands.Add(FString::Printf(TEXT("blueprint.get_controller_by_name('%s').add_variable_node('%s', '%s', None, %s, '%s', unreal.Vector2D(%f, %f), '%s')"),
+					*GetGraph()->GetGraphName(),
+					*VariableNode->GetVariableName().ToString(),
+					*VariableNode->GetVariableDescription().CPPType,
+					VariableNode->IsGetter() ? TEXT("True") : TEXT("False"),
+					*VariableNode->GetVariableDescription().DefaultValue,
+					VariableNode->GetPosition().X, 
+					VariableNode->GetPosition().Y,
+					*VariableNode->GetName()));	
+		}
+	}
+	else if (URigVMParameterNode* ParameterNode = Cast<URigVMParameterNode>(Node))
+	{
+		// add_parameter_node_from_object_path(parameter_name, cpp_type, cpp_type_object_path, is_input, default_value, position=[0.0, 0.0], node_name='', undo=True)
+		if (ParameterNode->GetParameterDescription().CPPTypeObject)
+		{
+			Commands.Add(FString::Printf(TEXT("blueprint.get_controller_by_name('%s').add_parameter_node_from_object_path('%s', '%s', '%s', %s, '%s', unreal.Vector2D(%f, %f), '%s')"),
+					*GetGraph()->GetGraphName(),
+					*ParameterNode->GetParameterName().ToString(),
+					*ParameterNode->GetParameterDescription().CPPType,
+					*ParameterNode->GetParameterDescription().CPPTypeObject->GetPathName(),
+					ParameterNode->IsInput() ? TEXT("True") : TEXT("False"),
+					*ParameterNode->GetParameterDescription().DefaultValue,
+					ParameterNode->GetPosition().X, 
+					ParameterNode->GetPosition().Y,
+					*ParameterNode->GetName()));	
+		}
+		else
+		{
+			Commands.Add(FString::Printf(TEXT("blueprint.get_controller_by_name('%s').add_parameter_node('%s', '%s', None, %s, str(%s), unreal.Vector2D(%f, %f), '%s')"),
+					*GetGraph()->GetGraphName(),
+					*ParameterNode->GetParameterName().ToString(),
+					*ParameterNode->GetParameterDescription().CPPType,
+					ParameterNode->IsInput() ? TEXT("True") : TEXT("False"),
+					*ParameterNode->GetParameterDescription().DefaultValue,
+					ParameterNode->GetPosition().X, 
+					ParameterNode->GetPosition().Y,
+					*ParameterNode->GetName()));	
+		}
+	}
+	else if (URigVMCommentNode* CommentNode = Cast<URigVMCommentNode>(Node))
+	{
+		// add_comment_node(comment_text, position=[0.0, 0.0], size=[400.0, 300.0], color=[0.0, 0.0, 0.0, 0.0], node_name='', undo=True)
+		Commands.Add(FString::Printf(TEXT("blueprint.get_controller_by_name('%s').add_comment_node('%s', unreal.Vector2D(%f, %f), unreal.Vector2D(%f, %f), unreal.LinearColor(%f, %f, %f, %f), '%s')"),
+					*GetGraph()->GetGraphName(),
+					*CommentNode->GetCommentText(),
+					CommentNode->GetPosition().X, 
+					CommentNode->GetPosition().Y,
+					CommentNode->GetSize().X, 
+					CommentNode->GetSize().Y,
+					CommentNode->GetNodeColor().R, 
+					CommentNode->GetNodeColor().G,
+					CommentNode->GetNodeColor().B,
+					CommentNode->GetNodeColor().A,
+					*CommentNode->GetName()));	
+	}
+	else if (URigVMBranchNode* BranchNode = Cast<URigVMBranchNode>(Node))
+	{
+		// add_branch_node(position=[0.0, 0.0], node_name='', undo=True)
+		Commands.Add(FString::Printf(TEXT("blueprint.get_controller_by_name('%s').add_branch_node(unreal.Vector2D(%f, %f), '%s')"),
+						*GetGraph()->GetGraphName(),
+						BranchNode->GetPosition().X, 
+						BranchNode->GetPosition().Y,
+						*BranchNode->GetName()));	
+	}
+	else if (URigVMIfNode* IfNode = Cast<URigVMIfNode>(Node))
+	{
+		// add_if_node(cpp_type, cpp_type_object_path, position=[0.0, 0.0], node_name='', undo=True)
+		URigVMPin* ResultPin = IfNode->FindPin(URigVMIfNode::ResultName);
+		Commands.Add(FString::Printf(TEXT("blueprint.get_controller_by_name('%s').add_if_node('%s', '%s', unreal.Vector2D(%f, %f), '%s')"),
+						*GetGraph()->GetGraphName(),
+						*ResultPin->GetCPPType(),
+						*ResultPin->CPPTypeObject->GetPathName(),
+						IfNode->GetPosition().X, 
+						IfNode->GetPosition().Y,
+						*IfNode->GetName()));
+	}
+	else if (URigVMSelectNode* SelectNode = Cast<URigVMSelectNode>(Node))
+	{
+		// add_select_node(cpp_type, cpp_type_object_path, position=[0.0, 0.0], node_name='', undo=True)
+		URigVMPin* ResultPin = SelectNode->FindPin(URigVMSelectNode::ResultName);
+		Commands.Add(FString::Printf(TEXT("blueprint.get_controller_by_name('%s').add_select_node('%s', '%s', unreal.Vector2D(%f, %f), '%s')"),
+						*GetGraph()->GetGraphName(),
+						*ResultPin->GetCPPType(),
+						*ResultPin->CPPTypeObject->GetPathName(),
+						SelectNode->GetPosition().X, 
+						SelectNode->GetPosition().Y,
+						*SelectNode->GetName()));
+	}
+	else if (URigVMPrototypeNode* PrototypeNode = Cast<URigVMPrototypeNode>(Node))
+	{
+		// add_prototype_node(notation, position=[0.0, 0.0], node_name='', undo=True)
+		Commands.Add(FString::Printf(TEXT("blueprint.get_controller_by_name('%s').add_prototype_node('%s', unreal.Vector2D(%f, %f), '%s')"),
+						*GetGraph()->GetGraphName(),
+						*PrototypeNode->GetNotation().ToString(),
+						PrototypeNode->GetPosition().X, 
+						PrototypeNode->GetPosition().Y,
+						*PrototypeNode->GetName()));
+	}
+	else if (URigVMEnumNode* EnumNode = Cast<URigVMEnumNode>(Node))
+	{
+		// add_enum_node(cpp_type_object_path, position=[0.0, 0.0], node_name='', undo=True)
+		Commands.Add(FString::Printf(TEXT("blueprint.get_controller_by_name('%s').add_enum_node('%s', unreal.Vector2D(%f, %f), '%s')"),
+						*GetGraph()->GetGraphName(),
+						*EnumNode->GetCPPTypeObject()->GetPathName(),
+						EnumNode->GetPosition().X, 
+						EnumNode->GetPosition().Y,
+						*EnumNode->GetName()));
+	}
+	else if (URigVMLibraryNode* LibraryNode = Cast<URigVMLibraryNode>(Node))
+	{
+		// AddFunctionReferenceNode(URigVMLibraryNode* InFunctionDefinition, const FVector2D& InNodePosition = FVector2D::ZeroVector, const FString& InNodeName = TEXT(""), bool bSetupUndoRedo = true);
+		Commands.Add(FString::Printf(TEXT("blueprint.get_controller_by_name('%s').add_function_reference_node(function_%s, unreal.Vector2D(%f, %f), '%s')"),
+						*GetGraph()->GetGraphName(),
+						*RigVMPythonUtils::NameToPep8(LibraryNode->GetContainedGraph()->GetGraphName()),
+						LibraryNode->GetPosition().X, 
+						LibraryNode->GetPosition().Y,
+						*LibraryNode->GetName()));
+
+		if (Node->IsA<URigVMCollapseNode>())
+		{
+			Commands.Add(FString::Printf(TEXT("blueprint.get_controller_by_name('%s').promote_function_reference_node_to_collapse_node('%s')"),
+					*GetGraph()->GetGraphName(),
+					*Node->GetName()));
+			Commands.Add(FString::Printf(TEXT("library_controller.remove_function_from_library('%s')"),
+					*LibraryNode->GetContainedGraph()->GetGraphName()));
+		}
+	}
+	else if (URigVMRerouteNode* RerouteNode = Cast<URigVMRerouteNode>(Node))
+	{
+		// Do nothing, we need to create the links first
+	}
+	else if (Node->IsA<URigVMFunctionEntryNode>() || Node->IsA<URigVMFunctionReturnNode>())
+	{
+		
+		
+	}
+	else
+	{
+		ensure(false);
+	}
+
+	if (!Commands.IsEmpty())
+	{
+		for (const URigVMPin* Pin : Node->GetPins())
+		{
+			if (Pin->GetDirection() == ERigVMPinDirection::Output)
+			{
+				continue;
+			}
+			
+			const FString DefaultValue =Pin->GetDefaultValue();
+			if (!DefaultValue.IsEmpty())
+			{
+				Commands.Add(FString::Printf(TEXT("blueprint.get_controller_by_name('%s').set_pin_default_value('%s', '%s')"),
+							*GetGraph()->GetGraphName(),
+							*Pin->GetPinPath(),
+							*Pin->GetDefaultValue()));
+			}
+		}
+	}
+
+	return Commands;
+}
+
 #if WITH_EDITOR
 
 URigVMUnitNode* URigVMController::AddUnitNode(UScriptStruct* InScriptStruct, const FName& InMethodName, const FVector2D& InPosition, const FString& InNodeName, bool bSetupUndoRedo)
@@ -3637,8 +3984,8 @@ bool URigVMController::RemoveNode(URigVMNode* InNode, bool bSetupUndoRedo, bool
 	if (bSetupUndoRedo)
 	{
 		// don't allow deletion of function entry / return nodes
-		if (Cast<URigVMFunctionEntryNode>(InNode) != nullptr ||
-			Cast<URigVMFunctionReturnNode>(InNode) != nullptr)
+		if ((Cast<URigVMFunctionEntryNode>(InNode) != nullptr && InNode->GetName() == TEXT("Entry")) ||
+			(Cast<URigVMFunctionReturnNode>(InNode) != nullptr && InNode->GetName() == TEXT("Return")))
 		{
 			return false;
 		}
@@ -6770,6 +7117,28 @@ FRigVMGraphVariableDescription URigVMController::AddLocalVariable(const FName& I
 	return NewVariable;
 }
 
+FRigVMGraphVariableDescription URigVMController::AddLocalVariableFromObjectPath(const FName& InVariableName, const FString& InCPPType, const FString& InCPPTypeObjectPath, const FString& InDefaultValue, bool bSetupUndoRedo)
+{
+	FRigVMGraphVariableDescription Description;
+	if (!IsValidGraph())
+	{
+		return Description;
+	}
+
+	UObject* CPPTypeObject = nullptr;
+	if (!InCPPTypeObjectPath.IsEmpty())
+	{
+		CPPTypeObject = URigVMPin::FindObjectFromCPPTypeObjectPath<UObject>(InCPPTypeObjectPath);
+		if (CPPTypeObject == nullptr)
+		{
+			ReportErrorf(TEXT("Cannot find cpp type object for path '%s'."), *InCPPTypeObjectPath);
+			return Description;
+		}
+	}
+
+	return AddLocalVariable(InVariableName, InCPPType, CPPTypeObject, InDefaultValue, bSetupUndoRedo);
+}
+
 bool URigVMController::RemoveLocalVariable(const FName& InVariableName, bool bSetupUndoRedo)
 {
 	if (!IsValidGraph())
diff --git a/Engine/Source/Developer/RigVMDeveloper/Private/RigVMModel/RigVMControllerActions.cpp b/Engine/Source/Developer/RigVMDeveloper/Private/RigVMModel/RigVMControllerActions.cpp
--- a/Engine/Source/Developer/RigVMDeveloper/Private/RigVMModel/RigVMControllerActions.cpp
+++ b/Engine/Source/Developer/RigVMDeveloper/Private/RigVMModel/RigVMControllerActions.cpp
@@ -720,6 +720,10 @@ FRigVMRemoveNodeAction::FRigVMRemoveNodeAction(URigVMNode* InNode, URigVMControl
 	{
 		InverseAction.AddAction(FRigVMImportNodeFromTextAction(LibraryNode, InController));
 	}
+	else if (InNode->IsA<URigVMFunctionEntryNode>() || InNode->IsA<URigVMFunctionReturnNode>())
+	{
+		// Do nothing
+	}
 	else
 	{
 		ensure(false);
diff --git a/Engine/Source/Developer/RigVMDeveloper/Public/RigVMModel/Nodes/RigVMLibraryNode.h b/Engine/Source/Developer/RigVMDeveloper/Public/RigVMModel/Nodes/RigVMLibraryNode.h
--- a/Engine/Source/Developer/RigVMDeveloper/Public/RigVMModel/Nodes/RigVMLibraryNode.h
+++ b/Engine/Source/Developer/RigVMDeveloper/Public/RigVMModel/Nodes/RigVMLibraryNode.h
@@ -33,8 +33,13 @@ public:
 	virtual FString GetNodeCategory() const { return FString(); }
 	virtual FString GetNodeKeywords() const { return FString(); }
 	virtual FString GetNodeDescription() const { return FString(); }
+	
+	UFUNCTION(BlueprintCallable, Category = RigVMLibraryNode)
 	virtual URigVMFunctionLibrary* GetLibrary() const { return nullptr; }
+
+	UFUNCTION(BlueprintCallable, Category = RigVMLibraryNode)
 	virtual URigVMGraph* GetContainedGraph() const { return nullptr; }
+	
 	virtual const TArray<URigVMNode*>& GetContainedNodes() const;
 	virtual const TArray<URigVMLink*>& GetContainedLinks() const;
 	virtual URigVMFunctionEntryNode* GetEntryNode() const;
diff --git a/Engine/Source/Developer/RigVMDeveloper/Public/RigVMModel/RigVMController.h b/Engine/Source/Developer/RigVMDeveloper/Public/RigVMModel/RigVMController.h
--- a/Engine/Source/Developer/RigVMDeveloper/Public/RigVMModel/RigVMController.h
+++ b/Engine/Source/Developer/RigVMDeveloper/Public/RigVMModel/RigVMController.h
@@ -122,6 +122,11 @@ public:
 	UFUNCTION(BlueprintPure, Category = RigVMController)
 	bool IsReportingEnabled() const { return bReportWarningsAndErrors; }
 
+	UFUNCTION(BlueprintCallable, Category = RigVMController)
+	TArray<FString> GeneratePythonCommands();
+
+	TArray<FString> GetAddNodePythonCommands(URigVMNode* Node) const;
+
 #if WITH_EDITOR
 	// Note: The functions below are scoped with WITH_EDITOR since we are considering
 	// to move this code into the runtime in the future. Right now there's a dependency
@@ -596,6 +601,10 @@ public:
 	UFUNCTION(BlueprintCallable, Category = RigVMController)
 	FRigVMGraphVariableDescription AddLocalVariable(const FName& InVariableName, const FString& InCPPType, UObject* InCPPTypeObject, const FString& InDefaultValue, bool bSetupUndoRedo = true);
 
+	// Add a local variable to the graph given a struct object path name.
+	UFUNCTION(BlueprintCallable, Category = RigVMController)
+	FRigVMGraphVariableDescription AddLocalVariableFromObjectPath(const FName& InVariableName, const FString& InCPPType, const FString& InCPPTypeObjectPath, const FString& InDefaultValue, bool bSetupUndoRedo = true);
+
 	// Remove a local variable from the graph
 	UFUNCTION(BlueprintCallable, Category = RigVMController)
 	bool RemoveLocalVariable(const FName& InVariableName, bool bSetupUndoRedo = true);
diff --git a/Engine/Source/Runtime/RigVM/Public/RigVMCore/RigVMPythonUtils.h b/Engine/Source/Runtime/RigVM/Public/RigVMCore/RigVMPythonUtils.h
new file mode 100644
--- /dev/null
+++ b/Engine/Source/Runtime/RigVM/Public/RigVMCore/RigVMPythonUtils.h
@@ -0,0 +1,53 @@
+﻿// Copyright Epic Games, Inc. All Rights Reserved.
+
+#pragma once
+
+namespace RigVMPythonUtils
+{
+	FString NameToPep8(const FString& Name)
+	{
+		// Wish we could use PyGenUtil::PythonizeName, but unfortunately it's private
+	
+		const FString NameNoSpaces = Name.Replace(TEXT(" "), TEXT("_"));
+		FString Result;
+
+		for (const TCHAR& Char : NameNoSpaces)
+		{
+			if (FChar::IsUpper(Char))
+			{
+				if (!Result.IsEmpty() && !Result.EndsWith(TEXT("_")))
+				{
+					Result.AppendChar(TEXT('_'));
+				}
+				Result.AppendChar(FChar::ToLower(Char));
+			}
+			else
+			{
+				Result.AppendChar(Char);
+			}
+		}
+		return Result;
+	}
+
+	FString TransformToPythonString(const FTransform& Transform)
+	{
+		return FString::Printf(TEXT("unreal.Transform(location=[%f,%f,%f],rotation=[%f,%f,%f],scale=[%f,%f,%f])"),
+			Transform.GetLocation().X,
+			Transform.GetLocation().Y,
+			Transform.GetLocation().Z,
+			Transform.Rotator().Pitch,
+			Transform.Rotator().Yaw,
+			Transform.Rotator().Roll,
+			Transform.GetScale3D().X,
+			Transform.GetScale3D().Y,
+			Transform.GetScale3D().Z);
+	}
+
+	FString LinearColorToPythonString(const FLinearColor& Color)
+	{
+		return FString::Printf(TEXT("unreal.LinearColor(%f, %f, %f, %f)"),
+			Color.R, Color.G, Color.B, Color.A);
+	}
+
+
+}
\ No newline at end of file
