Control Rig: Support variable binding in details panel.

Also added support for local variables in binding.
+ React to removing local variable
+ React to renaming local variable
+ React to changing local variable type

diff --git a/Engine/Plugins/Animation/ControlRig/Source/ControlRig/Public/Units/RigUnit.h b/Engine/Plugins/Animation/ControlRig/Source/ControlRig/Public/Units/RigUnit.h
--- a/Engine/Plugins/Animation/ControlRig/Source/ControlRig/Public/Units/RigUnit.h
+++ b/Engine/Plugins/Animation/ControlRig/Source/ControlRig/Public/Units/RigUnit.h
@@ -35,7 +35,7 @@ struct CONTROLRIG_API FRigUnit : public FRigVMStruct
 		static FName MethodName = FName(NAME_None);
 		if(MethodName.IsNone())
 		{
-			MethodName = TEXT("Execute");
+			MethodName = FRigVMStruct::ExecuteName;
 		}
 		return MethodName;
 	}
diff --git a/Engine/Plugins/Animation/ControlRig/Source/ControlRigEditor/Private/ControlRigElementDetails.cpp b/Engine/Plugins/Animation/ControlRig/Source/ControlRigEditor/Private/ControlRigElementDetails.cpp
--- a/Engine/Plugins/Animation/ControlRig/Source/ControlRigEditor/Private/ControlRigElementDetails.cpp
+++ b/Engine/Plugins/Animation/ControlRig/Source/ControlRigEditor/Private/ControlRigElementDetails.cpp
@@ -20,6 +20,9 @@
 #include "PropertyCustomizationHelpers.h"
 #include "SEnumCombo.h"
 #include "ControlRig/Private/Units/Execution/RigUnit_BeginExecution.h"
+#include "RigVMModel/RigVMGraph.h"
+#include "RigVMModel/RigVMNode.h"
+#include "Graph/SControlRigGraphPinVariableBinding.h"
 
 #define LOCTEXT_NAMESPACE "ControlRigElementDetails"
 
@@ -430,6 +433,24 @@ void FRigUnitDetails::CustomizeDetails(IDetailLayoutBuilder& DetailBuilder)
 		return;
 	}
 
+	URigVMGraph* Model = GraphBeingCustomized->GetModel();
+	if(Model == nullptr)
+	{
+		return;
+	}
+
+	const TArray<FName> SelectedNodeNames = Model->GetSelectNodes();
+	if(SelectedNodeNames.Num() == 0)
+	{
+		return;
+	}
+
+	URigVMNode* ModelNode = Model->FindNodeByName(SelectedNodeNames[0]);
+	if(ModelNode == nullptr)
+	{
+		return;
+	}
+
 	UScriptStruct* ScriptStruct = Cast<UScriptStruct>((UStruct*)StructBeingCustomized->GetStruct());
 	check(ScriptStruct);
 
@@ -445,6 +466,29 @@ void FRigUnitDetails::CustomizeDetails(IDetailLayoutBuilder& DetailBuilder)
 		}
 		DetailBuilder.HideProperty(PropertyHandle);
 
+		URigVMPin* ModelPin = ModelNode->FindPin(Property->GetName());
+		if(ModelPin == nullptr)
+		{
+			continue;
+		}
+
+		if(ModelPin->IsBoundToVariable())
+		{
+			CategoryBuilder.AddCustomRow(FText::FromString(Property->GetName()))
+			.NameContent()
+			[
+				PropertyHandle->CreatePropertyNameWidget()
+			]
+			.ValueContent()
+			[
+				SNew(SControlRigVariableBinding)
+					.ModelPin(ModelPin)
+					.Blueprint(BlueprintBeingCustomized)
+			];
+
+			continue;
+		}
+
 		if (FNameProperty* NameProperty = CastField<FNameProperty>(Property))
 		{
 			FString CustomWidgetName = NameProperty->GetMetaData(TEXT("CustomWidget"));
diff --git a/Engine/Plugins/Animation/ControlRig/Source/ControlRigEditor/Private/Graph/SControlRigGraphNode.cpp b/Engine/Plugins/Animation/ControlRig/Source/ControlRigEditor/Private/Graph/SControlRigGraphNode.cpp
--- a/Engine/Plugins/Animation/ControlRig/Source/ControlRigEditor/Private/Graph/SControlRigGraphNode.cpp
+++ b/Engine/Plugins/Animation/ControlRig/Source/ControlRigEditor/Private/Graph/SControlRigGraphNode.cpp
@@ -176,7 +176,7 @@ void SControlRigGraphNode::Construct( const FArguments& InArgs )
 		PinInfo.bExpanded = ModelPin->IsExpanded();
 		PinInfo.ModelPinPath = ModelPin->GetPinPath();
 
-		PinInfo.bHideInputWidget = ModelPin->IsStruct() || PinInfo.bIsContainer;
+		PinInfo.bHideInputWidget = (!ModelPin->IsBoundToVariable()) && (ModelPin->IsStruct() || PinInfo.bIsContainer);
 		if (!PinInfo.bHideInputWidget)
 		{
 			if (ModelPin->GetSubPins().Num() == 0)
diff --git a/Engine/Plugins/Animation/ControlRig/Source/ControlRigEditor/Private/Graph/SControlRigGraphPinVariableBinding.cpp b/Engine/Plugins/Animation/ControlRig/Source/ControlRigEditor/Private/Graph/SControlRigGraphPinVariableBinding.cpp
--- a/Engine/Plugins/Animation/ControlRig/Source/ControlRigEditor/Private/Graph/SControlRigGraphPinVariableBinding.cpp
+++ b/Engine/Plugins/Animation/ControlRig/Source/ControlRigEditor/Private/Graph/SControlRigGraphPinVariableBinding.cpp
@@ -7,6 +7,8 @@
 #include "DetailLayoutBuilder.h"
 #include "RigVMModel/RigVMController.h"
 #include "ControlRigBlueprintGeneratedClass.h"
+#include "Graph/ControlRigGraphSchema.h"
+#include "Kismet2/BlueprintEditorUtils.h"
 
 #define LOCTEXT_NAMESPACE "SControlRigGraphPinVariableBinding"
 
@@ -37,6 +39,13 @@ void SControlRigVariableBinding::Construct(const FArguments& InArgs)
 	BindingArgs.bAllowStructMemberBindings = false;
 	BindingArgs.bAllowUObjectFunctions = false;
 
+	BindingArgs.MenuExtender = MakeShareable(new FExtender);
+	BindingArgs.MenuExtender->AddMenuExtension(
+		"Properties",
+		EExtensionHook::After,
+		nullptr,
+		FMenuExtensionDelegate::CreateSP(this, &SControlRigVariableBinding::FillLocalVariableMenu));
+
 	this->ChildSlot
 	[
 		PropertyAccessEditor.MakePropertyBindingWidget(Blueprint, BindingArgs)
@@ -72,8 +81,7 @@ FLinearColor SControlRigVariableBinding::GetBindingColor() const
 {
 	if (Blueprint)
 	{
-		const UEdGraphSchema_K2* Schema = GetDefault<UEdGraphSchema_K2>();
-
+		const UControlRigGraphSchema* Schema = GetDefault<UControlRigGraphSchema>();
 		FName BoundVariable(NAME_None);
 
 		if(ModelPin)
@@ -96,6 +104,28 @@ FLinearColor SControlRigVariableBinding::GetBindingColor() const
 				return Schema->GetPinTypeColor(VariableDescription.VarType);
 			}
 		}
+
+		URigVMGraph* Model = ModelPin->GetGraph();
+		if(Model == nullptr)
+		{
+			return  FLinearColor::Red;
+		}
+
+		const TArray<FRigVMGraphVariableDescription>& LocalVariables =  Model->GetLocalVariables();
+		for(const FRigVMGraphVariableDescription& LocalVariable : LocalVariables)
+		{
+			const FRigVMExternalVariable ExternalVariable = LocalVariable.ToExternalVariable();
+			if(!ExternalVariable.IsValid(true))
+			{
+				continue;
+			}
+
+			if (ExternalVariable.Name == BoundVariable)
+			{
+				const FEdGraphPinType PinType = UControlRig::GetPinTypeFromExternalVariable(ExternalVariable);
+				return Schema->GetPinTypeColor(PinType);
+			}
+		}
 	}
 	return FLinearColor::White;
 }
@@ -198,6 +228,115 @@ void SControlRigVariableBinding::OnRemoveBinding(FName InPropertyName)
 	}
 }
 
+void SControlRigVariableBinding::FillLocalVariableMenu(FMenuBuilder& MenuBuilder)
+{
+	if(ModelPin == nullptr)
+	{
+		return;
+	}
+
+	URigVMGraph* Model = ModelPin->GetGraph();
+	if(Model == nullptr)
+	{
+		return;
+	}
+
+	int32 ValidLocalVariables = 0;
+	const TArray<FRigVMGraphVariableDescription>& LocalVariables =  Model->GetLocalVariables();
+	for(const FRigVMGraphVariableDescription& LocalVariable : LocalVariables)
+	{
+		const FRigVMExternalVariable ExternalVariable = LocalVariable.ToExternalVariable();
+		if(!ExternalVariable.IsValid(true))
+		{
+			continue;
+		}
+
+		if(!ModelPin->CanBeBoundToVariable(ExternalVariable))
+		{
+			continue;
+		}
+
+		ValidLocalVariables++;
+	}
+
+	if(ValidLocalVariables == 0)
+	{
+		return;
+	}
+	
+	MenuBuilder.BeginSection("LocalVariables", LOCTEXT("LocalVariables", "Local Variables"));
+	{
+		static FName PropertyIcon(TEXT("Kismet.VariableList.TypeIcon"));
+		const UControlRigGraphSchema* Schema = GetDefault<UControlRigGraphSchema>();
+
+		for(const FRigVMGraphVariableDescription& LocalVariable : LocalVariables)
+		{
+			const FRigVMExternalVariable ExternalVariable = LocalVariable.ToExternalVariable();
+			if(!ExternalVariable.IsValid(true))
+			{
+				continue;
+			}
+
+			if(!ModelPin->CanBeBoundToVariable(ExternalVariable))
+			{
+				continue;
+			}
+			
+			const FEdGraphPinType PinType = UControlRig::GetPinTypeFromExternalVariable(ExternalVariable);
+
+			MenuBuilder.AddMenuEntry(
+				FUIAction(FExecuteAction::CreateSP(this, &SControlRigVariableBinding::HandleBindToLocalVariable, LocalVariable)),
+				SNew(SHorizontalBox)
+					+SHorizontalBox::Slot()
+					.AutoWidth()
+					[
+						SNew(SSpacer)
+						.Size(FVector2D(18.0f, 0.0f))
+					]
+					+SHorizontalBox::Slot()
+					.AutoWidth()
+					.VAlign(VAlign_Center)
+					.Padding(1.0f, 0.0f)
+					[
+						SNew(SImage)
+						.Image(FBlueprintEditorUtils::GetIconFromPin(PinType, true))
+						.ColorAndOpacity(Schema->GetPinTypeColor(PinType))
+					]
+					+SHorizontalBox::Slot()
+					.AutoWidth()
+					.VAlign(VAlign_Center)
+					.Padding(4.0f, 0.0f)
+					[
+						SNew(STextBlock)
+						.Text(FText::FromName(LocalVariable.Name))
+					]);
+		}
+	}
+	MenuBuilder.EndSection(); // Local Variables
+}
+
+void SControlRigVariableBinding::HandleBindToLocalVariable(FRigVMGraphVariableDescription InLocalVariable)
+{
+	if((ModelPin == nullptr) || (Blueprint == nullptr))
+	{
+		return;
+	}
+
+	URigVMGraph* Model = ModelPin->GetGraph();
+	if(Model == nullptr)
+	{
+		return;
+	}
+
+	URigVMController* Controller = Blueprint->GetOrCreateController(Model);
+	if(Controller == nullptr)
+	{
+		return;
+	}
+
+	Controller->BindPinToVariable(ModelPin->GetPinPath(), InLocalVariable.Name.ToString());
+}
+
 /////////////////////////////////////////////////////////////////////////////////////////////////////////////
 
 void SControlRigGraphPinVariableBinding::Construct(const FArguments& InArgs, UEdGraphPin* InGraphPinObj)
diff --git a/Engine/Plugins/Animation/ControlRig/Source/ControlRigEditor/Private/Graph/SControlRigGraphPinVariableBinding.h b/Engine/Plugins/Animation/ControlRig/Source/ControlRigEditor/Private/Graph/SControlRigGraphPinVariableBinding.h
--- a/Engine/Plugins/Animation/ControlRig/Source/ControlRigEditor/Private/Graph/SControlRigGraphPinVariableBinding.h
+++ b/Engine/Plugins/Animation/ControlRig/Source/ControlRigEditor/Private/Graph/SControlRigGraphPinVariableBinding.h
@@ -43,6 +43,8 @@ protected:
 	void OnAddBinding(FName InPropertyName, const TArray<FBindingChainElement>& InBindingChain);
 	bool OnCanRemoveBinding(FName InPropertyName);
 	void OnRemoveBinding(FName InPropertyName);
+	void FillLocalVariableMenu( FMenuBuilder& MenuBuilder );
+	void HandleBindToLocalVariable(FRigVMGraphVariableDescription InLocalVariable);
 
 	URigVMPin* ModelPin;
 	URigVMFunctionReferenceNode* FunctionReferenceNode;
diff --git a/Engine/Source/Developer/RigVMDeveloper/Private/RigVMCompiler/RigVMCompiler.cpp b/Engine/Source/Developer/RigVMDeveloper/Private/RigVMCompiler/RigVMCompiler.cpp
--- a/Engine/Source/Developer/RigVMDeveloper/Private/RigVMCompiler/RigVMCompiler.cpp
+++ b/Engine/Source/Developer/RigVMDeveloper/Private/RigVMCompiler/RigVMCompiler.cpp
@@ -1700,6 +1700,14 @@ FRigVMOperand URigVMCompiler::FindOrAddRegister(const FRigVMVarExprAST* InVarExp
 					Operand.RegisterOffset = WorkData.VM->GetWorkMemory().GetOrAddRegisterOffset(Operand.GetRegisterIndex(), ScriptStruct, SegmentPath, 0 /*ArrayIndex */);
 				}
 			}
+
+			ExistingOperand = WorkData.PinPathToOperand->Find(FString::Printf(TEXT("LocalVariable::%s"), *VariableName));
+			if (ExistingOperand)
+			{
+				Operand = *ExistingOperand;
+				Operand.RegisterOffset = INDEX_NONE;
+				check(SegmentPath.IsEmpty());
+			}
 		}
 	}
 
diff --git a/Engine/Source/Developer/RigVMDeveloper/Private/RigVMModel/RigVMController.cpp b/Engine/Source/Developer/RigVMDeveloper/Private/RigVMModel/RigVMController.cpp
--- a/Engine/Source/Developer/RigVMDeveloper/Private/RigVMModel/RigVMController.cpp
+++ b/Engine/Source/Developer/RigVMDeveloper/Private/RigVMModel/RigVMController.cpp
@@ -5743,7 +5743,7 @@ bool URigVMController::BindPinToVariable(URigVMPin* InPin, const FString& InNewB
 		FString VariableName = InNewBoundVariablePath, SegmentPath;
 		InNewBoundVariablePath.Split(TEXT("."), &VariableName, &SegmentPath);
 
-		FRigVMExternalVariable ExternalVariable = GetExternalVariableByName(*VariableName);
+		FRigVMExternalVariable ExternalVariable = GetVariableByName(*VariableName);
 		if (ExternalVariable.IsValid(true))
 		{
 			FRigVMRegisterOffset RegisterOffset;
@@ -5853,7 +5853,7 @@ bool URigVMController::MakeBindingsFromVariableNode(URigVMVariableNode* InNode,
 	}
 
 	FName VariableName = InNode->GetVariableName();
-	FRigVMExternalVariable ExternalVariable = GetExternalVariableByName(VariableName);
+	FRigVMExternalVariable ExternalVariable = GetVariableByName(VariableName);
 	if (!ExternalVariable.IsValid(true /* allow nullptr */))
 	{
 		return false;
@@ -5943,7 +5943,7 @@ bool URigVMController::PromotePinToVariable(URigVMPin* InPin, bool bCreateVariab
 	FString SegmentPath;
 	if (InPin->IsBoundToVariable())
 	{
-		VariableForPin = GetExternalVariableByName(*InPin->GetBoundVariableName());
+		VariableForPin = GetVariableByName(*InPin->GetBoundVariableName());
 		check(VariableForPin.IsValid(true /* allow nullptr */));
 		SegmentPath = InPin->GetBoundVariablePath();
 		if (SegmentPath.StartsWith(VariableForPin.Name.ToString() + TEXT(".")))
@@ -5969,7 +5969,7 @@ bool URigVMController::PromotePinToVariable(URigVMPin* InPin, bool bCreateVariab
 			return false;
 		}
 
-		VariableForPin = GetExternalVariableByName(VariableName);
+		VariableForPin = GetVariableByName(VariableName);
 		if (!VariableForPin.IsValid(true /* allow nullptr*/))
 		{
 			return false;
@@ -6929,7 +6929,7 @@ bool URigVMController::SetRemappedVariable(URigVMFunctionReferenceNode* InFuncti
 	FRigVMExternalVariable InnerExternalVariable;
 	{
 		FRigVMControllerGraphGuard GraphGuard(this, InFunctionRefNode->GetContainedGraph());
-		InnerExternalVariable = GetExternalVariableByName(InInnerVariableName);
+		InnerExternalVariable = GetVariableByName(InInnerVariableName);
 	}
 
 	if(!InnerExternalVariable.IsValid(true))
@@ -6947,7 +6947,7 @@ bool URigVMController::SetRemappedVariable(URigVMFunctionReferenceNode* InFuncti
 	}
 	else
 	{
-		const FRigVMExternalVariable OuterExternalVariable = GetExternalVariableByName(InOuterVariableName);
+		const FRigVMExternalVariable OuterExternalVariable = GetVariableByName(InOuterVariableName);
 		if(!OuterExternalVariable.IsValid(true))
 		{
 			ReportErrorf(TEXT("External variable '%s' cannot be found."), *InOuterVariableName.ToString());
@@ -7162,6 +7162,32 @@ bool URigVMController::RemoveLocalVariable(const FName& InVariableName, bool bSe
 
 	if (FoundIndex != INDEX_NONE)
 	{
+		const FString VarNameStr = InVariableName.ToString();
+		TArray<URigVMNode*> Nodes = Graph->GetNodes();
+		for (URigVMNode* Node : Nodes)
+		{
+			if (URigVMVariableNode* VariableNode = Cast<URigVMVariableNode>(Node))
+			{
+				if (URigVMPin* VariablePin = VariableNode->FindPin(URigVMVariableNode::VariableName))
+				{
+					if (VariablePin->GetDefaultValue() == VarNameStr)
+					{
+						RemoveNode(Node, bSetupUndoRedo, true);
+						continue;
+					}
+				}
+			}
+
+			TArray<URigVMPin*> AllPins = Node->GetAllPinsRecursively();
+			for (URigVMPin* Pin : AllPins)
+			{
+				if (Pin->GetBoundVariableName() == InVariableName.ToString())
+				{
+					BindPinToVariable(Pin, FString(), bSetupUndoRedo);
+				}
+			}
+		}
+
 		if (!bSuspendNotifications)
 		{
 			Graph->MarkPackageDirty();
@@ -7223,7 +7249,7 @@ bool URigVMController::RenameLocalVariable(const FName& InVariableName, const FN
 	{
 		Graph->MarkPackageDirty();
 	}
-	
+
 	if (bSetupUndoRedo)
 	{
 		FRigVMInverseAction InverseAction;
@@ -7235,6 +7261,40 @@ bool URigVMController::RenameLocalVariable(const FName& InVariableName, const FN
 	}	
 	
 	LocalVariables[FoundIndex].Name = InNewVariableName;
+
+	TArray<URigVMNode*> RenamedNodes;
+	for (URigVMNode* Node : Graph->Nodes)
+	{
+		if(URigVMVariableNode* VariableNode = Cast<URigVMVariableNode>(Node))
+		{
+			if (VariableNode->GetVariableName() == InVariableName)
+			{
+				VariableNode->FindPin(URigVMVariableNode::VariableName)->DefaultValue = InNewVariableName.ToString();
+				RenamedNodes.Add(Node);
+			}
+		}
+		
+		const TArray<URigVMPin*> AllPins = Node->GetAllPinsRecursively();
+		for (URigVMPin* Pin : AllPins)
+		{
+			if (Pin->GetBoundVariableName() == InVariableName.ToString())
+			{
+				FString OldVariablePath = Pin->GetBoundVariablePath();
+				FString NewVariablePath = OldVariablePath.Replace(*InVariableName.ToString(), *InNewVariableName.ToString());
+				BindPinToVariable(Pin, NewVariablePath, bSetupUndoRedo);
+			}
+		}
+	}
+
+	for (URigVMNode* RenamedNode : RenamedNodes)
+	{
+		Notify(ERigVMGraphNotifType::VariableRenamed, RenamedNode);
+		if (!bSuspendNotifications)
+		{
+			Graph->MarkPackageDirty();
+		}
+	}
+
 	return true;
 }
 
@@ -7291,7 +7351,6 @@ bool URigVMController::SetLocalVariableType(const FName& InVariableName, const F
 		LocalVariables[FoundIndex].DefaultValue = DefaultValue;
 	}
 
-
 	// Change pin types on variable nodes
 	TArray<URigVMNode*> Nodes = Graph->GetNodes();
 	for (URigVMNode* Node : Nodes)
@@ -7307,6 +7366,15 @@ bool URigVMController::SetLocalVariableType(const FName& InVariableName, const F
 				}
 			}
 		}
+
+		const TArray<URigVMPin*> AllPins = Node->GetAllPinsRecursively();
+		for (URigVMPin* Pin : AllPins)
+		{
+			if (Pin->GetBoundVariableName() == InVariableName.ToString())
+			{
+				BindPinToVariable(Pin, FString(), bSetupUndoRedo);
+			}
+		}
 	}
 	
 	return true;
@@ -9434,7 +9502,7 @@ void URigVMController::SetupDefaultUnitNodeDelegates(TDelegate<FName(FRigVMExter
 	UnitNodeCreatedContext.GetAllExternalVariablesDelegate().BindLambda([WeakThis]() -> TArray<FRigVMExternalVariable> {
 		if (WeakThis.IsValid())
 		{
-			return WeakThis->GetExternalVariables();
+			return WeakThis->GetAllVariables();
 		}
 		return TArray<FRigVMExternalVariable>();
 	});
@@ -10029,9 +10097,9 @@ void URigVMController::DestroyObject(UObject* InObjectToDestroy)
 	InObjectToDestroy->RemoveFromRoot();
 }
 
-FRigVMExternalVariable URigVMController::GetExternalVariableByName(const FName& InExternalVariableName)
+FRigVMExternalVariable URigVMController::GetVariableByName(const FName& InExternalVariableName)
 {
-	TArray<FRigVMExternalVariable> ExternalVariables = GetExternalVariables();
+	TArray<FRigVMExternalVariable> ExternalVariables = GetAllVariables();
 	for (const FRigVMExternalVariable& ExternalVariable : ExternalVariables)
 	{
 		if (ExternalVariable.Name == InExternalVariableName)
@@ -10039,16 +10107,39 @@ FRigVMExternalVariable URigVMController::GetExternalVariableByName(const FName&
 			return ExternalVariable;
 		}
 	}
+
+	if(URigVMGraph* Graph = GetGraph())
+	{
+		for (FRigVMGraphVariableDescription LocalVariable : Graph->GetLocalVariables())
+		{
+			if (LocalVariable.Name == InExternalVariableName)
+			{
+				return LocalVariable.ToExternalVariable();
+			}
+		}
+	}
+	
 	return FRigVMExternalVariable();
 }
 
-TArray<FRigVMExternalVariable> URigVMController::GetExternalVariables()
+TArray<FRigVMExternalVariable> URigVMController::GetAllVariables()
 {
+	TArray<FRigVMExternalVariable> ExternalVariables;
+	
 	if (GetExternalVariablesDelegate.IsBound())
 	{
-		return GetExternalVariablesDelegate.Execute(GetGraph());
+		ExternalVariables.Append(GetExternalVariablesDelegate.Execute(GetGraph()));
+	}
+
+	if(URigVMGraph* Graph = GetGraph())
+	{
+		for (FRigVMGraphVariableDescription LocalVariable : Graph->GetLocalVariables())
+		{
+			ExternalVariables.Add(LocalVariable.ToExternalVariable());
+		}
 	}
-	return TArray<FRigVMExternalVariable>();
+	
+	return ExternalVariables;
 }
 
 const FRigVMByteCode* URigVMController::GetCurrentByteCode() const
