Control Rig: Follow up on CL 16703406 - fixes build break

diff --git a/Engine/Source/Developer/RigVMDeveloper/Private/RigVMCompiler/RigVMCompiler.cpp b/Engine/Source/Developer/RigVMDeveloper/Private/RigVMCompiler/RigVMCompiler.cpp
--- a/Engine/Source/Developer/RigVMDeveloper/Private/RigVMCompiler/RigVMCompiler.cpp
+++ b/Engine/Source/Developer/RigVMDeveloper/Private/RigVMCompiler/RigVMCompiler.cpp
@@ -1688,19 +1688,6 @@ FRigVMOperand URigVMCompiler::FindOrAddRegister(const FRigVMVarExprAST* InVarExp
 			FString VariableName = VariablePath, SegmentPath;
 			VariablePath.Split(TEXT("."), &VariableName, &SegmentPath);
 
-			ExistingOperand = WorkData.PinPathToOperand->Find(FString::Printf(TEXT("Variable::%s"), *VariableName));
-			if (ExistingOperand)
-			{
-				Operand = *ExistingOperand;
-				Operand.RegisterOffset = INDEX_NONE;
-				if (!SegmentPath.IsEmpty())
-				{
-					const FRigVMExternalVariable& ExternalVariable = WorkData.VM->GetExternalVariables()[Operand.GetRegisterIndex()];
-					UScriptStruct* ScriptStruct = CastChecked<UScriptStruct>(ExternalVariable.TypeObject);
-					Operand.RegisterOffset = WorkData.VM->GetWorkMemory().GetOrAddRegisterOffset(Operand.GetRegisterIndex(), ScriptStruct, SegmentPath, 0 /*ArrayIndex */);
-				}
-			}
-
 			ExistingOperand = WorkData.PinPathToOperand->Find(FString::Printf(TEXT("LocalVariable::%s"), *VariableName));
 			if (ExistingOperand)
 			{
@@ -1708,6 +1695,21 @@ FRigVMOperand URigVMCompiler::FindOrAddRegister(const FRigVMVarExprAST* InVarExp
 				Operand.RegisterOffset = INDEX_NONE;
 				check(SegmentPath.IsEmpty());
 			}
+			else
+			{
+				ExistingOperand = WorkData.PinPathToOperand->Find(FString::Printf(TEXT("Variable::%s"), *VariableName));
+				if (ExistingOperand)
+				{
+					Operand = *ExistingOperand;
+					Operand.RegisterOffset = INDEX_NONE;
+					if (!SegmentPath.IsEmpty())
+					{
+						const FRigVMExternalVariable& ExternalVariable = WorkData.VM->GetExternalVariables()[Operand.GetRegisterIndex()];
+						UScriptStruct* ScriptStruct = CastChecked<UScriptStruct>(ExternalVariable.TypeObject);
+						Operand.RegisterOffset = WorkData.VM->GetWorkMemory().GetOrAddRegisterOffset(Operand.GetRegisterIndex(), ScriptStruct, SegmentPath, 0 /*ArrayIndex */);
+					}
+				}
+			}
 		}
 	}
 
diff --git a/Engine/Source/Developer/RigVMDeveloper/Public/RigVMModel/RigVMController.h b/Engine/Source/Developer/RigVMDeveloper/Public/RigVMModel/RigVMController.h
--- a/Engine/Source/Developer/RigVMDeveloper/Public/RigVMModel/RigVMController.h
+++ b/Engine/Source/Developer/RigVMDeveloper/Public/RigVMModel/RigVMController.h
@@ -797,8 +797,8 @@ private:
 
 	void DestroyObject(UObject* InObjectToDestroy);
 
-	FRigVMExternalVariable GetExternalVariableByName(const FName& InExternalVariableName);
-	TArray<FRigVMExternalVariable> GetExternalVariables();
+	FRigVMExternalVariable GetVariableByName(const FName& InExternalVariableName);
+	TArray<FRigVMExternalVariable> GetAllVariables();
 
 	void RefreshFunctionReferences(URigVMLibraryNode* InFunctionDefinition, bool bSetupUndoRedo);
 
