Control Rig: Improve colors for entry, return, collapse and function ref nodes

diff --git a/Engine/Source/Developer/RigVMDeveloper/Private/RigVMModel/Nodes/RigVMCollapseNode.cpp b/Engine/Source/Developer/RigVMDeveloper/Private/RigVMModel/Nodes/RigVMCollapseNode.cpp
--- a/Engine/Source/Developer/RigVMDeveloper/Private/RigVMModel/Nodes/RigVMCollapseNode.cpp
+++ b/Engine/Source/Developer/RigVMDeveloper/Private/RigVMModel/Nodes/RigVMCollapseNode.cpp
@@ -3,6 +3,12 @@
 #include "RigVMModel/Nodes/RigVMCollapseNode.h"
 #include "RigVMModel/RigVMFunctionLibrary.h"
 
+URigVMCollapseNode::URigVMCollapseNode()
+	: URigVMLibraryNode()
+{
+	NodeColor = FLinearColor::White;
+}
+
 FText URigVMCollapseNode::GetToolTipText() const
 {
 	const FString ToolTipString = GetNodeDescription();
diff --git a/Engine/Source/Developer/RigVMDeveloper/Private/RigVMModel/Nodes/RigVMFunctionEntryNode.cpp b/Engine/Source/Developer/RigVMDeveloper/Private/RigVMModel/Nodes/RigVMFunctionEntryNode.cpp
--- a/Engine/Source/Developer/RigVMDeveloper/Private/RigVMModel/Nodes/RigVMFunctionEntryNode.cpp
+++ b/Engine/Source/Developer/RigVMDeveloper/Private/RigVMModel/Nodes/RigVMFunctionEntryNode.cpp
@@ -1,6 +1,19 @@
 // Copyright Epic Games, Inc. All Rights Reserved.
 
 #include "RigVMModel/Nodes/RigVMFunctionEntryNode.h"
+#include "RigVMModel/RigVMFunctionLibrary.h"
+
+FLinearColor URigVMFunctionEntryNode::GetNodeColor() const
+{
+	if(URigVMGraph* RootGraph = GetRootGraph())
+	{
+		if(RootGraph->IsA<URigVMFunctionLibrary>())
+		{
+			return FLinearColor(FColor::FromHex("CB00FFFF"));
+		}
+	}
+	return FLinearColor(FColor::FromHex("005DFFFF"));
+}
 
 bool URigVMFunctionEntryNode::IsDefinedAsVarying() const
 { 
diff --git a/Engine/Source/Developer/RigVMDeveloper/Private/RigVMModel/Nodes/RigVMFunctionReturnNode.cpp b/Engine/Source/Developer/RigVMDeveloper/Private/RigVMModel/Nodes/RigVMFunctionReturnNode.cpp
--- a/Engine/Source/Developer/RigVMDeveloper/Private/RigVMModel/Nodes/RigVMFunctionReturnNode.cpp
+++ b/Engine/Source/Developer/RigVMDeveloper/Private/RigVMModel/Nodes/RigVMFunctionReturnNode.cpp
@@ -1,6 +1,19 @@
 // Copyright Epic Games, Inc. All Rights Reserved.
 
 #include "RigVMModel/Nodes/RigVMFunctionReturnNode.h"
+#include "RigVMModel/RigVMFunctionLibrary.h"
+
+FLinearColor URigVMFunctionReturnNode::GetNodeColor() const
+{
+	if(URigVMGraph* RootGraph = GetRootGraph())
+	{
+		if(RootGraph->IsA<URigVMFunctionLibrary>())
+		{
+			return FLinearColor(FColor::FromHex("CB00FFFF"));
+		}
+	}
+	return FLinearColor(FColor::FromHex("005DFFFF"));
+}
 
 bool URigVMFunctionReturnNode::IsDefinedAsVarying() const
 { 
diff --git a/Engine/Source/Developer/RigVMDeveloper/Private/RigVMModel/RigVMController.cpp b/Engine/Source/Developer/RigVMDeveloper/Private/RigVMModel/RigVMController.cpp
--- a/Engine/Source/Developer/RigVMDeveloper/Private/RigVMModel/RigVMController.cpp
+++ b/Engine/Source/Developer/RigVMDeveloper/Private/RigVMModel/RigVMController.cpp
@@ -3823,6 +3823,7 @@ URigVMFunctionReferenceNode* URigVMController::PromoteCollapseNodeToFunctionRefe
 
 		if (FunctionDefinition)
 		{
+			FunctionDefinition->NodeColor = FLinearColor(FColor::FromHex("005DFFFF"));
 			FunctionLibrary->Nodes.Add(FunctionDefinition);
 			Notify(ERigVMGraphNotifType::NodeAdded, FunctionDefinition);
 		}
@@ -3905,6 +3906,7 @@ URigVMCollapseNode* URigVMController::PromoteFunctionReferenceNodeToCollapseNode
 	URigVMCollapseNode* CollapseNode = DuplicateObject<URigVMCollapseNode>(FunctionDefinition, Graph, *NodeName);
 	if(CollapseNode)
 	{
+		CollapseNode->NodeColor = FLinearColor::White;
 		CollapseNode->Position = NodePosition;
 		Graph->Nodes.Add(CollapseNode);
 		Notify(ERigVMGraphNotifType::NodeAdded, CollapseNode);
diff --git a/Engine/Source/Developer/RigVMDeveloper/Public/RigVMModel/Nodes/RigVMCollapseNode.h b/Engine/Source/Developer/RigVMDeveloper/Public/RigVMModel/Nodes/RigVMCollapseNode.h
--- a/Engine/Source/Developer/RigVMDeveloper/Public/RigVMModel/Nodes/RigVMCollapseNode.h
+++ b/Engine/Source/Developer/RigVMDeveloper/Public/RigVMModel/Nodes/RigVMCollapseNode.h
@@ -17,6 +17,8 @@ class RIGVMDEVELOPER_API URigVMCollapseNode : public URigVMLibraryNode
 
 public:
 
+	URigVMCollapseNode();
+
 	// RigVM node interface
 	virtual FText GetToolTipText() const override;
 
diff --git a/Engine/Source/Developer/RigVMDeveloper/Public/RigVMModel/Nodes/RigVMFunctionEntryNode.h b/Engine/Source/Developer/RigVMDeveloper/Public/RigVMModel/Nodes/RigVMFunctionEntryNode.h
--- a/Engine/Source/Developer/RigVMDeveloper/Public/RigVMModel/Nodes/RigVMFunctionEntryNode.h
+++ b/Engine/Source/Developer/RigVMDeveloper/Public/RigVMModel/Nodes/RigVMFunctionEntryNode.h
@@ -17,6 +17,7 @@ class RIGVMDEVELOPER_API URigVMFunctionEntryNode : public URigVMNode
 public:
 
 	// Override node functions
+	virtual FLinearColor GetNodeColor() const override;
 	virtual bool IsDefinedAsVarying() const override;
 
 private:
diff --git a/Engine/Source/Developer/RigVMDeveloper/Public/RigVMModel/Nodes/RigVMFunctionReturnNode.h b/Engine/Source/Developer/RigVMDeveloper/Public/RigVMModel/Nodes/RigVMFunctionReturnNode.h
--- a/Engine/Source/Developer/RigVMDeveloper/Public/RigVMModel/Nodes/RigVMFunctionReturnNode.h
+++ b/Engine/Source/Developer/RigVMDeveloper/Public/RigVMModel/Nodes/RigVMFunctionReturnNode.h
@@ -17,6 +17,7 @@ class RIGVMDEVELOPER_API URigVMFunctionReturnNode : public URigVMNode
 public:
 
 	// Override node functions
+	virtual FLinearColor GetNodeColor() const override;
 	virtual bool IsDefinedAsVarying() const override;
 
 private:
