Control Rig: Fixed python function references test not passing

diff --git a/Engine/Source/Developer/RigVMDeveloper/Public/RigVMModel/RigVMController.h b/Engine/Source/Developer/RigVMDeveloper/Public/RigVMModel/RigVMController.h
--- a/Engine/Source/Developer/RigVMDeveloper/Public/RigVMModel/RigVMController.h
+++ b/Engine/Source/Developer/RigVMDeveloper/Public/RigVMModel/RigVMController.h
@@ -713,6 +713,7 @@ public:
 	/**
 	 * Helper function to disable a series of checks that can be ignored during a unit test
 	 */
+	UFUNCTION(BlueprintCallable, Category = RigVMController)
 	void SetIsRunningUnitTest(bool bIsRunning) { bIsRunningUnitTest = bIsRunning; }
 
 private:
