ControlRig: Improve bulk edit return information / reaction to undo / redo

diff --git a/Engine/Plugins/Animation/ControlRig/Source/ControlRigDeveloper/Private/ControlRigBlueprint.cpp b/Engine/Plugins/Animation/ControlRig/Source/ControlRigDeveloper/Private/ControlRigBlueprint.cpp
--- a/Engine/Plugins/Animation/ControlRig/Source/ControlRigDeveloper/Private/ControlRigBlueprint.cpp
+++ b/Engine/Plugins/Animation/ControlRig/Source/ControlRigDeveloper/Private/ControlRigBlueprint.cpp
@@ -1243,7 +1243,7 @@ URigVMController* UControlRigBlueprint::GetOrCreateController(URigVMGraph* InGra
 	));
 
 	TWeakObjectPtr<URigVMController> WeakController = Controller;
-	Controller->RequestBulkEditDialogDelegate.BindLambda([WeakThis, WeakController](URigVMLibraryNode* InFunction, ERigVMControllerBulkEditType InEditType)
+	Controller->RequestBulkEditDialogDelegate.BindLambda([WeakThis, WeakController](URigVMLibraryNode* InFunction, ERigVMControllerBulkEditType InEditType) -> FRigVMController_BulkEditResult 
 	{
 		if(WeakThis.IsValid() && WeakController.IsValid())
 		{
@@ -1254,7 +1254,7 @@ URigVMController* UControlRigBlueprint::GetOrCreateController(URigVMGraph* InGra
 				return StrongThis->OnRequestBulkEditDialog().Execute(StrongThis, StrongController, InFunction, InEditType);
 			}
 		}
-		return true;
+		return FRigVMController_BulkEditResult();
 	});
 	
 #endif
diff --git a/Engine/Plugins/Animation/ControlRig/Source/ControlRigDeveloper/Public/ControlRigBlueprint.h b/Engine/Plugins/Animation/ControlRig/Source/ControlRigDeveloper/Public/ControlRigBlueprint.h
--- a/Engine/Plugins/Animation/ControlRig/Source/ControlRigDeveloper/Public/ControlRigBlueprint.h
+++ b/Engine/Plugins/Animation/ControlRig/Source/ControlRigDeveloper/Public/ControlRigBlueprint.h
@@ -40,7 +40,7 @@ DECLARE_EVENT_OneParam(UControlRigBlueprint, FOnGraphImported, UEdGraph*);
 DECLARE_EVENT_OneParam(UControlRigBlueprint, FOnPostEditChangeChainProperty, FPropertyChangedChainEvent&);
 DECLARE_EVENT_ThreeParams(UControlRigBlueprint, FOnLocalizeFunctionDialogRequested, URigVMLibraryNode*, UControlRigBlueprint*, bool);
 DECLARE_EVENT_ThreeParams(UControlRigBlueprint, FOnReportCompilerMessage, EMessageSeverity::Type, UObject*, const FString&);
-DECLARE_DELEGATE_RetVal_FourParams(bool, FControlRigOnBulkEditDialogRequestedDelegate, UControlRigBlueprint*, URigVMController*, URigVMLibraryNode*, ERigVMControllerBulkEditType);
+DECLARE_DELEGATE_RetVal_FourParams(FRigVMController_BulkEditResult, FControlRigOnBulkEditDialogRequestedDelegate, UControlRigBlueprint*, URigVMController*, URigVMLibraryNode*, ERigVMControllerBulkEditType);
 DECLARE_EVENT(UControlRigBlueprint, FOnBreakpointAdded);
 
 USTRUCT()
diff --git a/Engine/Plugins/Animation/ControlRig/Source/ControlRigEditor/Private/Editor/ControlRigEditor.cpp b/Engine/Plugins/Animation/ControlRig/Source/ControlRigEditor/Private/Editor/ControlRigEditor.cpp
--- a/Engine/Plugins/Animation/ControlRig/Source/ControlRigEditor/Private/Editor/ControlRigEditor.cpp
+++ b/Engine/Plugins/Animation/ControlRig/Source/ControlRigEditor/Private/Editor/ControlRigEditor.cpp
@@ -4255,13 +4255,13 @@ void FControlRigEditor::OnRequestLocalizeFunctionDialog(URigVMLibraryNode* InFun
 	}
 }
 
-bool FControlRigEditor::OnRequestBulkEditDialog(UControlRigBlueprint* InBlueprint, URigVMController* InController,
+FRigVMController_BulkEditResult FControlRigEditor::OnRequestBulkEditDialog(UControlRigBlueprint* InBlueprint, URigVMController* InController,
 	URigVMLibraryNode* InFunction, ERigVMControllerBulkEditType InEditType)
 {
 	const TArray<FAssetData> FirstLevelReferenceAssets = InController->GetAffectedAssets(InEditType, false, true);
 	if(FirstLevelReferenceAssets.Num() == 0)
 	{
-		return true;
+		return FRigVMController_BulkEditResult();
 	}
 	
 	TSharedRef<SControlRigFunctionBulkEditDialog> BulkEditDialog = SNew(SControlRigFunctionBulkEditDialog)
@@ -4269,8 +4269,11 @@ bool FControlRigEditor::OnRequestBulkEditDialog(UControlRigBlueprint* InBlueprin
 	.Controller(InController)
 	.Function(InFunction)
 	.EditType(InEditType);
-	
-	return BulkEditDialog->ShowModal() != EAppReturnType::Cancel;
+
+	FRigVMController_BulkEditResult Result;
+	Result.bCanceled = BulkEditDialog->ShowModal() == EAppReturnType::Cancel; 
+	Result.bSetupUndoRedo = false;
+	return Result;
 }
 
 void FControlRigEditor::UpdateDefaultValueForVariable(FBPVariableDescription& InVariable, bool bUseCDO)
diff --git a/Engine/Plugins/Animation/ControlRig/Source/ControlRigEditor/Private/Editor/ControlRigEditor.h b/Engine/Plugins/Animation/ControlRig/Source/ControlRigEditor/Private/Editor/ControlRigEditor.h
--- a/Engine/Plugins/Animation/ControlRig/Source/ControlRigEditor/Private/Editor/ControlRigEditor.h
+++ b/Engine/Plugins/Animation/ControlRig/Source/ControlRigEditor/Private/Editor/ControlRigEditor.h
@@ -400,7 +400,7 @@ protected:
 	virtual void OnFinishedChangingProperties(const FPropertyChangedEvent& PropertyChangedEvent) override;
 	void OnBlueprintPropertyChainEvent(FPropertyChangedChainEvent& PropertyChangedChainEvent);
 	void OnRequestLocalizeFunctionDialog(URigVMLibraryNode* InFunction, UControlRigBlueprint* InTargetBlueprint, bool bForce);
-	bool OnRequestBulkEditDialog(UControlRigBlueprint* InBlueprint, URigVMController* InController, URigVMLibraryNode* InFunction, ERigVMControllerBulkEditType InEditType);
+	FRigVMController_BulkEditResult OnRequestBulkEditDialog(UControlRigBlueprint* InBlueprint, URigVMController* InController, URigVMLibraryNode* InFunction, ERigVMControllerBulkEditType InEditType);
 	void UpdateDefaultValueForVariable(FBPVariableDescription& InVariable, bool bUseCDO);
 
 	URigVMGraph* GetFocusedModel() const;
diff --git a/Engine/Source/Developer/RigVMDeveloper/Private/RigVMModel/RigVMController.cpp b/Engine/Source/Developer/RigVMDeveloper/Private/RigVMModel/RigVMController.cpp
--- a/Engine/Source/Developer/RigVMDeveloper/Private/RigVMModel/RigVMController.cpp
+++ b/Engine/Source/Developer/RigVMDeveloper/Private/RigVMModel/RigVMController.cpp
@@ -796,11 +796,12 @@ URigVMVariableNode* URigVMController::AddVariableNode(const FName& InVariableNam
 						// Warn the user the changes are not undoable
 						if(RequestBulkEditDialogDelegate.IsBound())
 						{
-							if(!RequestBulkEditDialogDelegate.Execute(OuterFunction, ERigVMControllerBulkEditType::AddVariable))
+							FRigVMController_BulkEditResult Result = RequestBulkEditDialogDelegate.Execute(OuterFunction, ERigVMControllerBulkEditType::AddVariable);
+							if(Result.bCanceled)
 							{
 								return nullptr;
 							}
-							bSetupUndoRedo = false;
+							bSetupUndoRedo = Result.bSetupUndoRedo;
 						}
 					}
 				}
@@ -4022,11 +4023,12 @@ bool URigVMController::RemoveNode(URigVMNode* InNode, bool bSetupUndoRedo, bool
 						FRigVMControllerGraphGuard Guard(this, OuterFunction->GetContainedGraph(), false);
 						if(RequestBulkEditDialogDelegate.IsBound())
 						{
-							if(!RequestBulkEditDialogDelegate.Execute(OuterFunction, ERigVMControllerBulkEditType::RemoveVariable))
+							FRigVMController_BulkEditResult Result = RequestBulkEditDialogDelegate.Execute(OuterFunction, ERigVMControllerBulkEditType::RemoveVariable);
+							if(Result.bCanceled)
 							{
 								return false;
 							}
-							bSetupUndoRedo = false;
+							bSetupUndoRedo = Result.bSetupUndoRedo;
 						}
 					}
 				}
@@ -6472,11 +6474,12 @@ bool URigVMController::RemoveExposedPin(const FName& InPinName, bool bSetupUndoR
 	{
 		if(RequestBulkEditDialogDelegate.IsBound())
 		{
-			if(!RequestBulkEditDialogDelegate.Execute(LibraryNode, ERigVMControllerBulkEditType::RemoveExposedPin))
+			FRigVMController_BulkEditResult Result = RequestBulkEditDialogDelegate.Execute(LibraryNode, ERigVMControllerBulkEditType::RemoveExposedPin);
+			if(Result.bCanceled)
 			{
 				return false;
 			}
-			bSetupUndoRedo = false;
+			bSetupUndoRedo = Result.bSetupUndoRedo;
 		}
 	}
 
@@ -6551,11 +6554,12 @@ bool URigVMController::RenameExposedPin(const FName& InOldPinName, const FName&
 	{
 		if(RequestBulkEditDialogDelegate.IsBound())
 		{
-			if(!RequestBulkEditDialogDelegate.Execute(LibraryNode, ERigVMControllerBulkEditType::RenameExposedPin))
+			FRigVMController_BulkEditResult Result = RequestBulkEditDialogDelegate.Execute(LibraryNode, ERigVMControllerBulkEditType::RenameExposedPin); 
+			if(Result.bCanceled)
 			{
 				return false;
 			}
-			bSetupUndoRedo = false;
+			bSetupUndoRedo = Result.bSetupUndoRedo;
 		}
 	}
 
@@ -6686,11 +6690,12 @@ bool URigVMController::ChangeExposedPinType(const FName& InPinName, const FStrin
 	{
 		if(RequestBulkEditDialogDelegate.IsBound())
 		{
-			if(!RequestBulkEditDialogDelegate.Execute(LibraryNode, ERigVMControllerBulkEditType::ChangeExposedPinType))
+			const FRigVMController_BulkEditResult Result = RequestBulkEditDialogDelegate.Execute(LibraryNode, ERigVMControllerBulkEditType::ChangeExposedPinType); 
+			if(Result.bCanceled)
 			{
 				return false;
 			}
-			bSetupUndoRedo = false;
+			bSetupUndoRedo = Result.bSetupUndoRedo;
 		}
 	}
 
diff --git a/Engine/Source/Developer/RigVMDeveloper/Public/RigVMModel/RigVMController.h b/Engine/Source/Developer/RigVMDeveloper/Public/RigVMModel/RigVMController.h
--- a/Engine/Source/Developer/RigVMDeveloper/Public/RigVMModel/RigVMController.h
+++ b/Engine/Source/Developer/RigVMDeveloper/Public/RigVMModel/RigVMController.h
@@ -44,13 +44,24 @@ enum class ERigVMControllerBulkEditProgress : uint8
     Max UMETA(Hidden),
 };
 
+struct FRigVMController_BulkEditResult
+{
+	bool bCanceled;
+	bool bSetupUndoRedo;
+
+	FRigVMController_BulkEditResult()
+		: bCanceled(false)
+		, bSetupUndoRedo(true)
+	{}
+};
+
 DECLARE_DELEGATE_RetVal_OneParam(bool, FRigVMController_ShouldStructUnfoldDelegate, const UStruct*)
 DECLARE_DELEGATE_RetVal_OneParam(TArray<FRigVMExternalVariable>, FRigVMController_GetExternalVariablesDelegate, URigVMGraph*)
 DECLARE_DELEGATE_RetVal(const FRigVMByteCode*, FRigVMController_GetByteCodeDelegate)
 DECLARE_DELEGATE_RetVal_OneParam(bool, FRigVMController_IsFunctionAvailableDelegate, URigVMLibraryNode*)
 DECLARE_DELEGATE_RetVal_OneParam(bool, FRigVMController_RequestLocalizeFunctionDelegate, URigVMLibraryNode*)
 DECLARE_DELEGATE_RetVal_TwoParams(bool, FRigVMController_IsDependencyCyclicDelegate, UObject*, UObject*)
-DECLARE_DELEGATE_RetVal_TwoParams(bool, FRigVMController_RequestBulkEditDialogDelegate, URigVMLibraryNode*, ERigVMControllerBulkEditType)
+DECLARE_DELEGATE_RetVal_TwoParams(FRigVMController_BulkEditResult, FRigVMController_RequestBulkEditDialogDelegate, URigVMLibraryNode*, ERigVMControllerBulkEditType)
 DECLARE_DELEGATE_FiveParams(FRigVMController_OnBulkEditProgressDelegate, TSoftObjectPtr<URigVMFunctionReferenceNode>, ERigVMControllerBulkEditType, ERigVMControllerBulkEditProgress, int32, int32)
 
 /**
