Control Rig: Fix undo of change variable type crash

diff --git a/Engine/Plugins/Animation/ControlRig/Source/ControlRigDeveloper/Private/ControlRigBlueprint.cpp b/Engine/Plugins/Animation/ControlRig/Source/ControlRigDeveloper/Private/ControlRigBlueprint.cpp
--- a/Engine/Plugins/Animation/ControlRig/Source/ControlRigDeveloper/Private/ControlRigBlueprint.cpp
+++ b/Engine/Plugins/Animation/ControlRig/Source/ControlRigDeveloper/Private/ControlRigBlueprint.cpp
@@ -3430,14 +3430,20 @@ void UControlRigBlueprint::OnVariableTypeChanged(const FName& InVarName, FEdGrap
 	{
 		if (URigVMController* Controller = GetOrCreateController(Graph))
 		{
+#if WITH_EDITOR
+			const bool bSetupUndoRedo = !GIsTransacting;
+#else
+			const bool bSetupUndoRedo = false;
+#endif
+
 			FRigVMExternalVariable NewVariable = UControlRig::GetExternalVariableFromPinType(InVarName, InNewPinType);
 			if (NewVariable.IsValid(true)) // allow nullptr
 			{
-				Controller->OnExternalVariableTypeChanged(InVarName, NewVariable.TypeName.ToString(), NewVariable.TypeObject, true);
+				Controller->OnExternalVariableTypeChanged(InVarName, NewVariable.TypeName.ToString(), NewVariable.TypeObject, bSetupUndoRedo);
 			}
 			else
 			{
-				Controller->OnExternalVariableRemoved(InVarName, true);
+				Controller->OnExternalVariableRemoved(InVarName, bSetupUndoRedo);
 			}
 		}
 	}
diff --git a/Engine/Plugins/Animation/ControlRig/Source/ControlRigDeveloper/Private/Graph/ControlRigGraph.cpp b/Engine/Plugins/Animation/ControlRig/Source/ControlRigDeveloper/Private/Graph/ControlRigGraph.cpp
--- a/Engine/Plugins/Animation/ControlRig/Source/ControlRigDeveloper/Private/Graph/ControlRigGraph.cpp
+++ b/Engine/Plugins/Animation/ControlRig/Source/ControlRigDeveloper/Private/Graph/ControlRigGraph.cpp
@@ -523,29 +523,8 @@ void UControlRigGraph::HandleModifiedEvent(ERigVMGraphNotifType InNotifType, URi
 								SourceRigPin->BreakLinkTo(TargetRigPin);
 							}
 
-							for (int32 LinkedPinIndex = 0; LinkedPinIndex < SourceRigPin->LinkedTo.Num();)
-							{
-								if (SourceRigPin->LinkedTo[LinkedPinIndex])
-								{
-									++LinkedPinIndex;
-								}
-								else
-								{
-									SourceRigPin->LinkedTo.RemoveAtSwap(LinkedPinIndex);
-								}
-							}
-
-							for (int32 LinkedPinIndex = 0; LinkedPinIndex < TargetRigPin->LinkedTo.Num();)
-							{
-								if (TargetRigPin->LinkedTo[LinkedPinIndex])
-								{
-									++LinkedPinIndex;
-								}
-								else
-								{
-									TargetRigPin->LinkedTo.RemoveAtSwap(LinkedPinIndex);
-								}
-							}
+							SourceRigPin->LinkedTo.Remove(nullptr);
+							TargetRigPin->LinkedTo.Remove(nullptr);
 						}
 					}
 				}
diff --git a/Engine/Plugins/Animation/ControlRig/Source/ControlRigDeveloper/Private/Graph/ControlRigGraphNode.cpp b/Engine/Plugins/Animation/ControlRig/Source/ControlRigDeveloper/Private/Graph/ControlRigGraphNode.cpp
--- a/Engine/Plugins/Animation/ControlRig/Source/ControlRigDeveloper/Private/Graph/ControlRigGraphNode.cpp
+++ b/Engine/Plugins/Animation/ControlRig/Source/ControlRigDeveloper/Private/Graph/ControlRigGraphNode.cpp
@@ -251,6 +251,9 @@ void UControlRigGraphNode::RewireOldPinsToNewPins(TArray<UEdGraphPin*>& InOldPin
 		{
 			if(OldPin->PinName == NewPin->PinName && OldPin->PinType == NewPin->PinType && OldPin->Direction == NewPin->Direction)
 			{
+				// make sure to remove invalid entries from the linked to list
+				OldPin->LinkedTo.Remove(nullptr);
+				
 				NewPin->MovePersistentDataFromOldPin(*OldPin);
 				break;
 			}
diff --git a/Engine/Plugins/Animation/ControlRig/Source/ControlRigDeveloper/Public/Graph/ControlRigGraphNode.h b/Engine/Plugins/Animation/ControlRig/Source/ControlRigDeveloper/Public/Graph/ControlRigGraphNode.h
--- a/Engine/Plugins/Animation/ControlRig/Source/ControlRigDeveloper/Public/Graph/ControlRigGraphNode.h
+++ b/Engine/Plugins/Animation/ControlRig/Source/ControlRigDeveloper/Public/Graph/ControlRigGraphNode.h
@@ -79,6 +79,11 @@ public:
 
 	UControlRigGraphNode();
 
+	// UObject Interface.
+#if WITH_EDITOR
+	virtual bool Modify( bool bAlwaysMarkDirty=true ) { return false; }
+#endif
+	
 	// UEdGraphNode Interface.
 	virtual FText GetNodeTitle(ENodeTitleType::Type TitleType) const override;
 	virtual FLinearColor GetNodeTitleColor() const override;
diff --git a/Engine/Source/Developer/RigVMDeveloper/Private/RigVMModel/RigVMController.cpp b/Engine/Source/Developer/RigVMDeveloper/Private/RigVMModel/RigVMController.cpp
--- a/Engine/Source/Developer/RigVMDeveloper/Private/RigVMModel/RigVMController.cpp
+++ b/Engine/Source/Developer/RigVMDeveloper/Private/RigVMModel/RigVMController.cpp
@@ -1108,7 +1108,7 @@ void URigVMController::OnExternalVariableRenamed(const FName& InOldVarName, cons
 			{
 				if (VariablePin->GetDefaultValue() == VarNameStr)
 				{
-					RefreshVariableNode(Node->GetFName(), InNewVarName, FString(), nullptr, bSetupUndoRedo);
+					RefreshVariableNode(Node->GetFName(), InNewVarName, FString(), nullptr, bSetupUndoRedo, false);
 					continue;
 				}
 			}
@@ -1198,7 +1198,7 @@ void URigVMController::OnExternalVariableTypeChanged(const FName& InVarName, con
 			{
 				if (VariablePin->GetDefaultValue() == VarNameStr)
 				{
-					RefreshVariableNode(Node->GetFName(), InVarName, InCPPType, InCPPTypeObject, bSetupUndoRedo);
+					RefreshVariableNode(Node->GetFName(), InVarName, InCPPType, InCPPTypeObject, bSetupUndoRedo, false);
 					continue;
 				}
 			}
@@ -7368,7 +7368,7 @@ bool URigVMController::SetLocalVariableType(const FName& InVariableName, const F
 			{
 				if (VariablePin->GetDefaultValue() == InVariableName.ToString())
 				{
-					RefreshVariableNode(Node->GetFName(), InVariableName, InCPPType, InCPPTypeObject, bSetupUndoRedo);
+					RefreshVariableNode(Node->GetFName(), InVariableName, InCPPType, InCPPTypeObject, bSetupUndoRedo, false);
 					continue;
 				}
 			}
