Control Rig: Simplify python creation of select node

diff --git a/Engine/Source/Developer/RigVMDeveloper/Private/RigVMModel/RigVMController.cpp b/Engine/Source/Developer/RigVMDeveloper/Private/RigVMModel/RigVMController.cpp
--- a/Engine/Source/Developer/RigVMDeveloper/Private/RigVMModel/RigVMController.cpp
+++ b/Engine/Source/Developer/RigVMDeveloper/Private/RigVMModel/RigVMController.cpp
@@ -7971,6 +7971,17 @@ URigVMSelectNode* URigVMController::AddSelectNode(const FString& InCPPType, cons
 	return Node;
 }
 
+URigVMSelectNode* URigVMController::AddSelectNodeFromStruct(UScriptStruct* InScriptStruct, const FVector2D& InPosition,
+	const FString& InNodeName, bool bSetupUndoRedo)
+{
+	if (!InScriptStruct)
+	{
+		return nullptr;
+	}
+
+	return AddSelectNode(InScriptStruct->GetStructCPPName(), FName(InScriptStruct->GetPathName()), InPosition, InNodeName, bSetupUndoRedo);
+}
+
 URigVMPrototypeNode* URigVMController::AddPrototypeNode(const FName& InNotation, const FVector2D& InPosition, const FString& InNodeName, bool bSetupUndoRedo)
 {
 	if (!IsValidGraph())
diff --git a/Engine/Source/Developer/RigVMDeveloper/Public/RigVMModel/RigVMController.h b/Engine/Source/Developer/RigVMDeveloper/Public/RigVMModel/RigVMController.h
--- a/Engine/Source/Developer/RigVMDeveloper/Public/RigVMModel/RigVMController.h
+++ b/Engine/Source/Developer/RigVMDeveloper/Public/RigVMModel/RigVMController.h
@@ -249,6 +249,9 @@ public:
 	// Select nodes can be used to pick between multiple values based on an index.
 	UFUNCTION(BlueprintCallable, Category = RigVMController)
 	URigVMSelectNode* AddSelectNode(const FString& InCPPType, const FName& InCPPTypeObjectPath, const FVector2D& InPosition = FVector2D::ZeroVector, const FString& InNodeName = TEXT(""), bool bSetupUndoRedo = true);
+	UFUNCTION(BlueprintCallable, Category = RigVMController)
+	URigVMSelectNode* AddSelectNodeFromStruct(UScriptStruct* InScriptStruct, const FVector2D& InPosition = FVector2D::ZeroVector, const FString& InNodeName = TEXT(""), bool bSetupUndoRedo = true);
+
 
 	// Adds a prototype node to the graph.
 	UFUNCTION(BlueprintCallable, Category = RigVMController)
