Control Rig: Python Command Log

diff --git a/Engine/Plugins/Animation/ControlRig/Source/ControlRig/Private/Rigs/RigHierarchyController.cpp b/Engine/Plugins/Animation/ControlRig/Source/ControlRig/Private/Rigs/RigHierarchyController.cpp
--- a/Engine/Plugins/Animation/ControlRig/Source/ControlRig/Private/Rigs/RigHierarchyController.cpp
+++ b/Engine/Plugins/Animation/ControlRig/Source/ControlRig/Private/Rigs/RigHierarchyController.cpp
@@ -1,7 +1,6 @@
 // Copyright Epic Games, Inc. All Rights Reserved.
 
 #include "Rigs/RigHierarchyController.h"
-#include "RigVMCore/RigVMPythonUtils.h"
 
 #if WITH_EDITOR
 #include "Framework/Notifications/NotificationManager.h"
@@ -9,6 +8,7 @@
 #include "EditorStyleSet.h"
 #include "ScopedTransaction.h"
 #include "Engine/SkeletalMesh.h"
+#include "RigVMPythonUtils.h"
 #endif
 
 #include "ControlRig.h"
@@ -134,7 +134,7 @@ bool URigHierarchyController::SetSelection(TArray<FRigElementKey> InKeys)
 }
 
 FRigElementKey URigHierarchyController::AddBone(FName InName, FRigElementKey InParent, FTransform InTransform,
-                                                bool bTransformInGlobal, ERigBoneType InBoneType, bool bSetupUndo)
+                                                bool bTransformInGlobal, ERigBoneType InBoneType, bool bSetupUndo, bool bPrintPythonCommand)
 {
 	if(!IsValid())
 	{
@@ -169,13 +169,22 @@ FRigElementKey URigHierarchyController::AddBone(FName InName, FRigElementKey InP
 
 #if WITH_EDITOR
 	TransactionPtr.Reset();
+
+	if (bPrintPythonCommand)
+	{
+		TArray<FString> Commands = GetAddBonePythonCommands(NewElement);
+		for (const FString& Command : Commands)
+		{
+			RigVMPythonUtils::Print(FString::Printf(TEXT("%s"), *Command));
+		}
+	}
 #endif
 
 	return NewElement->Key;
 }
 
 FRigElementKey URigHierarchyController::AddNull(FName InName, FRigElementKey InParent, FTransform InTransform,
-	bool bTransformInGlobal, bool bSetupUndo)
+                                                bool bTransformInGlobal, bool bSetupUndo, bool bPrintPythonCommand)
 {
 	if(!IsValid())
 	{
@@ -211,6 +220,15 @@ FRigElementKey URigHierarchyController::AddNull(FName InName, FRigElementKey InP
 
 #if WITH_EDITOR
 	TransactionPtr.Reset();
+
+	if (bPrintPythonCommand)
+	{
+		TArray<FString> Commands = GetAddNullPythonCommands(NewElement);
+		for (const FString& Command : Commands)
+		{
+			RigVMPythonUtils::Print(FString::Printf(TEXT("%s"), *Command));
+		}
+	}
 #endif
 
 	return NewElement->Key;
@@ -223,7 +241,7 @@ FRigElementKey URigHierarchyController::AddControl(
 	FRigControlValue InValue,
 	FTransform InOffsetTransform,
 	FTransform InGizmoTransform,
-	bool bSetupUndo)
+	bool bSetupUndo, bool bPrintPythonCommand)
 {
 	if(!IsValid())
 	{
@@ -260,12 +278,21 @@ FRigElementKey URigHierarchyController::AddControl(
 
 #if WITH_EDITOR
 	TransactionPtr.Reset();
+
+	if (bPrintPythonCommand)
+	{
+		TArray<FString> Commands = GetAddControlPythonCommands(NewElement);
+		for (const FString& Command : Commands)
+		{
+			RigVMPythonUtils::Print(FString::Printf(TEXT("%s"), *Command));
+		}
+	}
 #endif
 
 	return NewElement->Key;
 }
 
-FRigElementKey URigHierarchyController::AddCurve(FName InName, float InValue, bool bSetupUndo)
+FRigElementKey URigHierarchyController::AddCurve(FName InName, float InValue, bool bSetupUndo, bool bPrintPythonCommand)
 {
 	if(!IsValid())
 	{
@@ -289,13 +316,22 @@ FRigElementKey URigHierarchyController::AddCurve(FName InName, float InValue, bo
 
 #if WITH_EDITOR
 	TransactionPtr.Reset();
+
+	if (bPrintPythonCommand)
+	{
+		TArray<FString> Commands = GetAddCurvePythonCommands(NewElement);
+		for (const FString& Command : Commands)
+		{
+			RigVMPythonUtils::Print(FString::Printf(TEXT("%s"), *Command));
+		}
+	}
 #endif
 
 	return NewElement->Key;
 }
 
 FRigElementKey URigHierarchyController::AddRigidBody(FName InName, FRigElementKey InParent,
-	FRigRigidBodySettings InSettings, FTransform InLocalTransform, bool bSetupUndo)
+                                                     FRigRigidBodySettings InSettings, FTransform InLocalTransform, bool bSetupUndo, bool bPrintPythonCommand)
 {
 	if(!IsValid())
 	{
@@ -322,6 +358,15 @@ FRigElementKey URigHierarchyController::AddRigidBody(FName InName, FRigElementKe
 
 #if WITH_EDITOR
 	TransactionPtr.Reset();
+
+	if (bPrintPythonCommand)
+	{
+		TArray<FString> Commands = GetAddRigidBodyPythonCommands(NewElement);
+		for (const FString& Command : Commands)
+		{
+			RigVMPythonUtils::Print(FString::Printf(TEXT("%s"), *Command));
+		}
+	}
 #endif
 
 	return NewElement->Key;
@@ -585,6 +630,36 @@ TArray<FRigElementKey> URigHierarchyController::ImportBones(const FReferenceSkel
 	return AddedBones;
 }
 
+TArray<FRigElementKey> URigHierarchyController::ImportBones(USkeleton* InSkeleton, FName InNameSpace,
+                                                            bool bReplaceExistingBones, bool bRemoveObsoleteBones,
+                                                            bool bSelectBones, bool bSetupUndo,
+                                                            bool bPrintPythonCommand)
+{
+	if (InSkeleton == nullptr)
+	{
+		return TArray<FRigElementKey>();
+	}
+
+	const TArray<FRigElementKey> BoneKeys = ImportBones(InSkeleton->GetReferenceSkeleton(), InNameSpace, bReplaceExistingBones, bRemoveObsoleteBones,
+	                   bSelectBones, bSetupUndo);
+
+#if WITH_EDITOR
+	if (!BoneKeys.IsEmpty() && bPrintPythonCommand)
+	{
+		if (bPrintPythonCommand)
+		{
+			RigVMPythonUtils::Print(FString::Printf(TEXT("hierarchy_controller.import_bones_from_asset('%s', '%s', %s, %s, %s)"),
+				*InSkeleton->GetPathName(),
+				*InNameSpace.ToString(),
+				(bReplaceExistingBones) ? TEXT("True") : TEXT("False"),
+				(bRemoveObsoleteBones) ? TEXT("True") : TEXT("False"),
+				(bSelectBones) ? TEXT("True") : TEXT("False")));
+		}
+	}
+#endif
+	return BoneKeys;
+}
+
 #if WITH_EDITOR
 
 TArray<FRigElementKey> URigHierarchyController::ImportBonesFromAsset(FString InAssetPath, FName InNameSpace,
@@ -631,7 +706,7 @@ USkeleton* URigHierarchyController::GetSkeletonFromAssetPath(const FString& InAs
 #endif
 
 TArray<FRigElementKey> URigHierarchyController::ImportCurves(USkeleton* InSkeleton, FName InNameSpace,
-                                                             bool bSelectCurves, bool bSetupUndo)
+                                                             bool bSelectCurves, bool bSetupUndo, bool bPrintPythonCommand)
 {
 	check(InSkeleton);
 
@@ -669,6 +744,19 @@ TArray<FRigElementKey> URigHierarchyController::ImportCurves(USkeleton* InSkelet
 	{
 		SetSelection(Keys);
 	}
+	
+#if WITH_EDITOR
+	if (!Keys.IsEmpty() && bPrintPythonCommand)
+	{
+		if (bPrintPythonCommand)
+		{
+			RigVMPythonUtils::Print(FString::Printf(TEXT("hierarchy_controller.import_curves_from_asset('%s', '%s', %s)"),
+				*InSkeleton->GetPathName(),
+				*InNameSpace.ToString(),
+				(bSelectCurves) ? TEXT("True") : TEXT("False")));
+		}
+	}
+#endif
 
 	return Keys;
 }
@@ -1120,6 +1208,7 @@ TArray<FRigElementKey> URigHierarchyController::ImportFromHierarchyContainer(con
 	return AddedKeys;
 }
 
+#if WITH_EDITOR
 TArray<FString> URigHierarchyController::GeneratePythonCommands()
 {
 	TArray<FString> Commands;
@@ -1381,6 +1470,7 @@ TArray<FString> URigHierarchyController::GetSetControlGizmoTransformPythonComman
 		*RigVMPythonUtils::TransformToPythonString(InTransform),
 		bInitial ? TEXT("True") : TEXT("False"))};
 }
+#endif
 
 void URigHierarchyController::Notify(ERigHierarchyNotification InNotifType, const FRigBaseElement* InElement)
 {
@@ -1428,7 +1518,7 @@ int32 URigHierarchyController::AddElement(FRigBaseElement* InElementToAdd, FRigB
 	return InElementToAdd->Index;
 }
 
-bool URigHierarchyController::RemoveElement(FRigElementKey InElement, bool bSetupUndo)
+bool URigHierarchyController::RemoveElement(FRigElementKey InElement, bool bSetupUndo, bool bPrintPythonCommand)
 {
 	if(!IsValid())
 	{
@@ -1459,6 +1549,12 @@ bool URigHierarchyController::RemoveElement(FRigElementKey InElement, bool bSetu
 		TransactionPtr->Cancel();
 	}
 	TransactionPtr.Reset();
+
+	if (bRemoved && bPrintPythonCommand)
+	{
+		RigVMPythonUtils::Print(FString::Printf(TEXT("hierarchy_controller.remove_element(%s)"),
+			*InElement.ToPythonString()));
+	}
 #endif
 	
 	return bRemoved;
@@ -1568,7 +1664,7 @@ bool URigHierarchyController::RemoveElement(FRigBaseElement* InElement)
 	return NumElementsRemoved == 1;
 }
 
-FRigElementKey URigHierarchyController::RenameElement(FRigElementKey InElement, FName InName, bool bSetupUndo)
+FRigElementKey URigHierarchyController::RenameElement(FRigElementKey InElement, FName InName, bool bSetupUndo, bool bPrintPythonCommand)
 {
 	if(!IsValid())
 	{
@@ -1604,6 +1700,13 @@ FRigElementKey URigHierarchyController::RenameElement(FRigElementKey InElement,
 	{
 		ClearSelection();
 	}
+
+	if (bRenamed && bPrintPythonCommand)
+	{
+		RigVMPythonUtils::Print(FString::Printf(TEXT("hierarchy_controller.rename_element(%s, '%s')"),
+			*InElement.ToPythonString(),
+			*InName.ToString()));
+	}
 #endif
 
 	return bRenamed ? Element->GetKey() : FRigElementKey();
@@ -1814,7 +1917,7 @@ bool URigHierarchyController::AddParent(FRigBaseElement* InChild, FRigBaseElemen
 	return false;
 }
 
-bool URigHierarchyController::RemoveParent(FRigElementKey InChild, FRigElementKey InParent, bool bMaintainGlobalTransform, bool bSetupUndo)
+bool URigHierarchyController::RemoveParent(FRigElementKey InChild, FRigElementKey InParent, bool bMaintainGlobalTransform, bool bSetupUndo, bool bPrintPythonCommand)
 {
 	if(!IsValid())
 	{
@@ -1852,6 +1955,14 @@ bool URigHierarchyController::RemoveParent(FRigElementKey InChild, FRigElementKe
 		TransactionPtr->Cancel();
 	}
 	TransactionPtr.Reset();
+
+	if (bRemoved && bPrintPythonCommand)
+	{
+		RigVMPythonUtils::Print(FString::Printf(TEXT("hierarchy_controller.remove_parent(%s, %s, %s)"),
+			*InChild.ToPythonString(),
+			*InParent.ToPythonString(),
+			(bMaintainGlobalTransform) ? TEXT("True") : TEXT("False")));
+	}
 #endif
 	
 	return bRemoved;
@@ -1984,7 +2095,7 @@ bool URigHierarchyController::RemoveParent(FRigBaseElement* InChild, FRigBaseEle
 	return false;
 }
 
-bool URigHierarchyController::RemoveAllParents(FRigElementKey InChild, bool bMaintainGlobalTransform, bool bSetupUndo)
+bool URigHierarchyController::RemoveAllParents(FRigElementKey InChild, bool bMaintainGlobalTransform, bool bSetupUndo, bool bPrintPythonCommand)
 {
 	if(!IsValid())
 	{
@@ -2015,6 +2126,13 @@ bool URigHierarchyController::RemoveAllParents(FRigElementKey InChild, bool bMai
 		TransactionPtr->Cancel();
 	}
 	TransactionPtr.Reset();
+	
+	if (bRemoved && bPrintPythonCommand)
+	{
+		RigVMPythonUtils::Print(FString::Printf(TEXT("hierarchy_controller.remove_parent(%s, %s)"),
+			*InChild.ToPythonString(),
+			(bMaintainGlobalTransform) ? TEXT("True") : TEXT("False")));
+	}
 #endif
 
 	return bRemoved;
@@ -2044,7 +2162,7 @@ bool URigHierarchyController::RemoveAllParents(FRigBaseElement* InChild, bool bM
 	return false;
 }
 
-bool URigHierarchyController::SetParent(FRigElementKey InChild, FRigElementKey InParent, bool bMaintainGlobalTransform, bool bSetupUndo)
+bool URigHierarchyController::SetParent(FRigElementKey InChild, FRigElementKey InParent, bool bMaintainGlobalTransform, bool bSetupUndo, bool bPrintPythonCommand)
 {
 	if(!IsValid())
 	{
@@ -2082,6 +2200,14 @@ bool URigHierarchyController::SetParent(FRigElementKey InChild, FRigElementKey I
 		TransactionPtr->Cancel();
 	}
 	TransactionPtr.Reset();
+
+	if (bParentSet && bPrintPythonCommand)
+	{
+		RigVMPythonUtils::Print(FString::Printf(TEXT("hierarchy_controller.set_parent(%s, %s, %s)"),
+			*InChild.ToPythonString(),
+			*InParent.ToPythonString(),
+			(bMaintainGlobalTransform) ? TEXT("True") : TEXT("False")));
+	}
 #endif
 
 	return bParentSet;
diff --git a/Engine/Plugins/Animation/ControlRig/Source/ControlRig/Public/Rigs/RigHierarchyController.h b/Engine/Plugins/Animation/ControlRig/Source/ControlRig/Public/Rigs/RigHierarchyController.h
--- a/Engine/Plugins/Animation/ControlRig/Source/ControlRig/Public/Rigs/RigHierarchyController.h
+++ b/Engine/Plugins/Animation/ControlRig/Source/ControlRig/Public/Rigs/RigHierarchyController.h
@@ -78,10 +78,11 @@ public:
 	 * @param bTransformInGlobal Set this to true if the Transform passed is expressed in global space, false for local space.
 	 * @param InBoneType The type of bone to add. This can be used to differentiate between imported bones and user defined bones.
 	 * @param bSetupUndo If set to true the stack will record the change for undo / redo
+	 * @param bPrintPythonCommand If set to true a python command equivalent to this call will be printed out
 	 * @return The key for the newly created bone.
 	 */
 	UFUNCTION(BlueprintCallable, Category = URigHierarchyController)
-	FRigElementKey AddBone(FName InName, FRigElementKey InParent, FTransform InTransform, bool bTransformInGlobal = true, ERigBoneType InBoneType = ERigBoneType::User, bool bSetupUndo = false);
+	FRigElementKey AddBone(FName InName, FRigElementKey InParent, FTransform InTransform, bool bTransformInGlobal = true, ERigBoneType InBoneType = ERigBoneType::User, bool bSetupUndo = false, bool bPrintPythonCommand = false);
 
 	/**
 	 * Adds a null to the hierarchy
@@ -90,10 +91,11 @@ public:
 	 * @param InTransform The transform for the new null - either in local or global null, based on bTransformInGlobal
 	 * @param bTransformInGlobal Set this to true if the Transform passed is expressed in global null, false for local null.
 	 * @param bSetupUndo If set to true the stack will record the change for undo / redo
+	 * @param bPrintPythonCommand If set to true a python command equivalent to this call will be printed out
 	 * @return The key for the newly created null.
 	 */
 	UFUNCTION(BlueprintCallable, Category = URigHierarchyController)
-    FRigElementKey AddNull(FName InName, FRigElementKey InParent, FTransform InTransform, bool bTransformInGlobal = true, bool bSetupUndo = false);
+    FRigElementKey AddNull(FName InName, FRigElementKey InParent, FTransform InTransform, bool bTransformInGlobal = true, bool bSetupUndo = false, bool bPrintPythonCommand = false);
 
 	/**
 	 * Adds a control to the hierarchy
@@ -104,6 +106,7 @@ public:
 	 * @param InOffsetTransform The transform to use for the offset
 	 * @param InGizmoTransform The transform to use for the gizmo
 	 * @param bSetupUndo If set to true the stack will record the change for undo / redo
+	 * @param bPrintPythonCommand If set to true a python command equivalent to this call will be printed out
 	 * @return The key for the newly created control.
 	 */
     FRigElementKey AddControl(
@@ -113,7 +116,8 @@ public:
     	FRigControlValue InValue,
     	FTransform InOffsetTransform = FTransform::Identity,
         FTransform InGizmoTransform = FTransform::Identity,
-        bool bSetupUndo = true
+        bool bSetupUndo = true,
+        bool bPrintPythonCommand = false
         );
 
 	/**
@@ -142,13 +146,15 @@ public:
 	 * @param InName The suggested name of the new curve - will eventually be corrected by the namespace
 	 * @param InValue The value to use for the curve
 	 * @param bSetupUndo If set to true the stack will record the change for undo / redo
+	 * @param bPrintPythonCommand If set to true a python command equivalent to this call will be printed out
 	 * @return The key for the newly created curve.
 	 */
 	UFUNCTION(BlueprintCallable, Category = URigHierarchyController)
     FRigElementKey AddCurve(
         FName InName,
         float InValue = 0.f,
-        bool bSetupUndo = true
+        bool bSetupUndo = true,
+		bool bPrintPythonCommand = false
         );
 
 	/**
@@ -158,6 +164,7 @@ public:
 	* @param InSettings All of the rigidbody's settings
 	* @param InLocalTransform The transform for the new rigidbody - in the space of the provided parent
 	* @param bSetupUndo If set to true the stack will record the change for undo / redo
+	* @param bPrintPythonCommand If set to true a python command equivalent to this call will be printed out
 	* @return The key for the newly created rigidbody.
 	*/
 	UFUNCTION(BlueprintCallable, Category = URigHierarchyController)
@@ -166,7 +173,8 @@ public:
     	FRigElementKey InParent,
         FRigRigidBodySettings InSettings,
     	FTransform InLocalTransform,
-    	bool bSetupUndo = false);
+    	bool bSetupUndo = false,
+		bool bPrintPythonCommand = false);
 
 	/**
 	* Adds an socket to the hierarchy
@@ -225,23 +233,18 @@ public:
 	 * @param bRemoveObsoleteBones If true bones non-existent in the skeleton will be removed from the hierarchy
 	 * @param bSelectBones If true the bones will be selected upon import
 	 * @param bSetupUndo If set to true the stack will record the change for undo / redo
+	 * @param bPrintPythonCommand If set to true a python command equivalent to this call will be printed out
 	 * @return The keys of the imported elements
 	 */
 	UFUNCTION(BlueprintCallable, Category = URigHierarchyController)
-    FORCEINLINE TArray<FRigElementKey> ImportBones(
-        USkeleton* InSkeleton,
-        FName InNameSpace = NAME_None,
-        bool bReplaceExistingBones = true,
-        bool bRemoveObsoleteBones = true,
-        bool bSelectBones = false,
-        bool bSetupUndo = false)
-	{
-		if(InSkeleton == nullptr)
-		{
-			return TArray<FRigElementKey>();
-		}
-		return ImportBones(InSkeleton->GetReferenceSkeleton(), InNameSpace, bReplaceExistingBones, bRemoveObsoleteBones, bSelectBones, bSetupUndo);
-	}
+	TArray<FRigElementKey> ImportBones(
+		USkeleton* InSkeleton,
+		FName InNameSpace = NAME_None,
+		bool bReplaceExistingBones = true,
+		bool bRemoveObsoleteBones = true,
+		bool bSelectBones = false,
+		bool bSetupUndo = false,
+		bool bPrintPythonCommand = false);
 
 #if WITH_EDITOR
 	/**
@@ -277,7 +280,8 @@ public:
 		USkeleton* InSkeleton, 
 		FName InNameSpace = NAME_None,  
 		bool bSelectCurves = false,
-		bool bSetupUndo = false);
+		bool bSetupUndo = false,
+		bool bPrintPythonCommand = false);
 
 #if WITH_EDITOR
 	/**
@@ -336,20 +340,22 @@ public:
 	 * Removes an existing element from the hierarchy
 	 * @param InElement The key of the element to remove
 	 * @param bSetupUndo If set to true the stack will record the change for undo / redo
+	 * @param bPrintPythonCommand If set to true a python command equivalent to this call will be printed out
 	 * @return Returns true if successful.
 	 */
 	UFUNCTION(BlueprintCallable, Category = URigHierarchyController)
-    bool RemoveElement(FRigElementKey InElement, bool bSetupUndo = false);
+	bool RemoveElement(FRigElementKey InElement, bool bSetupUndo = false, bool bPrintPythonCommand = false);
 
 	/**
 	 * Renames an existing element in the hierarchy
 	 * @param InElement The key of the element to rename
 	 * @param InName The new name to set for the element
 	 * @param bSetupUndo If set to true the stack will record the change for undo / redo
+	 * @param bPrintPythonCommand If set to true a python command equivalent to this call will be printed out
 	 * @return Returns the new element key used for the element
 	 */
 	UFUNCTION(BlueprintCallable, Category = URigHierarchyController)
-    FRigElementKey RenameElement(FRigElementKey InElement, FName InName, bool bSetupUndo = false);
+    FRigElementKey RenameElement(FRigElementKey InElement, FName InName, bool bSetupUndo = false, bool bPrintPythonCommand = false);
 
 	/**
 	 * Adds a new parent to an element. For elements that allow only one parent the parent will be replaced (Same as ::SetParent).
@@ -368,20 +374,22 @@ public:
 	 * @param InParent The key of the parent to remove
 	 * @param bMaintainGlobalTransform If set to true the child will stay in the same place spatially, otherwise it will maintain it's local transform (and potential move).
 	 * @param bSetupUndo If set to true the stack will record the change for undo / redo
+	 * @param bPrintPythonCommand If set to true a python command equivalent to this call will be printed out
 	 * @return Returns true if successful.
 	 */
 	UFUNCTION(BlueprintCallable, Category = URigHierarchyController)
-	bool RemoveParent(FRigElementKey InChild, FRigElementKey InParent, bool bMaintainGlobalTransform = true, bool bSetupUndo = false);
+	bool RemoveParent(FRigElementKey InChild, FRigElementKey InParent, bool bMaintainGlobalTransform = true, bool bSetupUndo = false, bool bPrintPythonCommand = false);
 
 	/**
  	 * Removes all parents from an element in the hierarchy.
  	 * @param InChild The key of the element to remove all parents for
  	 * @param bMaintainGlobalTransform If set to true the child will stay in the same place spatially, otherwise it will maintain it's local transform (and potential move).
 	 * @param bSetupUndo If set to true the stack will record the change for undo / redo
+	 * @param bPrintPythonCommand If set to true a python command equivalent to this call will be printed out
 	 * @return Returns true if successful.
 	 */
 	UFUNCTION(BlueprintCallable, Category = URigHierarchyController)
-	bool RemoveAllParents(FRigElementKey InChild, bool bMaintainGlobalTransform = true, bool bSetupUndo = false);
+	bool RemoveAllParents(FRigElementKey InChild, bool bMaintainGlobalTransform = true, bool bSetupUndo = false, bool bPrintPythonCommand = false);
 
 	/**
 	 * Sets a new parent to an element. For elements that allow more than one parent the parent list will be replaced.
@@ -389,10 +397,11 @@ public:
 	 * @param InParent The key of the new parent to set
 	 * @param bMaintainGlobalTransform If set to true the child will stay in the same place spatially, otherwise it will maintain it's local transform (and potential move).
 	 * @param bSetupUndo If set to true the stack will record the change for undo / redo
+	 * @param bPrintPythonCommand If set to true a python command equivalent to this call will be printed out
 	 * @return Returns true if successful.
 	 */
 	UFUNCTION(BlueprintCallable, Category = URigHierarchyController)
-	bool SetParent(FRigElementKey InChild, FRigElementKey InParent, bool bMaintainGlobalTransform = true, bool bSetupUndo = false);
+	bool SetParent(FRigElementKey InChild, FRigElementKey InParent, bool bMaintainGlobalTransform = true, bool bSetupUndo = false, bool bPrintPythonCommand = false);
 
 	/**
 	 * Duplicate the given elements
@@ -466,7 +475,7 @@ public:
 	 */
 	FORCEINLINE bool& GetSuspendNotificationsFlag() { return bSuspendNotifications; }
 
-
+#if WITH_EDITOR
 	UFUNCTION(BlueprintCallable, Category = URigHierarchyController)
 	TArray<FString> GeneratePythonCommands();
 
@@ -487,7 +496,8 @@ public:
 	TArray<FString> GetSetControlOffsetTransformPythonCommands(const FRigControlElement* Control, const FTransform& Offset, bool bInitial = false, bool bAffectChildren = true) const;
 	
 	TArray<FString> GetSetControlGizmoTransformPythonCommands(const FRigControlElement* Control, const FTransform& Transform, bool bInitial = false) const;
-
+#endif
+	
 private:
 
 	UPROPERTY(transient)
diff --git a/Engine/Plugins/Animation/ControlRig/Source/ControlRigDeveloper/Private/ControlRigBlueprint.cpp b/Engine/Plugins/Animation/ControlRig/Source/ControlRigDeveloper/Private/ControlRigBlueprint.cpp
--- a/Engine/Plugins/Animation/ControlRig/Source/ControlRigDeveloper/Private/ControlRigBlueprint.cpp
+++ b/Engine/Plugins/Animation/ControlRig/Source/ControlRigDeveloper/Private/ControlRigBlueprint.cpp
@@ -25,7 +25,7 @@
 #include "Async/TaskGraphInterfaces.h"
 #include "Misc/CoreDelegates.h"
 #include "AssetRegistryModule.h"
-#include "RigVMCore/RigVMPythonUtils.h"
+#include "RigVMPythonUtils.h"
 
 #if WITH_EDITOR
 #include "IControlRigEditorModule.h"
@@ -3404,7 +3404,7 @@ void UControlRigBlueprint::OnVariableRemoved(const FName& InVarName)
 	{
 		if (URigVMController* Controller = GetOrCreateController(Graph))
 		{
-			Controller->OnExternalVariableRemoved(InVarName, true);
+			Controller->OnExternalVariableRemoved(InVarName, true, true);
 		}
 	}
 	BroadcastExternalVariablesChangedEvent();
@@ -3417,7 +3417,7 @@ void UControlRigBlueprint::OnVariableRenamed(const FName& InOldVarName, const FN
 	{
 		if (URigVMController* Controller = GetOrCreateController(Graph))
 		{
-			Controller->OnExternalVariableRenamed(InOldVarName, InNewVarName, true);
+			Controller->OnExternalVariableRenamed(InOldVarName, InNewVarName, true, true);
 		}
 	}
 	BroadcastExternalVariablesChangedEvent();
@@ -3439,11 +3439,11 @@ void UControlRigBlueprint::OnVariableTypeChanged(const FName& InVarName, FEdGrap
 			FRigVMExternalVariable NewVariable = UControlRig::GetExternalVariableFromPinType(InVarName, InNewPinType);
 			if (NewVariable.IsValid(true)) // allow nullptr
 			{
-				Controller->OnExternalVariableTypeChanged(InVarName, NewVariable.TypeName.ToString(), NewVariable.TypeObject, bSetupUndoRedo);
+				Controller->OnExternalVariableTypeChanged(InVarName, NewVariable.TypeName.ToString(), NewVariable.TypeObject, bSetupUndoRedo, true);
 			}
 			else
 			{
-				Controller->OnExternalVariableRemoved(InVarName, bSetupUndoRedo);
+				Controller->OnExternalVariableRemoved(InVarName, bSetupUndoRedo, true);
 			}
 		}
 	}
diff --git a/Engine/Plugins/Animation/ControlRig/Source/ControlRigDeveloper/Private/Graph/ControlRigGraphNode.cpp b/Engine/Plugins/Animation/ControlRig/Source/ControlRigDeveloper/Private/Graph/ControlRigGraphNode.cpp
--- a/Engine/Plugins/Animation/ControlRig/Source/ControlRigDeveloper/Private/Graph/ControlRigGraphNode.cpp
+++ b/Engine/Plugins/Animation/ControlRig/Source/ControlRigDeveloper/Private/Graph/ControlRigGraphNode.cpp
@@ -324,7 +324,7 @@ void UControlRigGraphNode::HandleAddArrayElement(FString InPinPath)
 	if (URigVMController* Controller = GetController())
 	{
 		Controller->OpenUndoBracket(TEXT("Add Array Pin"));
-		FString PinPath = Controller->AddArrayPin(InPinPath);
+		FString PinPath = Controller->AddArrayPin(InPinPath, FString(), true, true);
 		Controller->SetPinExpansion(InPinPath, true);
 		Controller->SetPinExpansion(PinPath, true);
 		Controller->CloseUndoBracket();
@@ -337,7 +337,7 @@ void UControlRigGraphNode::HandleRemoveArrayElement(FString InPinPath)
 
 	if (URigVMController* Controller = GetController())
 	{
-		Controller->RemoveArrayPin(InPinPath);
+		Controller->RemoveArrayPin(InPinPath, true, true);
 	}
 }
 
@@ -352,7 +352,7 @@ void UControlRigGraphNode::HandleInsertArrayElement(FString InPinPath)
 			if (URigVMPin* ArrayPin = ArrayElementPin->GetParentPin())
 				{
 				Controller->OpenUndoBracket(TEXT("Add Array Pin"));
-				FString PinPath = Controller->InsertArrayPin(InPinPath, ArrayElementPin->GetPinIndex() + 1, FString());
+				FString PinPath = Controller->InsertArrayPin(InPinPath, ArrayElementPin->GetPinIndex() + 1, FString(), true, true);
 				Controller->SetPinExpansion(InPinPath, true);
 				Controller->SetPinExpansion(PinPath, true);
 				Controller->CloseUndoBracket();
@@ -774,10 +774,10 @@ void UControlRigGraphNode::DestroyNode()
 
 void UControlRigGraphNode::PinDefaultValueChanged(UEdGraphPin* Pin)
 {
-	CopyPinDefaultsToModel(Pin, true);
+	CopyPinDefaultsToModel(Pin, true, true);
 }
 
-void UControlRigGraphNode::CopyPinDefaultsToModel(UEdGraphPin* Pin, bool bUndo)
+void UControlRigGraphNode::CopyPinDefaultsToModel(UEdGraphPin* Pin, bool bUndo, bool bPrintPythonCommand)
 {
 	DECLARE_SCOPE_HIERARCHICAL_COUNTER_FUNC()
 
@@ -804,7 +804,7 @@ void UControlRigGraphNode::CopyPinDefaultsToModel(UEdGraphPin* Pin, bool bUndo)
 		{
 			if (URigVMController* Controller = GetController())
 			{
-				Controller->SetPinDefaultValue(ModelPin->GetPinPath(), DefaultValue, false, true, false);
+				Controller->SetPinDefaultValue(ModelPin->GetPinPath(), DefaultValue, false, true, false, bPrintPythonCommand);
 			}
 		}
 	}
diff --git a/Engine/Plugins/Animation/ControlRig/Source/ControlRigDeveloper/Private/Graph/ControlRigGraphSchema.cpp b/Engine/Plugins/Animation/ControlRig/Source/ControlRigDeveloper/Private/Graph/ControlRigGraphSchema.cpp
--- a/Engine/Plugins/Animation/ControlRig/Source/ControlRigDeveloper/Private/Graph/ControlRigGraphSchema.cpp
+++ b/Engine/Plugins/Animation/ControlRig/Source/ControlRigDeveloper/Private/Graph/ControlRigGraphSchema.cpp
@@ -59,7 +59,7 @@ void FControlRigGraphSchemaAction_LocalVar::ChangeVariableType(const FEdGraphPin
 		FString NewCPPType;
 		UObject* NewCPPTypeObject = nullptr;
 		FRigVMGraphVariableDescription::CPPTypeFromPinType(NewPinType, NewCPPType, NewCPPTypeObject);
-		Graph->GetController()->SetLocalVariableType(GetVariableName(), NewCPPType, NewCPPTypeObject);
+		Graph->GetController()->SetLocalVariableType(GetVariableName(), NewCPPType, NewCPPTypeObject, true, true);
 	}
 }
 
@@ -69,7 +69,7 @@ void FControlRigGraphSchemaAction_LocalVar::RenameVariable(const FName& NewName)
 	
 	if (UControlRigGraph* Graph = Cast<UControlRigGraph>(GetVariableScope()))
 	{
-		if (Graph->GetController()->RenameLocalVariable(OldName, NewName))
+		if (Graph->GetController()->RenameLocalVariable(OldName, NewName, true, true))
 		{
 			SetVariableInfo(NewName, GetVariableScope(), GetPinType().PinCategory == TEXT("bool"));		
 		}
@@ -104,7 +104,7 @@ void FControlRigGraphSchemaAction_LocalVar::DeleteVariable()
 {
 	if (UControlRigGraph* Graph = Cast<UControlRigGraph>(GetVariableScope()))
 	{
-		Graph->GetController()->RemoveLocalVariable(GetVariableName());
+		Graph->GetController()->RemoveLocalVariable(GetVariableName(), true, true);
 	}
 }
 
@@ -154,7 +154,7 @@ FReply FControlRigFunctionDragDropAction::DroppedOnPanel(const TSharedRef< class
 							FUIAction(
 							FExecuteAction::CreateLambda([Controller, LocalVariable, GraphPosition]()
 							{
-								Controller->AddVariableNode(LocalVariable.Name, LocalVariable.CPPType, LocalVariable.CPPTypeObject, true, LocalVariable.DefaultValue, GraphPosition);
+								Controller->AddVariableNode(LocalVariable.Name, LocalVariable.CPPType, LocalVariable.CPPTypeObject, true, LocalVariable.DefaultValue, GraphPosition, FString(), true, true);
 							}),
 							FCanExecuteAction()));
 
@@ -165,7 +165,7 @@ FReply FControlRigFunctionDragDropAction::DroppedOnPanel(const TSharedRef< class
 							FUIAction(
 							FExecuteAction::CreateLambda([Controller, LocalVariable, GraphPosition]()
 							{
-								Controller->AddVariableNode(LocalVariable.Name, LocalVariable.CPPType, LocalVariable.CPPTypeObject, false, LocalVariable.DefaultValue, GraphPosition);
+								Controller->AddVariableNode(LocalVariable.Name, LocalVariable.CPPType, LocalVariable.CPPTypeObject, false, LocalVariable.DefaultValue, GraphPosition, FString(), true, true);
 							}),
 							FCanExecuteAction()));
 
@@ -210,7 +210,7 @@ FReply FControlRigFunctionDragDropAction::DroppedOnPanel(const TSharedRef< class
 									}
 								}
 #endif
-								TargetController->AddFunctionReferenceNode(FunctionDefinitionNode, GraphPosition);
+								TargetController->AddFunctionReferenceNode(FunctionDefinitionNode, GraphPosition, FString(), true, true);
 							}
 						}
 					}
@@ -408,7 +408,7 @@ bool UControlRigGraphSchema::TryCreateConnection(UEdGraphPin* PinA, UEdGraphPin*
 			}
 #endif
 			
-			return Controller->AddLink(PinA->GetName(), PinB->GetName());
+			return Controller->AddLink(PinA->GetName(), PinB->GetName(), true, true);
 		}
 	}
 	return false;
@@ -632,7 +632,7 @@ void UControlRigGraphSchema::BreakPinLinks(UEdGraphPin& TargetPin, bool bSendsNo
 	// cache this here, as BreakPinLinks can trigger a node reconstruction invalidating the TargetPin referenceS
 	if (UControlRigGraphNode* Node = Cast< UControlRigGraphNode>(TargetPin.GetOwningNode()))
 	{
-		Node->GetController()->BreakAllLinks(TargetPin.GetName(), TargetPin.Direction == EGPD_Input);
+		Node->GetController()->BreakAllLinks(TargetPin.GetName(), TargetPin.Direction == EGPD_Input, true, true);
 	}
 }
 
@@ -649,7 +649,7 @@ void UControlRigGraphSchema::BreakSinglePinLink(UEdGraphPin* SourcePin, UEdGraph
 			SourcePin = Temp;
 		}
 		
-		Node->GetController()->BreakLink(SourcePin->GetName(), TargetPin->GetName());
+		Node->GetController()->BreakLink(SourcePin->GetName(), TargetPin->GetName(), true, true);
 	}
 }
 
@@ -760,7 +760,7 @@ void UControlRigGraphSchema::OnPinConnectionDoubleCicked(UEdGraphPin* PinA, UEdG
 	{
 		if (URigVMLink* Link = Node->GetModel()->FindLink(FString::Printf(TEXT("%s -> %s"), *PinA->GetName(), *PinB->GetName())))
 		{
-			Node->GetController()->AddRerouteNodeOnLink(Link, false, GraphPosition);
+			Node->GetController()->AddRerouteNodeOnLink(Link, false, GraphPosition, FString(), true, true);
 		}
 	}
 }
@@ -794,14 +794,14 @@ void UControlRigGraphSchema::SetNodePosition(UEdGraphNode* Node, const FVector2D
 
 	if (UControlRigGraphNode* RigNode = Cast<UControlRigGraphNode>(Node))
 	{
-		RigNode->GetController()->SetNodePosition(RigNode->GetModelNode(), Position, bSetupUndo, false);
+		RigNode->GetController()->SetNodePosition(RigNode->GetModelNode(), Position, bSetupUndo, false, false);
 	}
 	
 	if (UEdGraphNode_Comment* CommentNode = Cast<UEdGraphNode_Comment>(Node))
 	{
 		if(UControlRigGraph* Graph = CommentNode->GetTypedOuter<UControlRigGraph>())
 		{
-			Graph->GetController()->SetNodePositionByName(CommentNode->GetFName(), Position, bSetupUndo, false);
+			Graph->GetController()->SetNodePositionByName(CommentNode->GetFName(), Position, bSetupUndo, false, false);
 		}
 	}
 }
@@ -906,7 +906,7 @@ FReply UControlRigGraphSchema::TrySetGraphCategory(const UEdGraph* InGraph, cons
 				{
 					if (URigVMController* Controller = RigBlueprint->GetOrCreateController(CollapseNode->GetGraph()))
 					{
-						if (Controller->SetNodeCategory(CollapseNode, InCategory.ToString()))
+						if (Controller->SetNodeCategory(CollapseNode, InCategory.ToString(), true, false, true))
 						{
 							return FReply::Handled();
 						}
@@ -1043,7 +1043,9 @@ UEdGraphPin* UControlRigGraphSchema::DropPinOnNode(UEdGraphNode* InTargetNode, c
 								PinDirection,
 								TypeName,
 								TypeObjectPathName,
-								DefaultValue
+								DefaultValue,
+								true,
+								true
 							);
 							
 							if (!ExposedPinName.IsNone())
@@ -1261,7 +1263,7 @@ bool UControlRigGraphSchema::SafeDeleteNodeFromGraph(UEdGraph* Graph, UEdGraphNo
 {
 	if (UControlRigGraphNode* RigNode = Cast<UControlRigGraphNode>(Node))
 	{
-		return RigNode->GetController()->RemoveNode(RigNode->GetModelNode());
+		return RigNode->GetController()->RemoveNode(RigNode->GetModelNode(), true, false, true);
 	}
 	return false;
 }
@@ -1307,7 +1309,7 @@ bool UControlRigGraphSchema::RequestVariableDropOnPin(UEdGraph* InGraph, FProper
 					FModifierKeysState KeyState = FSlateApplication::Get().GetModifierKeys();
 					if (KeyState.IsAltDown())
 					{
-						return Graph->GetController()->BindPinToVariable(ModelPin->GetPinPath(), InVariableToDrop->GetName());
+						return Graph->GetController()->BindPinToVariable(ModelPin->GetPinPath(), InVariableToDrop->GetName(), true, true);
 					}
 					else
 					{
@@ -1397,7 +1399,7 @@ void UControlRigGraphSchema::EndGraphNodeInteraction(UEdGraphNode* InNode) const
 				Graph->GetController()->SetNodePositionByName(NodeName, *OldPosition, false, false);
 			}
 			
-			if(Graph->GetController()->SetNodePositionByName(NodeName, NewPosition, true, false))
+			if(Graph->GetController()->SetNodePositionByName(NodeName, NewPosition, true, false, true))
 			{
 				bMovedSomething = true;
 			}
diff --git a/Engine/Plugins/Animation/ControlRig/Source/ControlRigDeveloper/Public/Graph/ControlRigGraphNode.h b/Engine/Plugins/Animation/ControlRig/Source/ControlRigDeveloper/Public/Graph/ControlRigGraphNode.h
--- a/Engine/Plugins/Animation/ControlRig/Source/ControlRigDeveloper/Public/Graph/ControlRigGraphNode.h
+++ b/Engine/Plugins/Animation/ControlRig/Source/ControlRigDeveloper/Public/Graph/ControlRigGraphNode.h
@@ -116,7 +116,7 @@ public:
 	bool IsPinExpanded(const FString& InPinPath);
 
 	/** Propagate pin defaults to underlying properties if they have changed */
-	void CopyPinDefaultsToModel(UEdGraphPin* Pin, bool bUndo = false);
+	void CopyPinDefaultsToModel(UEdGraphPin* Pin, bool bUndo = false, bool bPrintPythonCommand = false);
 
 	/** Get the blueprint that this node is contained within */
 	UControlRigBlueprint* GetBlueprint() const;
diff --git a/Engine/Plugins/Animation/ControlRig/Source/ControlRigEditor/ControlRigEditor.Build.cs b/Engine/Plugins/Animation/ControlRig/Source/ControlRigEditor/ControlRigEditor.Build.cs
--- a/Engine/Plugins/Animation/ControlRig/Source/ControlRigEditor/ControlRigEditor.Build.cs
+++ b/Engine/Plugins/Animation/ControlRig/Source/ControlRigEditor/ControlRigEditor.Build.cs
@@ -81,6 +81,7 @@ namespace UnrealBuildTool.Rules
                     "SequencerScripting",
 					"PropertyAccessEditor",
 					"KismetWidgets",
+					"PythonScriptPlugin",
 				}
             );
 
diff --git a/Engine/Plugins/Animation/ControlRig/Source/ControlRigEditor/Private/ControlRigCompilerDetails.cpp b/Engine/Plugins/Animation/ControlRig/Source/ControlRigEditor/Private/ControlRigCompilerDetails.cpp
--- a/Engine/Plugins/Animation/ControlRig/Source/ControlRigEditor/Private/ControlRigCompilerDetails.cpp
+++ b/Engine/Plugins/Animation/ControlRig/Source/ControlRigEditor/Private/ControlRigCompilerDetails.cpp
@@ -10,6 +10,8 @@
 #include "Widgets/Text/STextBlock.h"
 #include "HAL/PlatformApplicationMisc.h"
 #include "ControlRig.h"
+#include "IPythonScriptPlugin.h"
+#include "RigVMPythonUtils.h"
 
 #define LOCTEXT_NAMESPACE "ControlRigCompilerDetails"
 
@@ -96,6 +98,18 @@ void FRigVMCompileSettingsDetails::CustomizeChildren(TSharedRef<IPropertyHandle>
 						.Text(LOCTEXT("CopyPythonScript", "Copy Python Script"))
 					]
 				]
+				+ SVerticalBox::Slot()
+				[
+					SNew(SButton)
+					.OnClicked(this, &FRigVMCompileSettingsDetails::OnRunPythonContextClicked)
+					.ContentPadding(FMargin(2))
+					.Content()
+					[
+						SNew(STextBlock)
+						.Justification(ETextJustify::Center)
+						.Text(LOCTEXT("RunPythonContext", "Run Python Context"))
+					]
+				]
 			];
 	}
 }
@@ -147,4 +161,34 @@ FReply FRigVMCompileSettingsDetails::OnCopyPythonScriptClicked()
 	return FReply::Handled();
 }
 
+FReply FRigVMCompileSettingsDetails::OnRunPythonContextClicked()
+{
+	if (BlueprintBeingCustomized)
+	{
+		FString BlueprintName = BlueprintBeingCustomized->GetPathName();
+		int32 DotIndex = BlueprintName.Find(TEXT("."), ESearchCase::IgnoreCase, ESearchDir::FromEnd);
+		if (DotIndex != INDEX_NONE)
+		{
+			BlueprintName = BlueprintName.Left(DotIndex);
+		}
+		
+		TArray<FString> PyCommands = {
+			TEXT("import unreal"),
+			FString::Printf(TEXT("blueprint = unreal.load_object(name = '%s', outer = None)"), *BlueprintName),
+			TEXT("library = blueprint.get_local_function_library()"),
+			TEXT("library_controller = blueprint.get_controller(library)"),
+			TEXT("hierarchy = blueprint.hierarchy"),
+			TEXT("hierarchy_controller = hierarchy.get_controller()")};
+
+		for (FString& Command : PyCommands)
+		{
+			RigVMPythonUtils::Print(Command);
+		
+			// Run the Python commands
+			IPythonScriptPlugin::Get()->ExecPythonCommand(*Command);
+		}
+	}
+	return FReply::Handled();
+}
+
 #undef LOCTEXT_NAMESPACE
diff --git a/Engine/Plugins/Animation/ControlRig/Source/ControlRigEditor/Private/ControlRigCompilerDetails.h b/Engine/Plugins/Animation/ControlRig/Source/ControlRigEditor/Private/ControlRigCompilerDetails.h
--- a/Engine/Plugins/Animation/ControlRig/Source/ControlRigEditor/Private/ControlRigCompilerDetails.h
+++ b/Engine/Plugins/Animation/ControlRig/Source/ControlRigEditor/Private/ControlRigCompilerDetails.h
@@ -27,6 +27,7 @@ protected:
 	FReply OnCopyASTClicked();
 	FReply OnCopyByteCodeClicked();
 	FReply OnCopyPythonScriptClicked();
+	FReply OnRunPythonContextClicked();
 
 	UControlRigBlueprint* BlueprintBeingCustomized;
 };
diff --git a/Engine/Plugins/Animation/ControlRig/Source/ControlRigEditor/Private/ControlRigEditorModule.cpp b/Engine/Plugins/Animation/ControlRig/Source/ControlRigEditor/Private/ControlRigEditorModule.cpp
--- a/Engine/Plugins/Animation/ControlRig/Source/ControlRigEditor/Private/ControlRigEditorModule.cpp
+++ b/Engine/Plugins/Animation/ControlRig/Source/ControlRigEditor/Private/ControlRigEditorModule.cpp
@@ -1280,7 +1280,7 @@ void FControlRigEditorModule::GetContextMenuActions(const UControlRigGraphSchema
 							LOCTEXT("RemoveArrayPin_Tooltip", "Removes the selected element from the array"),
 							FSlateIcon(),
 							FUIAction(FExecuteAction::CreateLambda([Controller, ModelPin]() {
-								Controller->RemoveArrayPin(ModelPin->GetPinPath());
+								Controller->RemoveArrayPin(ModelPin->GetPinPath(), true, true);
 							})
 						));
 						Section.AddMenuEntry(
@@ -1289,7 +1289,7 @@ void FControlRigEditorModule::GetContextMenuActions(const UControlRigGraphSchema
 							LOCTEXT("DuplicateArrayPin_Tooltip", "Duplicates the selected element"),
 							FSlateIcon(),
 							FUIAction(FExecuteAction::CreateLambda([Controller, ModelPin]() {
-								Controller->DuplicateArrayPin(ModelPin->GetPinPath());
+								Controller->DuplicateArrayPin(ModelPin->GetPinPath(), true, true);
 							})
 						));
 					}
@@ -1308,7 +1308,7 @@ void FControlRigEditorModule::GetContextMenuActions(const UControlRigGraphSchema
 								LOCTEXT("MakeVariableNodeFromBinding_Tooltip", "Turns the variable binding on the pin to a variable node"),
 								FSlateIcon(),
 								FUIAction(FExecuteAction::CreateLambda([Controller, ModelPin, NodePosition]() {
-									Controller->MakeVariableNodeFromBinding(ModelPin->GetPinPath(), NodePosition);
+									Controller->MakeVariableNodeFromBinding(ModelPin->GetPinPath(), NodePosition, true, true);
 								})
 							));
 						}
@@ -1327,7 +1327,7 @@ void FControlRigEditorModule::GetContextMenuActions(const UControlRigGraphSchema
 									FModifierKeysState KeyState = FSlateApplication::Get().GetModifierKeys();
 									bool bCreateVariableNode = !KeyState.IsAltDown();
 
-									Controller->PromotePinToVariable(ModelPin->GetPinPath(), bCreateVariableNode, NodePosition);
+									Controller->PromotePinToVariable(ModelPin->GetPinPath(), bCreateVariableNode, NodePosition, true, true);
 								})
 							));
 						}
@@ -1396,7 +1396,7 @@ void FControlRigEditorModule::GetContextMenuActions(const UControlRigGraphSchema
 								LOCTEXT("EjectLastNode_Tooltip", "Eject the last injected node"),
 								FSlateIcon(),
 								FUIAction(FExecuteAction::CreateLambda([Controller, ModelPin]() {
-									Controller->EjectNodeFromPin(ModelPin->GetPinPath());
+									Controller->EjectNodeFromPin(ModelPin->GetPinPath(), true, true);
 								})
 							));
 						}
@@ -1443,7 +1443,7 @@ void FControlRigEditorModule::GetContextMenuActions(const UControlRigGraphSchema
 									LOCTEXT("AddAlphaInterp_Tooltip", "Injects an interpolate node"),
 									FSlateIcon(),
 									FUIAction(FExecuteAction::CreateLambda([Controller, InGraphPin, ModelPin, ScriptStruct]() {
-										URigVMInjectionInfo* Injection = Controller->AddInjectedNode(ModelPin->GetPinPath(), ModelPin->GetDirection() != ERigVMPinDirection::Output, ScriptStruct, FRigUnit::GetMethodName(), TEXT("Value"), TEXT("Result"));
+										URigVMInjectionInfo* Injection = Controller->AddInjectedNode(ModelPin->GetPinPath(), ModelPin->GetDirection() != ERigVMPinDirection::Output, ScriptStruct, FRigUnit::GetMethodName(), TEXT("Value"), TEXT("Result"), FString(), true, true);
 										if (Injection)
 										{
 											TArray<FName> NodeNames;
@@ -1472,7 +1472,7 @@ void FControlRigEditorModule::GetContextMenuActions(const UControlRigGraphSchema
 									LOCTEXT("RemoveAlphaInterp_Tooltip", "Removes the interpolate node"),
 									FSlateIcon(),
 									FUIAction(FExecuteAction::CreateLambda([Controller, InGraphPin, ModelPin, InterpNode]() {
-										Controller->RemoveNodeByName(InterpNode->GetFName());
+										Controller->RemoveNodeByName(InterpNode->GetFName(), true, false, true);
 									})
 								));
 							}
@@ -1525,7 +1525,7 @@ void FControlRigEditorModule::GetContextMenuActions(const UControlRigGraphSchema
 									LOCTEXT("AddVisualDebug_Tooltip", "Injects a visual debugging node"),
 									FSlateIcon(),
 									FUIAction(FExecuteAction::CreateLambda([RigBlueprint, Controller, InGraphPin, ModelPin, ScriptStruct]() {
-										URigVMInjectionInfo* Injection = Controller->AddInjectedNode(ModelPin->GetPinPath(), ModelPin->GetDirection() != ERigVMPinDirection::Output, ScriptStruct, FRigUnit::GetMethodName(), TEXT("Value"), TEXT("Value"));
+										URigVMInjectionInfo* Injection = Controller->AddInjectedNode(ModelPin->GetPinPath(), ModelPin->GetDirection() != ERigVMPinDirection::Output, ScriptStruct, FRigUnit::GetMethodName(), TEXT("Value"), TEXT("Value"), FString(), true, true);
 										if (Injection)
 										{
 											TArray<FName> NodeNames;
@@ -1554,11 +1554,11 @@ void FControlRigEditorModule::GetContextMenuActions(const UControlRigGraphSchema
 															if(URigVMPin* SpaceTypePin = SpacePin->FindSubPin(TEXT("Type")))
 															{
 																FString SpaceTypeStr = StaticEnum<ERigElementType>()->GetDisplayNameTextByValue((int64)SpaceKey.Type).ToString();
-																Controller->SetPinDefaultValue(SpaceTypePin->GetPinPath(), SpaceTypeStr);
+																Controller->SetPinDefaultValue(SpaceTypePin->GetPinPath(), SpaceTypeStr, true, true, false, true);
 															}
 															if(URigVMPin* SpaceNamePin = SpacePin->FindSubPin(TEXT("Name")))
 															{
-																Controller->SetPinDefaultValue(SpaceNamePin->GetPinPath(), SpaceKey.Name.ToString());
+																Controller->SetPinDefaultValue(SpaceNamePin->GetPinPath(), SpaceKey.Name.ToString(), true, true, false, true);
 															}
 														}
 													}
@@ -1589,7 +1589,7 @@ void FControlRigEditorModule::GetContextMenuActions(const UControlRigGraphSchema
 									FUIAction(FExecuteAction::CreateLambda([Controller, VisualDebugNode]() {
 										URigVMPin* EnabledPin = VisualDebugNode->FindPin(TEXT("bEnabled"));
 										check(EnabledPin);
-										Controller->SetPinDefaultValue(EnabledPin->GetPinPath(), EnabledPin->GetDefaultValue() == TEXT("True") ? TEXT("False") : TEXT("True"), false);
+										Controller->SetPinDefaultValue(EnabledPin->GetPinPath(), EnabledPin->GetDefaultValue() == TEXT("True") ? TEXT("False") : TEXT("True"), false, true, false, true);
 									})
 								));
 								Section.AddMenuEntry(
@@ -1598,7 +1598,7 @@ void FControlRigEditorModule::GetContextMenuActions(const UControlRigGraphSchema
 									LOCTEXT("RemoveVisualDebug_Tooltip", "Removes the visual debugging node"),
 									FSlateIcon(),
 										FUIAction(FExecuteAction::CreateLambda([Controller, InGraphPin, ModelPin, VisualDebugNode]() {
-										Controller->RemoveNodeByName(VisualDebugNode->GetFName());
+										Controller->RemoveNodeByName(VisualDebugNode->GetFName(), true, false, true);
 									})
 								));
 							}
@@ -1853,7 +1853,7 @@ void FControlRigEditorModule::GetContextMenuActions(const UControlRigGraphSchema
 										Key.Name = *NewNameStr;
 										if(RigBlueprint->Hierarchy->GetIndex(Key) != INDEX_NONE)
 										{
-											Controller->SetPinDefaultValue(Pin->GetPinPath(), NewNameStr, false);
+											Controller->SetPinDefaultValue(Pin->GetPinPath(), NewNameStr, false, true, false, true);
 											ReplacedNames++;
 										}
 										else
@@ -1990,7 +1990,7 @@ void FControlRigEditorModule::GetContextMenuActions(const UControlRigGraphSchema
 						FSlateIcon(),
 						FUIAction(FExecuteAction::CreateLambda([Model, Controller]() {
 							TArray<FName> Nodes = Model->GetSelectNodes();
-							Controller->CollapseNodes(Nodes);
+							Controller->CollapseNodes(Nodes, FString(), true, true);
 						})
 					));
 					OrganizationSection.AddMenuEntry(
@@ -2001,8 +2001,8 @@ void FControlRigEditorModule::GetContextMenuActions(const UControlRigGraphSchema
 						FUIAction(FExecuteAction::CreateLambda([Model, Controller]() {
 							TArray<FName> Nodes = Model->GetSelectNodes();
 							Controller->OpenUndoBracket(TEXT("Collapse to Function"));
-							URigVMCollapseNode* CollapseNode = Controller->CollapseNodes(Nodes, TEXT("New Function"));
-							Controller->PromoteCollapseNodeToFunctionReferenceNode(CollapseNode->GetFName());
+							URigVMCollapseNode* CollapseNode = Controller->CollapseNodes(Nodes, TEXT("New Function"), true, true);
+							Controller->PromoteCollapseNodeToFunctionReferenceNode(CollapseNode->GetFName(), true, true);
 							Controller->CloseUndoBracket();
 						})
 					));
@@ -2015,7 +2015,7 @@ void FControlRigEditorModule::GetContextMenuActions(const UControlRigGraphSchema
 							LOCTEXT("PromoteToFunction_Tooltip", "Turns the Collapse Node into a Function"),
 							FSlateIcon(),
 							FUIAction(FExecuteAction::CreateLambda([Controller, CollapseNode]() {
-								Controller->PromoteCollapseNodeToFunctionReferenceNode(CollapseNode->GetFName());
+								Controller->PromoteCollapseNodeToFunctionReferenceNode(CollapseNode->GetFName(), true, true);
 							})
 						));
 					}
@@ -2107,7 +2107,7 @@ void FControlRigEditorModule::GetContextMenuActions(const UControlRigGraphSchema
 							FSlateIcon(),
 							FUIAction(FExecuteAction::CreateLambda([Controller, LibraryNode]() {
 								Controller->OpenUndoBracket(TEXT("Expand node"));
-								TArray<URigVMNode*> ExpandedNodes = Controller->ExpandLibraryNode(LibraryNode->GetFName());
+								TArray<URigVMNode*> ExpandedNodes = Controller->ExpandLibraryNode(LibraryNode->GetFName(), true, true);
 								if (ExpandedNodes.Num() > 0)
 								{
 									TArray<FName> ExpandedNodeNames;
diff --git a/Engine/Plugins/Animation/ControlRig/Source/ControlRigEditor/Private/ControlRigElementDetails.cpp b/Engine/Plugins/Animation/ControlRig/Source/ControlRigEditor/Private/ControlRigElementDetails.cpp
--- a/Engine/Plugins/Animation/ControlRig/Source/ControlRigEditor/Private/ControlRigElementDetails.cpp
+++ b/Engine/Plugins/Animation/ControlRig/Source/ControlRigEditor/Private/ControlRigElementDetails.cpp
@@ -781,7 +781,7 @@ void FRigBaseElementDetails::SetName(const FText& InNewText, ETextCommit::Type I
 	{
 		URigHierarchyController* Controller = Hierarchy->GetController(true);
 		check(Controller);
-		Controller->RenameElement(ElementKeyBeingCustomized, *InNewText.ToString(), true);
+		Controller->RenameElement(ElementKeyBeingCustomized, *InNewText.ToString(), true, true);
 	}
 }
 
diff --git a/Engine/Plugins/Animation/ControlRig/Source/ControlRigEditor/Private/ControlRigGraphDetails.cpp b/Engine/Plugins/Animation/ControlRig/Source/ControlRigEditor/Private/ControlRigGraphDetails.cpp
--- a/Engine/Plugins/Animation/ControlRig/Source/ControlRigEditor/Private/ControlRigGraphDetails.cpp
+++ b/Engine/Plugins/Animation/ControlRig/Source/ControlRigEditor/Private/ControlRigGraphDetails.cpp
@@ -225,7 +225,7 @@ void FControlRigArgumentLayout::OnRemoveClicked()
 		{
 			if (URigVMController* Controller = Blueprint->GetController(LibraryNode->GetContainedGraph()))
 			{
-				Controller->RemoveExposedPin(Pin->GetFName(), true);
+				Controller->RemoveExposedPin(Pin->GetFName(), true, true);
 			}
 		}
 	}
@@ -241,7 +241,7 @@ FReply FControlRigArgumentLayout::OnArgMoveUp()
 		{
 			if (URigVMController* Controller = Blueprint->GetController(LibraryNode->GetContainedGraph()))
 			{
-				Controller->SetExposedPinIndex(Pin->GetFName(), Pin->GetPinIndex() - 1);
+				Controller->SetExposedPinIndex(Pin->GetFName(), Pin->GetPinIndex() - 1, true, true);
 				return FReply::Handled();
 			}
 		}
@@ -259,7 +259,7 @@ FReply FControlRigArgumentLayout::OnArgMoveDown()
 		{
 			if (URigVMController* Controller = Blueprint->GetController(LibraryNode->GetContainedGraph()))
 			{
-				Controller->SetExposedPinIndex(Pin->GetFName(), Pin->GetPinIndex() + 1);
+				Controller->SetExposedPinIndex(Pin->GetFName(), Pin->GetPinIndex() + 1, true, true);
 				return FReply::Handled();
 			}
 		}
@@ -313,7 +313,7 @@ void FControlRigArgumentLayout::OnArgNameTextCommitted(const FText& NewText, ETe
 			if (URigVMController* Controller = Blueprint->GetController(LibraryNode->GetContainedGraph()))
 			{
 				const FString& NewName = NewText.ToString();
-				Controller->RenameExposedPin(Pin->GetFName(), *NewName, true);
+				Controller->RenameExposedPin(Pin->GetFName(), *NewName, true, true);
 			}
 		}
 	}
@@ -362,7 +362,7 @@ void FControlRigArgumentLayout::PinInfoChanged(const FEdGraphPinType& PinType)
 				}
 
 				bool bSetupUndoRedo = true;
-				Controller->ChangeExposedPinType(Pin->GetFName(), CPPType, CPPTypeObjectName, bSetupUndoRedo);
+				Controller->ChangeExposedPinType(Pin->GetFName(), CPPType, CPPTypeObjectName, bSetupUndoRedo, true, true);
 
 				// If the controller has identified this as a bulk change, it has not added the actions to the action stack
 				// We need to disable the transaction from the UI as well to keep them synced
@@ -809,7 +809,7 @@ FReply FControlRigGraphDetails::OnAddNewInputClicked()
 				}
 			}
 
-			Controller->AddExposedPin(ArgumentName, ERigVMPinDirection::Input, CPPType, CPPTypeObjectPath, DefaultValue, true);
+			Controller->AddExposedPin(ArgumentName, ERigVMPinDirection::Input, CPPType, CPPTypeObjectPath, DefaultValue, true, true);
 		}
 	}
 	return FReply::Unhandled();
@@ -829,7 +829,7 @@ FReply FControlRigGraphDetails::OnAddNewOutputClicked()
 			FString DefaultValue = TEXT("False");
 			// todo: base decisions on types on last argument
 
-			Controller->AddExposedPin(ArgumentName, ERigVMPinDirection::Output, CPPType, CPPTypeObjectPath, DefaultValue, true);
+			Controller->AddExposedPin(ArgumentName, ERigVMPinDirection::Output, CPPType, CPPTypeObjectPath, DefaultValue, true, true);
 		}
 	}
 	return FReply::Unhandled();
@@ -863,7 +863,7 @@ void FControlRigGraphDetails::SetNodeCategory(const FText& InNewText, ETextCommi
 			{
 				if (URigVMController* Controller = Blueprint->GetOrCreateController(OuterNode->GetGraph()))
 				{
-					Controller->SetNodeCategory(OuterNode, InNewText.ToString());
+					Controller->SetNodeCategory(OuterNode, InNewText.ToString(), true, false, true);
 				}
 			}
 		}
@@ -898,7 +898,7 @@ void FControlRigGraphDetails::SetNodeKeywords(const FText& InNewText, ETextCommi
 			{
 				if (URigVMController* Controller = Blueprint->GetOrCreateController(OuterNode->GetGraph()))
 				{
-					Controller->SetNodeKeywords(OuterNode, InNewText.ToString());
+					Controller->SetNodeKeywords(OuterNode, InNewText.ToString(), true, false, true);
 				}
 			}
 		}
@@ -933,7 +933,7 @@ void FControlRigGraphDetails::SetNodeDescription(const FText& InNewText, ETextCo
 			{
 				if (URigVMController* Controller = Blueprint->GetOrCreateController(OuterNode->GetGraph()))
 				{
-					Controller->SetNodeDescription(OuterNode, InNewText.ToString());
+					Controller->SetNodeDescription(OuterNode, InNewText.ToString(), true, false, true);
 				}
 			}
 		}
@@ -969,7 +969,7 @@ void FControlRigGraphDetails::SetNodeColor(FLinearColor InColor, bool bSetupUndo
 			{
 				if (URigVMController* Controller = Blueprint->GetOrCreateController(OuterNode->GetGraph()))
 				{
-					Controller->SetNodeColor(OuterNode, TargetColor, bSetupUndoRedo, bIsPickingColor);
+					Controller->SetNodeColor(OuterNode, TargetColor, bSetupUndoRedo, bIsPickingColor, true);
 				}
 			}
 		}
diff --git a/Engine/Plugins/Animation/ControlRig/Source/ControlRigEditor/Private/Editor/ControlRigEditor.cpp b/Engine/Plugins/Animation/ControlRig/Source/ControlRigEditor/Private/Editor/ControlRigEditor.cpp
--- a/Engine/Plugins/Animation/ControlRig/Source/ControlRigEditor/Private/Editor/ControlRigEditor.cpp
+++ b/Engine/Plugins/Animation/ControlRig/Source/ControlRigEditor/Private/Editor/ControlRigEditor.cpp
@@ -1949,7 +1949,7 @@ void FControlRigEditor::OnAddNewLocalVariable()
 		return;
 	}
 
-	FRigVMGraphVariableDescription NewVar = GetFocusedController()->AddLocalVariable(TEXT("NewLocalVar"), TEXT("bool"), nullptr, TEXT("0"));
+	FRigVMGraphVariableDescription NewVar = GetFocusedController()->AddLocalVariable(TEXT("NewLocalVar"), TEXT("bool"), nullptr, TEXT("0"), true, true);
 	if(NewVar.Name.IsNone())
 	{
 		LogSimpleMessage( LOCTEXT("AddLocalVariable_Error", "Adding new local variable failed.") );
@@ -1985,14 +1985,14 @@ void FControlRigEditor::DeleteSelectedNodes()
 				AnalyticsTrackNodeEvent(GetBlueprintObj(), Node, true);
 				if (UControlRigGraphNode* RigNode = Cast<UControlRigGraphNode>(Node))
 				{
-					if(GetFocusedController()->RemoveNodeByName(*RigNode->ModelNodePath))
+					if(GetFocusedController()->RemoveNodeByName(*RigNode->ModelNodePath, true, false, true))
 					{
 						DeletedAnything = true;
 					}
 				}
 				else if (UEdGraphNode_Comment* CommentNode = Cast<UEdGraphNode_Comment>(Node))
 				{
-					if(GetFocusedController()->RemoveNodeByName(CommentNode->GetFName()))
+					if(GetFocusedController()->RemoveNodeByName(CommentNode->GetFName(), true, false, true))
 					{
 						DeletedAnything = true;
 					}
@@ -2049,31 +2049,31 @@ FReply FControlRigEditor::OnSpawnGraphNodeByShortcut(FInputChord InChord, const
 			{
 				if(InChord.Key == EKeys::B)
 				{
-					Controller->AddBranchNode(InPosition);
+					Controller->AddBranchNode(InPosition, "", true, true);
 				}
 				else if(InChord.Key == EKeys::S)
 				{
-					Controller->AddUnitNode(FRigUnit_SequenceExecution::StaticStruct(), FRigUnit::GetMethodName(), InPosition);
+					Controller->AddUnitNode(FRigUnit_SequenceExecution::StaticStruct(), FRigUnit::GetMethodName(), InPosition, FString(), true, true);
 				}
 				else if(InChord.Key == EKeys::One)
 				{
-					Controller->AddUnitNode(FRigUnit_GetTransform::StaticStruct(), FRigUnit::GetMethodName(), InPosition);
+					Controller->AddUnitNode(FRigUnit_GetTransform::StaticStruct(), FRigUnit::GetMethodName(), InPosition, FString(), true, true);
 				}
 				else if(InChord.Key == EKeys::Two)
 				{
-					Controller->AddUnitNode(FRigUnit_SetTransform::StaticStruct(), FRigUnit::GetMethodName(), InPosition);
+					Controller->AddUnitNode(FRigUnit_SetTransform::StaticStruct(), FRigUnit::GetMethodName(), InPosition, FString(), true, true);
 				}
 				else if(InChord.Key == EKeys::Three)
 				{
-					Controller->AddUnitNode(FRigUnit_ParentConstraint::StaticStruct(), FRigUnit::GetMethodName(), InPosition);
+					Controller->AddUnitNode(FRigUnit_ParentConstraint::StaticStruct(), FRigUnit::GetMethodName(), InPosition, FString(), true, true);
 				}
 				else if(InChord.Key == EKeys::Four)
 				{
-					Controller->AddUnitNode(FRigUnit_GetControlFloat::StaticStruct(), FRigUnit::GetMethodName(), InPosition);
+					Controller->AddUnitNode(FRigUnit_GetControlFloat::StaticStruct(), FRigUnit::GetMethodName(), InPosition, FString(), true, true);
 				}
 				else if(InChord.Key == EKeys::Five)
 				{
-					Controller->AddUnitNode(FRigUnit_SetCurveValue::StaticStruct(), FRigUnit::GetMethodName(), InPosition);
+					Controller->AddUnitNode(FRigUnit_SetCurveValue::StaticStruct(), FRigUnit::GetMethodName(), InPosition, FString(), true, true);
 				}
 			}
 		}
@@ -2198,7 +2198,7 @@ void FControlRigEditor::PasteNodes()
 			check(Node);
 
 			FVector2D Position = Node->GetPosition();
-			GetFocusedController()->SetNodePositionByName(NodeName, PasteLocation + Position - Bounds.GetCenter());
+			GetFocusedController()->SetNodePositionByName(NodeName, PasteLocation + Position - Bounds.GetCenter(), true, false, true);
 		}
 
 		GetFocusedController()->SetNodeSelection(NodeNames);
@@ -2393,7 +2393,7 @@ void FControlRigEditor::NewDocument_OnClicked(ECreatedDocumentType GraphType)
 	{
 		if (URigVMController* Controller = Blueprint->GetOrCreateController(Blueprint->GetLocalFunctionLibrary()))
 		{
-			if (URigVMLibraryNode* FunctionNode = Controller->AddFunctionToLibrary(TEXT("New Function"), true))
+			if (URigVMLibraryNode* FunctionNode = Controller->AddFunctionToLibrary(TEXT("New Function"), true, FVector2D::ZeroVector, true, true))
 			{
 				if (UEdGraph* NewGraph = Blueprint->GetEdGraph(FunctionNode->GetContainedGraph()))
 				{
@@ -2841,13 +2841,13 @@ void FControlRigEditor::OnSelectedNodesChangedImpl(const TSet<class UObject*>& N
 						FVector2D NodePos(CommentNode->NodePosX, CommentNode->NodePosY);
 						FVector2D NodeSize(CommentNode->NodeWidth, CommentNode->NodeHeight);
 						FLinearColor NodeColor = CommentNode->CommentColor;
-						GetFocusedController()->AddCommentNode(CommentNode->NodeComment, NodePos, NodeSize, NodeColor, CommentNode->GetName(), true);
+						GetFocusedController()->AddCommentNode(CommentNode->NodeComment, NodePos, NodeSize, NodeColor, CommentNode->GetName(), true, true);
 					}
 				}
 				NodeNamesToSelect.Add(Node->GetFName());
 			}
 		}
-		GetFocusedController()->SetNodeSelection(NodeNamesToSelect, true);
+		GetFocusedController()->SetNodeSelection(NodeNamesToSelect, true, true);
 	}
 }
 
@@ -3894,7 +3894,7 @@ void FControlRigEditor::OnNodeTitleCommitted(const FText& NewText, ETextCommit::
 	{
 		if (UControlRigBlueprint* ControlRigBP = GetControlRigBlueprint())
 		{
-			GetFocusedController()->SetCommentTextByName(CommentBeingChanged->GetFName(), NewText.ToString(), true);
+			GetFocusedController()->SetCommentTextByName(CommentBeingChanged->GetFName(), NewText.ToString(), true, true);
 		}
 	}
 }
@@ -4058,7 +4058,7 @@ void FControlRigEditor::OnFinishedChangingProperties(const FPropertyChangedEvent
 		if (!DefaultValue.IsEmpty())
 		{
 			FString PinPath = FString::Printf(TEXT("%s.%s"), *NodeDetailName.ToString(), *PropertyChangedEvent.MemberProperty->GetName());
-			GetFocusedController()->SetPinDefaultValue(PinPath, DefaultValue, true, true);
+			GetFocusedController()->SetPinDefaultValue(PinPath, DefaultValue, true, true, false, true);
 		}
 	}
 
@@ -4281,7 +4281,7 @@ void FControlRigEditor::OnRequestLocalizeFunctionDialog(URigVMLibraryNode* InFun
 
 						if (LocalizationDialog->ShowModal() != EAppReturnType::Cancel)
 						{
-							TargetController->LocalizeFunctions(LocalizationDialog->GetFunctionsToLocalize());
+							TargetController->LocalizeFunctions(LocalizationDialog->GetFunctionsToLocalize(), true, true, true);
 						}
 					}
 				}
@@ -4994,7 +4994,7 @@ void FControlRigEditor::OnGraphNodeDropToPerform(TSharedPtr<FGraphNodeDragDropOp
 								{
 									Controller->OpenUndoBracket(TEXT("Create Collection from Items"));
 
-									if (URigVMNode* ItemsNode = Controller->AddUnitNode(FRigUnit_CollectionItems::StaticStruct(), FRigUnit::GetMethodName(), NodePosition))
+									if (URigVMNode* ItemsNode = Controller->AddUnitNode(FRigUnit_CollectionItems::StaticStruct(), FRigUnit::GetMethodName(), NodePosition, FString(), true, true))
 									{
 										if (URigVMPin* ItemsPin = ItemsNode->FindPin(TEXT("Items")))
 										{
@@ -5007,8 +5007,8 @@ void FControlRigEditor::OnGraphNodeDropToPerform(TSharedPtr<FGraphNodeDragDropOp
 											{
 												FString DefaultValue;
 												FRigElementKey::StaticStruct()->ExportText(DefaultValue, &DraggedKeys[ItemIndex], nullptr, nullptr, PPF_None, nullptr);
-												Controller->SetPinDefaultValue(ItemPins[ItemIndex]->GetPinPath(), DefaultValue);
-												Controller->SetPinExpansion(ItemPins[ItemIndex]->GetPinPath(), true);
+												Controller->SetPinDefaultValue(ItemPins[ItemIndex]->GetPinPath(), DefaultValue, true, true, false, true);
+												Controller->SetPinExpansion(ItemPins[ItemIndex]->GetPinPath(), true, true, true);
 											}
 										}
 									}
@@ -5305,7 +5305,7 @@ void FControlRigEditor::HandleMakeElementGetterSetter(ERigElementGetterSetterTyp
 		}
 
 		FName Name = FControlRigBlueprintUtils::ValidateName(Blueprint, StructTemplate->GetName());
-		if (URigVMUnitNode* ModelNode = GetFocusedController()->AddUnitNode(StructTemplate, FRigUnit::GetMethodName(), NodePosition))
+		if (URigVMUnitNode* ModelNode = GetFocusedController()->AddUnitNode(StructTemplate, FRigUnit::GetMethodName(), NodePosition, FString(), true, true))
 		{
 			FString ItemTypeStr = StaticEnum<ERigElementType>()->GetDisplayNameTextByValue((int64)Key.Type).ToString();
 			NewNode.Name = ModelNode->GetFName();
@@ -5313,13 +5313,13 @@ void FControlRigEditor::HandleMakeElementGetterSetter(ERigElementGetterSetterTyp
 			
 			for (const FName& ItemPin : ItemPins)
 			{
-				GetFocusedController()->SetPinDefaultValue(FString::Printf(TEXT("%s.%s.Name"), *ModelNode->GetName(), *ItemPin.ToString()), Key.Name.ToString());
-				GetFocusedController()->SetPinDefaultValue(FString::Printf(TEXT("%s.%s.Type"), *ModelNode->GetName(), *ItemPin.ToString()), ItemTypeStr);
+				GetFocusedController()->SetPinDefaultValue(FString::Printf(TEXT("%s.%s.Name"), *ModelNode->GetName(), *ItemPin.ToString()), Key.Name.ToString(), true, true, false, true);
+				GetFocusedController()->SetPinDefaultValue(FString::Printf(TEXT("%s.%s.Type"), *ModelNode->GetName(), *ItemPin.ToString()), ItemTypeStr, true, true, false, true);
 			}
 
 			for (const FName& NamePin : NamePins)
 			{
-				GetFocusedController()->SetPinDefaultValue(FString::Printf(TEXT("%s.%s"), *ModelNode->GetName(), *NamePin.ToString()), Key.Name.ToString());
+				GetFocusedController()->SetPinDefaultValue(FString::Printf(TEXT("%s.%s"), *ModelNode->GetName(), *NamePin.ToString()), Key.Name.ToString(), true, true, false, true);
 			}
 
 			if (!NewNode.ValuePinName.IsNone())
@@ -5351,7 +5351,7 @@ void FControlRigEditor::HandleMakeElementGetterSetter(ERigElementGetterSetterTyp
 				}
 				if (!DefaultValue.IsEmpty())
 				{
-					GetFocusedController()->SetPinDefaultValue(FString::Printf(TEXT("%s.%s"), *ModelNode->GetName(), *NewNode.ValuePinName.ToString()), DefaultValue);
+					GetFocusedController()->SetPinDefaultValue(FString::Printf(TEXT("%s.%s"), *ModelNode->GetName(), *NewNode.ValuePinName.ToString()), DefaultValue, true, true, false, true);
 				}
 			}
 
@@ -5518,7 +5518,7 @@ void FControlRigEditor::HandleVariableDroppedFromBlueprint(UObject* InSubject, F
 		FUIAction(
 			FExecuteAction::CreateLambda([ExternalVariable, Controller, InDropPosition] {
 
-				Controller->AddVariableNode(ExternalVariable.Name, ExternalVariable.TypeName.ToString(), ExternalVariable.TypeObject, true, FString(), InDropPosition);
+				Controller->AddVariableNode(ExternalVariable.Name, ExternalVariable.TypeName.ToString(), ExternalVariable.TypeObject, true, FString(), InDropPosition, FString(), true, true);
 
 			}),
 			FCanExecuteAction()
@@ -5532,7 +5532,7 @@ void FControlRigEditor::HandleVariableDroppedFromBlueprint(UObject* InSubject, F
 		FUIAction(
 			FExecuteAction::CreateLambda([ExternalVariable, Controller, InDropPosition] {
 
-				Controller->AddVariableNode(ExternalVariable.Name, ExternalVariable.TypeName.ToString(), ExternalVariable.TypeObject, false, FString(), InDropPosition);
+				Controller->AddVariableNode(ExternalVariable.Name, ExternalVariable.TypeName.ToString(), ExternalVariable.TypeObject, false, FString(), InDropPosition, FString(), true, true);
 
 			}),
 			FCanExecuteAction()
diff --git a/Engine/Plugins/Animation/ControlRig/Source/ControlRigEditor/Private/Editor/SRigCurveContainer.cpp b/Engine/Plugins/Animation/ControlRig/Source/ControlRigEditor/Private/Editor/SRigCurveContainer.cpp
--- a/Engine/Plugins/Animation/ControlRig/Source/ControlRigEditor/Private/Editor/SRigCurveContainer.cpp
+++ b/Engine/Plugins/Animation/ControlRig/Source/ControlRigEditor/Private/Editor/SRigCurveContainer.cpp
@@ -447,7 +447,7 @@ void SRigCurveContainer::OnNameCommitted(const FText& InNewName, ETextCommit::Ty
 			{
 				FName NewName = FName(*InNewName.ToString());
 				FName OldName = Item->CurveName;
-				Controller->RenameElement(FRigElementKey(OldName, ERigElementType::Curve), NewName, true);
+				Controller->RenameElement(FRigElementKey(OldName, ERigElementType::Curve), NewName, true, true);
 			}
 		}
 	}
@@ -465,7 +465,7 @@ void SRigCurveContainer::OnDeleteNameClicked()
 		{
 			if(URigHierarchyController* Controller = Hierarchy->GetController())
 			{
-				Controller->RemoveElement(FRigElementKey(Item->CurveName, ERigElementType::Curve), true);
+				Controller->RemoveElement(FRigElementKey(Item->CurveName, ERigElementType::Curve), true, true);
 			}
 		}
 	}
@@ -507,7 +507,7 @@ void SRigCurveContainer::ChangeCurveName(const FName& OldName, const FName& NewN
 	{
 		if(URigHierarchyController* Controller = Hierarchy->GetController())
 		{
-			Controller->RenameElement(FRigElementKey(OldName, ERigElementType::Curve), NewName, true);
+			Controller->RenameElement(FRigElementKey(OldName, ERigElementType::Curve), NewName, true, true);
 		}
 	}
 }
@@ -682,7 +682,7 @@ void SRigCurveContainer::ImportCurve(const FAssetData& InAssetData)
 			if(URigHierarchyController* Controller = Hierarchy->GetController())
 			{
 				Controller->ClearSelection();
-				Controller->ImportCurves(Skeleton, NAME_None, false, true);
+				Controller->ImportCurves(Skeleton, NAME_None, false, true, true);
 			}
 
 			FSlateApplication::Get().DismissAllMenus();
diff --git a/Engine/Plugins/Animation/ControlRig/Source/ControlRigEditor/Private/Editor/SRigHierarchy.cpp b/Engine/Plugins/Animation/ControlRig/Source/ControlRigEditor/Private/Editor/SRigHierarchy.cpp
--- a/Engine/Plugins/Animation/ControlRig/Source/ControlRigEditor/Private/Editor/SRigHierarchy.cpp
+++ b/Engine/Plugins/Animation/ControlRig/Source/ControlRigEditor/Private/Editor/SRigHierarchy.cpp
@@ -1629,8 +1629,8 @@ void SRigHierarchy::RefreshHierarchy(const FAssetData& InAssetData)
 		URigHierarchyController* Controller = Hierarchy->GetController(true);
 		check(Controller);
 		
-		Controller->ImportBones(RefSkeleton, NAME_None, true, true, bSelectBones, true);
-		Controller->ImportCurves(Mesh->GetSkeleton(), NAME_None, false, true);
+		Controller->ImportBones(Mesh->GetSkeleton(), NAME_None, true, true, bSelectBones, true, true);
+		Controller->ImportCurves(Mesh->GetSkeleton(), NAME_None, false, true, true);
 	}
 
 	ControlRigBlueprint->PropagateHierarchyFromBPToInstances();
@@ -1697,12 +1697,10 @@ void SRigHierarchy::ImportHierarchy(const FAssetData& InAssetData)
 			bSelectBones = !CurrentRig->IsSetupModeEnabled();
 		}
 
-		const FReferenceSkeleton& RefSkeleton = Mesh->GetRefSkeleton();
-
 		URigHierarchyController* Controller = Hierarchy->GetController(true);
 		check(Controller);
 
-		TArray<FRigElementKey> ImportedBones = Controller->ImportBones(RefSkeleton, NAME_None, false, false, bSelectBones, true);
+		TArray<FRigElementKey> ImportedBones = Controller->ImportBones(Mesh->GetSkeleton(), NAME_None, false, false, bSelectBones, true, true);
 		Controller->ImportCurves(Mesh->GetSkeleton(), NAME_None, true, true);
 
 		ControlRigBlueprint->SourceHierarchyImport = Mesh->GetSkeleton();
@@ -1848,7 +1846,7 @@ void SRigHierarchy::HandleDeleteItem()
 				}
 			}
 
-			Controller->RemoveElement(SelectedKey, true);
+			Controller->RemoveElement(SelectedKey, true, true);
 			RemovedItems.Add(SelectedKey);
 		}
 	}
@@ -1902,7 +1900,7 @@ void SRigHierarchy::HandleNewItem(ERigElementType InElementType)
 			{
 				case ERigElementType::Bone:
 				{
-					NewItemKey = Controller->AddBone(NewElementName, ParentKey, ParentTransform, true, ERigBoneType::User, true);
+					NewItemKey = Controller->AddBone(NewElementName, ParentKey, ParentTransform, true, ERigBoneType::User, true, true);
 					break;
 				}
 				case ERigElementType::Control:
@@ -1910,12 +1908,12 @@ void SRigHierarchy::HandleNewItem(ERigElementType InElementType)
 					FRigControlSettings Settings;
 					Settings.ControlType = ERigControlType::EulerTransform;
 						
-					NewItemKey = Controller->AddControl(NewElementName, ParentKey, Settings, Settings.GetIdentityValue(), FTransform::Identity, FTransform::Identity, true);
+					NewItemKey = Controller->AddControl(NewElementName, ParentKey, Settings, Settings.GetIdentityValue(), FTransform::Identity, FTransform::Identity, true, true);
 					break;
 				}
 				case ERigElementType::Null:
 				{
-					NewItemKey = Controller->AddNull(NewElementName, ParentKey, ParentTransform, true, true);
+					NewItemKey = Controller->AddNull(NewElementName, ParentKey, ParentTransform, true, true, true);
 					break;
 				}
 				default:
@@ -2408,7 +2406,7 @@ FName SRigHierarchy::RenameElement(const FRigElementKey& OldKey, const FString&
 		const FName SanitizedName = *SanitizedNameStr;
 		FName ResultingName = NAME_None;
 
-		ResultingName = Controller->RenameElement(OldKey, SanitizedName, true).Name;
+		ResultingName = Controller->RenameElement(OldKey, SanitizedName, true, true).Name;
 		ControlRigBlueprint->PropagateHierarchyFromBPToInstances();
 		return ResultingName;
 	}
@@ -2657,14 +2655,14 @@ void SRigHierarchy::HandleUnparent()
 
 				if (bUnparentImportedBones || !bIsImportedBone)
 				{
-					Controller->RemoveAllParents(SelectedKey, true, true);
+					Controller->RemoveAllParents(SelectedKey, true, true, true);
 				}
 				break;
 			}
 			case ERigElementType::Null:
 			case ERigElementType::Control:
 			{
-				Controller->RemoveAllParents(SelectedKey, true, true);
+				Controller->RemoveAllParents(SelectedKey, true, true, true);
 				break;
 			}
 			default:
@@ -2834,7 +2832,7 @@ FReply SRigHierarchy::ReparentOrMatchTransform(const TArray<FRigElementKey>& Dra
 			}
 			else
 			{
-				Controller->RemoveAllParents(DraggedKey, true, true);
+				Controller->RemoveAllParents(DraggedKey, true, true, true);
 			}
 
 			DebuggedHierarchy->SetInitialGlobalTransform(DraggedKey, InitialTransform, true, true);
diff --git a/Engine/Plugins/Animation/ControlRig/Source/ControlRigEditor/Private/Graph/NodeSpawners/ControlRigBranchNodeSpawner.cpp b/Engine/Plugins/Animation/ControlRig/Source/ControlRigEditor/Private/Graph/NodeSpawners/ControlRigBranchNodeSpawner.cpp
--- a/Engine/Plugins/Animation/ControlRig/Source/ControlRigEditor/Private/Graph/NodeSpawners/ControlRigBranchNodeSpawner.cpp
+++ b/Engine/Plugins/Animation/ControlRig/Source/ControlRigEditor/Private/Graph/NodeSpawners/ControlRigBranchNodeSpawner.cpp
@@ -53,6 +53,7 @@ UEdGraphNode* UControlRigBranchNodeSpawner::Invoke(UEdGraph* ParentGraph, FBindi
 
 	bool const bIsTemplateNode = FBlueprintNodeTemplateCache::IsTemplateOuter(ParentGraph);
 	bool const bUndo = !bIsTemplateNode;
+	bool const bPrintCommand = !bIsTemplateNode;
 
 	// First create a backing member for our node
 	UControlRigGraph* RigGraph = Cast<UControlRigGraph>(ParentGraph);
@@ -78,7 +79,7 @@ UEdGraphNode* UControlRigBranchNodeSpawner::Invoke(UEdGraph* ParentGraph, FBindi
 		Controller->OpenUndoBracket(FString::Printf(TEXT("Add '%s' Node"), *Name.ToString()));
 	}
 
-	if (URigVMBranchNode* ModelNode = Controller->AddBranchNode(Location, Name.ToString(), bUndo))
+	if (URigVMBranchNode* ModelNode = Controller->AddBranchNode(Location, Name.ToString(), bUndo, bPrintCommand))
 	{
 		NewNode = Cast<UControlRigGraphNode>(RigGraph->FindNodeForModelNodeName(ModelNode->GetFName()));
 
diff --git a/Engine/Plugins/Animation/ControlRig/Source/ControlRigEditor/Private/Graph/NodeSpawners/ControlRigEnumNodeSpawner.cpp b/Engine/Plugins/Animation/ControlRig/Source/ControlRigEditor/Private/Graph/NodeSpawners/ControlRigEnumNodeSpawner.cpp
--- a/Engine/Plugins/Animation/ControlRig/Source/ControlRigEditor/Private/Graph/NodeSpawners/ControlRigEnumNodeSpawner.cpp
+++ b/Engine/Plugins/Animation/ControlRig/Source/ControlRigEditor/Private/Graph/NodeSpawners/ControlRigEnumNodeSpawner.cpp
@@ -98,7 +98,7 @@ UEdGraphNode* UControlRigEnumNodeSpawner::Invoke(UEdGraph* ParentGraph, FBinding
 		Controller->OpenUndoBracket(FString::Printf(TEXT("Add '%s' Node"), *Name.ToString()));
 	}
 
-	if (URigVMEnumNode* ModelNode = Controller->AddEnumNode(*Enum->GetPathName(), Location, Name.ToString(), bUndo))
+	if (URigVMEnumNode* ModelNode = Controller->AddEnumNode(*Enum->GetPathName(), Location, Name.ToString(), bUndo, !bIsTemplateNode))
 	{
 		NewNode = Cast<UControlRigGraphNode>(RigGraph->FindNodeForModelNodeName(ModelNode->GetFName()));
 
diff --git a/Engine/Plugins/Animation/ControlRig/Source/ControlRigEditor/Private/Graph/NodeSpawners/ControlRigFunctionRefNodeSpawner.cpp b/Engine/Plugins/Animation/ControlRig/Source/ControlRigEditor/Private/Graph/NodeSpawners/ControlRigFunctionRefNodeSpawner.cpp
--- a/Engine/Plugins/Animation/ControlRig/Source/ControlRigEditor/Private/Graph/NodeSpawners/ControlRigFunctionRefNodeSpawner.cpp
+++ b/Engine/Plugins/Animation/ControlRig/Source/ControlRigEditor/Private/Graph/NodeSpawners/ControlRigFunctionRefNodeSpawner.cpp
@@ -245,7 +245,7 @@ UControlRigGraphNode* UControlRigFunctionRefNodeSpawner::SpawnNode(UEdGraph* Par
 			Controller->OpenUndoBracket(FString::Printf(TEXT("Add '%s' Node"), *Name.ToString()));
 		}
 
-		if (URigVMFunctionReferenceNode* ModelNode = Controller->AddFunctionReferenceNode(InFunction, Location, Name.ToString(), bUndo))
+		if (URigVMFunctionReferenceNode* ModelNode = Controller->AddFunctionReferenceNode(InFunction, Location, Name.ToString(), bUndo, !bIsTemplateNode))
 		{
 			NewNode = Cast<UControlRigGraphNode>(RigGraph->FindNodeForModelNodeName(ModelNode->GetFName()));
 			check(NewNode);
diff --git a/Engine/Plugins/Animation/ControlRig/Source/ControlRigEditor/Private/Graph/NodeSpawners/ControlRigIfNodeSpawner.cpp b/Engine/Plugins/Animation/ControlRig/Source/ControlRigEditor/Private/Graph/NodeSpawners/ControlRigIfNodeSpawner.cpp
--- a/Engine/Plugins/Animation/ControlRig/Source/ControlRigEditor/Private/Graph/NodeSpawners/ControlRigIfNodeSpawner.cpp
+++ b/Engine/Plugins/Animation/ControlRig/Source/ControlRigEditor/Private/Graph/NodeSpawners/ControlRigIfNodeSpawner.cpp
@@ -121,7 +121,7 @@ UEdGraphNode* UControlRigIfNodeSpawner::Invoke(UEdGraph* ParentGraph, FBindingSe
 		Controller->OpenUndoBracket(FString::Printf(TEXT("Add '%s' Node"), *Name.ToString()));
 	}
 
-	if (URigVMIfNode* ModelNode = Controller->AddIfNode(CPPType, CPPTypeObjectPath, Location, Name.ToString(), bUndo))
+	if (URigVMIfNode* ModelNode = Controller->AddIfNode(CPPType, CPPTypeObjectPath, Location, Name.ToString(), bUndo, !bIsTemplateNode))
 	{
 		NewNode = Cast<UControlRigGraphNode>(RigGraph->FindNodeForModelNodeName(ModelNode->GetFName()));
 
diff --git a/Engine/Plugins/Animation/ControlRig/Source/ControlRigEditor/Private/Graph/NodeSpawners/ControlRigPrototypeNodeSpawner.cpp b/Engine/Plugins/Animation/ControlRig/Source/ControlRigEditor/Private/Graph/NodeSpawners/ControlRigPrototypeNodeSpawner.cpp
--- a/Engine/Plugins/Animation/ControlRig/Source/ControlRigEditor/Private/Graph/NodeSpawners/ControlRigPrototypeNodeSpawner.cpp
+++ b/Engine/Plugins/Animation/ControlRig/Source/ControlRigEditor/Private/Graph/NodeSpawners/ControlRigPrototypeNodeSpawner.cpp
@@ -124,7 +124,7 @@ UControlRigGraphNode* UControlRigPrototypeNodeSpawner::SpawnNode(UEdGraph* Paren
 			Controller->OpenUndoBracket(FString::Printf(TEXT("Add '%s' Node"), *Name.ToString()));
 		}
 
-		if (URigVMPrototypeNode* ModelNode = Controller->AddPrototypeNode(InNotation, Location, Name.ToString(), bUndo))
+		if (URigVMPrototypeNode* ModelNode = Controller->AddPrototypeNode(InNotation, Location, Name.ToString(), bUndo, !bIsTemplateNode))
 		{
 			NewNode = Cast<UControlRigGraphNode>(RigGraph->FindNodeForModelNodeName(ModelNode->GetFName()));
 
diff --git a/Engine/Plugins/Animation/ControlRig/Source/ControlRigEditor/Private/Graph/NodeSpawners/ControlRigRerouteNodeSpawner.cpp b/Engine/Plugins/Animation/ControlRig/Source/ControlRigEditor/Private/Graph/NodeSpawners/ControlRigRerouteNodeSpawner.cpp
--- a/Engine/Plugins/Animation/ControlRig/Source/ControlRigEditor/Private/Graph/NodeSpawners/ControlRigRerouteNodeSpawner.cpp
+++ b/Engine/Plugins/Animation/ControlRig/Source/ControlRigEditor/Private/Graph/NodeSpawners/ControlRigRerouteNodeSpawner.cpp
@@ -121,7 +121,7 @@ UEdGraphNode* UControlRigRerouteNodeSpawner::Invoke(UEdGraph* ParentGraph, FBind
 			}
 		}
 
-		if (URigVMNode* ModelNode = Controller->AddRerouteNodeOnPin(PinPath, bIsInput, true, Location, FString(), true))
+		if (URigVMNode* ModelNode = Controller->AddRerouteNodeOnPin(PinPath, bIsInput, true, Location, FString(), true, true))
 		{
 			for (UEdGraphNode* Node : ParentGraph->Nodes)
 			{
diff --git a/Engine/Plugins/Animation/ControlRig/Source/ControlRigEditor/Private/Graph/NodeSpawners/ControlRigSelectNodeSpawner.cpp b/Engine/Plugins/Animation/ControlRig/Source/ControlRigEditor/Private/Graph/NodeSpawners/ControlRigSelectNodeSpawner.cpp
--- a/Engine/Plugins/Animation/ControlRig/Source/ControlRigEditor/Private/Graph/NodeSpawners/ControlRigSelectNodeSpawner.cpp
+++ b/Engine/Plugins/Animation/ControlRig/Source/ControlRigEditor/Private/Graph/NodeSpawners/ControlRigSelectNodeSpawner.cpp
@@ -122,7 +122,7 @@ UEdGraphNode* UControlRigSelectNodeSpawner::Invoke(UEdGraph* ParentGraph, FBindi
 		Controller->OpenUndoBracket(FString::Printf(TEXT("Add '%s' Node"), *Name.ToString()));
 	}
 
-	if (URigVMSelectNode* ModelNode = Controller->AddSelectNode(CPPType, CPPTypeObjectPath, Location, Name.ToString(), bUndo))
+	if (URigVMSelectNode* ModelNode = Controller->AddSelectNode(CPPType, CPPTypeObjectPath, Location, Name.ToString(), bUndo, !bIsTemplateNode))
 	{
 		NewNode = Cast<UControlRigGraphNode>(RigGraph->FindNodeForModelNodeName(ModelNode->GetFName()));
 
diff --git a/Engine/Plugins/Animation/ControlRig/Source/ControlRigEditor/Private/Graph/NodeSpawners/ControlRigUnitNodeSpawner.cpp b/Engine/Plugins/Animation/ControlRig/Source/ControlRigEditor/Private/Graph/NodeSpawners/ControlRigUnitNodeSpawner.cpp
--- a/Engine/Plugins/Animation/ControlRig/Source/ControlRigEditor/Private/Graph/NodeSpawners/ControlRigUnitNodeSpawner.cpp
+++ b/Engine/Plugins/Animation/ControlRig/Source/ControlRigEditor/Private/Graph/NodeSpawners/ControlRigUnitNodeSpawner.cpp
@@ -144,7 +144,7 @@ UControlRigGraphNode* UControlRigUnitNodeSpawner::SpawnNode(UEdGraph* ParentGrap
 		FRigVMUnitNodeCreatedContext& UnitNodeCreatedContext = Controller->GetUnitNodeCreatedContext();
 		FRigVMUnitNodeCreatedContext::FScope ReasonScope(UnitNodeCreatedContext, ERigVMNodeCreatedReason::NodeSpawner);
 
-		if (URigVMUnitNode* ModelNode = Controller->AddUnitNode(StructTemplate, FRigUnit::GetMethodName(), Location, Name.ToString(), bUndo))
+		if (URigVMUnitNode* ModelNode = Controller->AddUnitNode(StructTemplate, FRigUnit::GetMethodName(), Location, Name.ToString(), bUndo, !bIsTemplateNode))
 		{
 			NewNode = Cast<UControlRigGraphNode>(RigGraph->FindNodeForModelNodeName(ModelNode->GetFName()));
 			check(NewNode);
diff --git a/Engine/Plugins/Animation/ControlRig/Source/ControlRigEditor/Private/Graph/NodeSpawners/ControlRigVariableNodeSpawner.cpp b/Engine/Plugins/Animation/ControlRig/Source/ControlRigEditor/Private/Graph/NodeSpawners/ControlRigVariableNodeSpawner.cpp
--- a/Engine/Plugins/Animation/ControlRig/Source/ControlRigEditor/Private/Graph/NodeSpawners/ControlRigVariableNodeSpawner.cpp
+++ b/Engine/Plugins/Animation/ControlRig/Source/ControlRigEditor/Private/Graph/NodeSpawners/ControlRigVariableNodeSpawner.cpp
@@ -115,7 +115,7 @@ UEdGraphNode* UControlRigVariableNodeSpawner::Invoke(UEdGraph* ParentGraph, FBin
 		TypeName = FString::Printf(TEXT("TArray<%s>"), *TypeName);
 	}
 
-	if (URigVMNode* ModelNode = Controller->AddVariableNodeFromObjectPath(ExternalVariable.Name, TypeName, ObjectPath, bIsGetter, FString(), Location, FString(), !bIsTemplateNode))
+	if (URigVMNode* ModelNode = Controller->AddVariableNodeFromObjectPath(ExternalVariable.Name, TypeName, ObjectPath, bIsGetter, FString(), Location, FString(), !bIsTemplateNode, !bIsTemplateNode))
 	{
 		for (UEdGraphNode* Node : ParentGraph->Nodes)
 		{
diff --git a/Engine/Plugins/Animation/ControlRig/Source/ControlRigEditor/Private/Graph/SControlRigGraphNode.cpp b/Engine/Plugins/Animation/ControlRig/Source/ControlRigEditor/Private/Graph/SControlRigGraphNode.cpp
--- a/Engine/Plugins/Animation/ControlRig/Source/ControlRigEditor/Private/Graph/SControlRigGraphNode.cpp
+++ b/Engine/Plugins/Animation/ControlRig/Source/ControlRigEditor/Private/Graph/SControlRigGraphNode.cpp
@@ -1162,7 +1162,7 @@ FReply SControlRigGraphNode::OnExpanderArrowClicked(int32 InPinInfoIndex)
 			Controller->OpenUndoBracket(PinInfo.bExpanded ? TEXT("Collapsing Pin") : TEXT("Expanding Pin"));
 			for(const FString& PinPathToModify : PinPathsToModify)
 			{
-				Controller->SetPinExpansion(PinPathToModify, !PinInfo.bExpanded, true);
+				Controller->SetPinExpansion(PinPathToModify, !PinInfo.bExpanded, true, true);
 			}
 			Controller->CloseUndoBracket();
 			return FReply::Handled();
diff --git a/Engine/Plugins/Animation/ControlRig/Source/ControlRigEditor/Private/Graph/SControlRigGraphNodeComment.cpp b/Engine/Plugins/Animation/ControlRig/Source/ControlRigEditor/Private/Graph/SControlRigGraphNodeComment.cpp
--- a/Engine/Plugins/Animation/ControlRig/Source/ControlRigEditor/Private/Graph/SControlRigGraphNodeComment.cpp
+++ b/Engine/Plugins/Animation/ControlRig/Source/ControlRigEditor/Private/Graph/SControlRigGraphNodeComment.cpp
@@ -45,8 +45,8 @@ FReply SControlRigGraphNodeComment::OnMouseButtonUp(const FGeometry& MyGeometry,
 					if (URigVMController* Controller = Blueprint->GetController(Graph))
 					{
 						Controller->OpenUndoBracket(TEXT("Resize Comment Box"));
-						Controller->SetNodePositionByName(CommentNode->GetFName(), Position, true);
-						Controller->SetNodeSizeByName(CommentNode->GetFName(), Size, true);
+						Controller->SetNodePositionByName(CommentNode->GetFName(), Position, true, false, true);
+						Controller->SetNodeSizeByName(CommentNode->GetFName(), Size, true, false, true);
 						Controller->CloseUndoBracket();
 					}
 				}
@@ -95,7 +95,7 @@ void SControlRigGraphNodeComment::Tick(const FGeometry& AllottedGeometry, const
 				{
 					if (URigVMController* Controller = Blueprint->GetController(Graph))
 					{
-						Controller->SetCommentTextByName(CommentNode->GetFName(), CurrentCommentTitle, true);
+						Controller->SetCommentTextByName(CommentNode->GetFName(), CurrentCommentTitle, true, true);
 					}
 				}
 			}
diff --git a/Engine/Plugins/Animation/ControlRig/Source/ControlRigEditor/Private/Graph/SControlRigGraphPinNameList.cpp b/Engine/Plugins/Animation/ControlRig/Source/ControlRigEditor/Private/Graph/SControlRigGraphPinNameList.cpp
--- a/Engine/Plugins/Animation/ControlRig/Source/ControlRigEditor/Private/Graph/SControlRigGraphPinNameList.cpp
+++ b/Engine/Plugins/Animation/ControlRig/Source/ControlRigEditor/Private/Graph/SControlRigGraphPinNameList.cpp
@@ -226,7 +226,7 @@ FReply SControlRigGraphPinNameList::OnGetSelectedClicked()
 					{
 						if (UControlRigGraph* Graph = Cast<UControlRigGraph>(GetPinObj()->GetOwningNode()->GetGraph()))
 						{
-							Graph->GetController()->SetPinDefaultValue(ParentPin->GetPinPath(), *Result[0].Get());
+							Graph->GetController()->SetPinDefaultValue(ParentPin->GetPinPath(), *Result[0].Get(), true, true, false, true);
 							CurrentList = GetNameList();
 						}
 					}
diff --git a/Engine/Plugins/Animation/ControlRig/Source/ControlRigEditor/Private/Graph/SControlRigGraphPinVariableBinding.cpp b/Engine/Plugins/Animation/ControlRig/Source/ControlRigEditor/Private/Graph/SControlRigGraphPinVariableBinding.cpp
--- a/Engine/Plugins/Animation/ControlRig/Source/ControlRigEditor/Private/Graph/SControlRigGraphPinVariableBinding.cpp
+++ b/Engine/Plugins/Animation/ControlRig/Source/ControlRigEditor/Private/Graph/SControlRigGraphPinVariableBinding.cpp
@@ -109,25 +109,28 @@ FLinearColor SControlRigVariableBinding::GetBindingColor() const
 			}
 		}
 
-		URigVMGraph* Model = ModelPin->GetGraph();
-		if(Model == nullptr)
+		if (ModelPin)
 		{
-			return  FLinearColor::Red;
-		}
-
-		const TArray<FRigVMGraphVariableDescription>& LocalVariables =  Model->GetLocalVariables();
-		for(const FRigVMGraphVariableDescription& LocalVariable : LocalVariables)
-		{
-			const FRigVMExternalVariable ExternalVariable = LocalVariable.ToExternalVariable();
-			if(!ExternalVariable.IsValid(true))
+			URigVMGraph* Model = ModelPin->GetGraph();
+			if(Model == nullptr)
 			{
-				continue;
+				return  FLinearColor::Red;
 			}
 
-			if (ExternalVariable.Name == BoundVariable)
+			const TArray<FRigVMGraphVariableDescription>& LocalVariables =  Model->GetLocalVariables();
+			for(const FRigVMGraphVariableDescription& LocalVariable : LocalVariables)
 			{
-				const FEdGraphPinType PinType = UControlRig::GetPinTypeFromExternalVariable(ExternalVariable);
-				return Schema->GetPinTypeColor(PinType);
+				const FRigVMExternalVariable ExternalVariable = LocalVariable.ToExternalVariable();
+				if(!ExternalVariable.IsValid(true))
+				{
+					continue;
+				}
+
+				if (ExternalVariable.Name == BoundVariable)
+				{
+					const FEdGraphPinType PinType = UControlRig::GetPinTypeFromExternalVariable(ExternalVariable);
+					return Schema->GetPinTypeColor(PinType);
+				}
 			}
 		}
 	}
@@ -202,7 +205,7 @@ void SControlRigVariableBinding::OnAddBinding(FName InPropertyName, const TArray
 
 		if(ModelPin)
 		{
-			Blueprint->GetController(ModelPin->GetGraph())->BindPinToVariable(ModelPin->GetPinPath(), FString::Join(Parts, TEXT(".")), true /* undo */);
+			Blueprint->GetController(ModelPin->GetGraph())->BindPinToVariable(ModelPin->GetPinPath(), FString::Join(Parts, TEXT(".")), true /* undo */, true /* python */);
 		}
 		else if(FunctionReferenceNode && !InnerVariableName.IsNone())
 		{
@@ -223,7 +226,7 @@ void SControlRigVariableBinding::OnRemoveBinding(FName InPropertyName)
 	{
 		if(ModelPin)
 		{
-			Blueprint->GetController(ModelPin->GetGraph())->UnbindPinFromVariable(ModelPin->GetPinPath(), true /* undo */);
+			Blueprint->GetController(ModelPin->GetGraph())->UnbindPinFromVariable(ModelPin->GetPinPath(), true /* undo */, true /* python */);
 		}
 		else if(FunctionReferenceNode && !InnerVariableName.IsNone())
 		{
@@ -338,7 +341,7 @@ void SControlRigVariableBinding::HandleBindToLocalVariable(FRigVMGraphVariableDe
 		return;
 	}
 
-	Controller->BindPinToVariable(ModelPin->GetPinPath(), InLocalVariable.Name.ToString());
+	Controller->BindPinToVariable(ModelPin->GetPinPath(), InLocalVariable.Name.ToString(), true, true);
 }
 
 /////////////////////////////////////////////////////////////////////////////////////////////////////////////
diff --git a/Engine/Source/Developer/RigVMDeveloper/Private/RigVMModel/RigVMController.cpp b/Engine/Source/Developer/RigVMDeveloper/Private/RigVMModel/RigVMController.cpp
--- a/Engine/Source/Developer/RigVMDeveloper/Private/RigVMModel/RigVMController.cpp
+++ b/Engine/Source/Developer/RigVMDeveloper/Private/RigVMModel/RigVMController.cpp
@@ -13,7 +13,7 @@
 #include "UObject/Package.h"
 #include "Misc/CoreMisc.h"
 #include "Algo/Sort.h"
-#include "RigVMCore/RigVMPythonUtils.h"
+#include "RigVMPythonUtils.h"
 
 #if WITH_EDITOR
 #include "Exporters/Exporter.h"
@@ -382,12 +382,11 @@ TArray<FString> URigVMController::GeneratePythonCommands()
 					if(SourcePin)
 					{
 						// AddRerouteNodeOnLinkPath(const FString& InLinkPinPathRepresentation, bool bShowAsFullNode, const FVector2D& InPosition = FVector2D::ZeroVector, const FString& InNodeName = TEXT(""), bool bSetupUndoRedo = true);
-						Commands.Add(FString::Printf(TEXT("blueprint.get_controller_by_name('%s').add_reroute_node_on_link_path('%s', %s, unreal.Vector2D(%f, %f), '%s')"),
+						Commands.Add(FString::Printf(TEXT("blueprint.get_controller_by_name('%s').add_reroute_node_on_link_path('%s', %s, %s, '%s')"),
 								*GetGraph()->GetGraphName(),
 								*FString::Printf(TEXT("%s -> %s"), *SourcePin->GetPinPath(), *TargetPin->GetPinPath()),
 								Reroute->GetShowsAsFullNode() ? TEXT("True") : TEXT("False"),
-								Reroute->GetPosition().X, 
-								Reroute->GetPosition().Y,
+								*RigVMPythonUtils::Vector2DToPythonString(Reroute->GetPosition()),
 								*Reroute->GetName()));
 
 						ReroutesAdded.Add(Reroute->GetName());
@@ -422,11 +421,10 @@ TArray<FString> URigVMController::GetAddNodePythonCommands(URigVMNode* Node) con
 		else
 		{
 			// add_struct_node_from_struct_path(script_struct_path, method_name, position=[0.0, 0.0], node_name='', undo=True)
-			Commands.Add(FString::Printf(TEXT("blueprint.get_controller_by_name('%s').add_unit_node_from_struct_path('%s', 'Execute', unreal.Vector2D(%f, %f), '%s')"),
+			Commands.Add(FString::Printf(TEXT("blueprint.get_controller_by_name('%s').add_unit_node_from_struct_path('%s', 'Execute', %s, '%s')"),
 					*GetGraph()->GetGraphName(),
 					*UnitNode->GetScriptStruct()->GetPathName(),
-					UnitNode->GetPosition().X, 
-					UnitNode->GetPosition().Y,
+					*RigVMPythonUtils::Vector2DToPythonString(UnitNode->GetPosition()),
 					*UnitNode->GetName()));
 		}
 	}
@@ -435,27 +433,25 @@ TArray<FString> URigVMController::GetAddNodePythonCommands(URigVMNode* Node) con
 		// add_variable_node(variable_name, cpp_type, cpp_type_object, is_getter, default_value, position=[0.0, 0.0], node_name='', undo=True)
 		if (VariableNode->GetVariableDescription().CPPTypeObject)
 		{
-			Commands.Add(FString::Printf(TEXT("blueprint.get_controller_by_name('%s').add_variable_node_from_object_path('%s', '%s', '%s', %s, '%s', unreal.Vector2D(%f, %f), '%s')"),
+			Commands.Add(FString::Printf(TEXT("blueprint.get_controller_by_name('%s').add_variable_node_from_object_path('%s', '%s', '%s', %s, '%s', %s, '%s')"),
 					*GetGraph()->GetGraphName(),
 					*VariableNode->GetVariableName().ToString(),
 					*VariableNode->GetVariableDescription().CPPType,
 					*VariableNode->GetVariableDescription().CPPTypeObject->GetPathName(),
 					VariableNode->IsGetter() ? TEXT("True") : TEXT("False"),
 					*VariableNode->GetVariableDescription().DefaultValue,
-					VariableNode->GetPosition().X, 
-					VariableNode->GetPosition().Y,
+					*RigVMPythonUtils::Vector2DToPythonString(VariableNode->GetPosition()),
 					*VariableNode->GetName()));	
 		}
 		else
 		{
-			Commands.Add(FString::Printf(TEXT("blueprint.get_controller_by_name('%s').add_variable_node('%s', '%s', None, %s, '%s', unreal.Vector2D(%f, %f), '%s')"),
+			Commands.Add(FString::Printf(TEXT("blueprint.get_controller_by_name('%s').add_variable_node('%s', '%s', None, %s, '%s', %s, '%s')"),
 					*GetGraph()->GetGraphName(),
 					*VariableNode->GetVariableName().ToString(),
 					*VariableNode->GetVariableDescription().CPPType,
 					VariableNode->IsGetter() ? TEXT("True") : TEXT("False"),
 					*VariableNode->GetVariableDescription().DefaultValue,
-					VariableNode->GetPosition().X, 
-					VariableNode->GetPosition().Y,
+					*RigVMPythonUtils::Vector2DToPythonString(VariableNode->GetPosition()),
 					*VariableNode->GetName()));	
 		}
 	}
@@ -464,108 +460,109 @@ TArray<FString> URigVMController::GetAddNodePythonCommands(URigVMNode* Node) con
 		// add_parameter_node_from_object_path(parameter_name, cpp_type, cpp_type_object_path, is_input, default_value, position=[0.0, 0.0], node_name='', undo=True)
 		if (ParameterNode->GetParameterDescription().CPPTypeObject)
 		{
-			Commands.Add(FString::Printf(TEXT("blueprint.get_controller_by_name('%s').add_parameter_node_from_object_path('%s', '%s', '%s', %s, '%s', unreal.Vector2D(%f, %f), '%s')"),
+			Commands.Add(FString::Printf(TEXT("blueprint.get_controller_by_name('%s').add_parameter_node_from_object_path('%s', '%s', '%s', %s, '%s', %s, '%s')"),
 					*GetGraph()->GetGraphName(),
 					*ParameterNode->GetParameterName().ToString(),
 					*ParameterNode->GetParameterDescription().CPPType,
 					*ParameterNode->GetParameterDescription().CPPTypeObject->GetPathName(),
 					ParameterNode->IsInput() ? TEXT("True") : TEXT("False"),
 					*ParameterNode->GetParameterDescription().DefaultValue,
-					ParameterNode->GetPosition().X, 
-					ParameterNode->GetPosition().Y,
+					*RigVMPythonUtils::Vector2DToPythonString(ParameterNode->GetPosition()),
 					*ParameterNode->GetName()));	
 		}
 		else
 		{
-			Commands.Add(FString::Printf(TEXT("blueprint.get_controller_by_name('%s').add_parameter_node('%s', '%s', None, %s, str(%s), unreal.Vector2D(%f, %f), '%s')"),
+			Commands.Add(FString::Printf(TEXT("blueprint.get_controller_by_name('%s').add_parameter_node('%s', '%s', None, %s, str(%s), %s, '%s')"),
 					*GetGraph()->GetGraphName(),
 					*ParameterNode->GetParameterName().ToString(),
 					*ParameterNode->GetParameterDescription().CPPType,
 					ParameterNode->IsInput() ? TEXT("True") : TEXT("False"),
 					*ParameterNode->GetParameterDescription().DefaultValue,
-					ParameterNode->GetPosition().X, 
-					ParameterNode->GetPosition().Y,
+					*RigVMPythonUtils::Vector2DToPythonString(ParameterNode->GetPosition()),
 					*ParameterNode->GetName()));	
 		}
 	}
 	else if (URigVMCommentNode* CommentNode = Cast<URigVMCommentNode>(Node))
 	{
 		// add_comment_node(comment_text, position=[0.0, 0.0], size=[400.0, 300.0], color=[0.0, 0.0, 0.0, 0.0], node_name='', undo=True)
-		Commands.Add(FString::Printf(TEXT("blueprint.get_controller_by_name('%s').add_comment_node('%s', unreal.Vector2D(%f, %f), unreal.Vector2D(%f, %f), unreal.LinearColor(%f, %f, %f, %f), '%s')"),
+		Commands.Add(FString::Printf(TEXT("blueprint.get_controller_by_name('%s').add_comment_node('%s', %s, %s, %s, '%s')"),
 					*GetGraph()->GetGraphName(),
 					*CommentNode->GetCommentText(),
-					CommentNode->GetPosition().X, 
-					CommentNode->GetPosition().Y,
-					CommentNode->GetSize().X, 
-					CommentNode->GetSize().Y,
-					CommentNode->GetNodeColor().R, 
-					CommentNode->GetNodeColor().G,
-					CommentNode->GetNodeColor().B,
-					CommentNode->GetNodeColor().A,
+					*RigVMPythonUtils::Vector2DToPythonString(CommentNode->GetPosition()),
+					*RigVMPythonUtils::Vector2DToPythonString(CommentNode->GetSize()),
+					*RigVMPythonUtils::LinearColorToPythonString(CommentNode->GetNodeColor()),
 					*CommentNode->GetName()));	
 	}
 	else if (URigVMBranchNode* BranchNode = Cast<URigVMBranchNode>(Node))
 	{
 		// add_branch_node(position=[0.0, 0.0], node_name='', undo=True)
-		Commands.Add(FString::Printf(TEXT("blueprint.get_controller_by_name('%s').add_branch_node(unreal.Vector2D(%f, %f), '%s')"),
+		Commands.Add(FString::Printf(TEXT("blueprint.get_controller_by_name('%s').add_branch_node(%s, '%s')"),
 						*GetGraph()->GetGraphName(),
-						BranchNode->GetPosition().X, 
-						BranchNode->GetPosition().Y,
+						*RigVMPythonUtils::Vector2DToPythonString(BranchNode->GetPosition()),
 						*BranchNode->GetName()));	
 	}
 	else if (URigVMIfNode* IfNode = Cast<URigVMIfNode>(Node))
 	{
 		// add_if_node(cpp_type, cpp_type_object_path, position=[0.0, 0.0], node_name='', undo=True)
 		URigVMPin* ResultPin = IfNode->FindPin(URigVMIfNode::ResultName);
-		Commands.Add(FString::Printf(TEXT("blueprint.get_controller_by_name('%s').add_if_node('%s', '%s', unreal.Vector2D(%f, %f), '%s')"),
+		Commands.Add(FString::Printf(TEXT("blueprint.get_controller_by_name('%s').add_if_node('%s', '%s', %s, '%s')"),
 						*GetGraph()->GetGraphName(),
 						*ResultPin->GetCPPType(),
 						*ResultPin->CPPTypeObject->GetPathName(),
-						IfNode->GetPosition().X, 
-						IfNode->GetPosition().Y,
+						*RigVMPythonUtils::Vector2DToPythonString(IfNode->GetPosition()),
 						*IfNode->GetName()));
 	}
 	else if (URigVMSelectNode* SelectNode = Cast<URigVMSelectNode>(Node))
 	{
 		// add_select_node(cpp_type, cpp_type_object_path, position=[0.0, 0.0], node_name='', undo=True)
 		URigVMPin* ResultPin = SelectNode->FindPin(URigVMSelectNode::ResultName);
-		Commands.Add(FString::Printf(TEXT("blueprint.get_controller_by_name('%s').add_select_node('%s', '%s', unreal.Vector2D(%f, %f), '%s')"),
+		Commands.Add(FString::Printf(TEXT("blueprint.get_controller_by_name('%s').add_select_node('%s', '%s', %s, '%s')"),
 						*GetGraph()->GetGraphName(),
 						*ResultPin->GetCPPType(),
 						*ResultPin->CPPTypeObject->GetPathName(),
-						SelectNode->GetPosition().X, 
-						SelectNode->GetPosition().Y,
+						*RigVMPythonUtils::Vector2DToPythonString(SelectNode->GetPosition()),
 						*SelectNode->GetName()));
 	}
 	else if (URigVMPrototypeNode* PrototypeNode = Cast<URigVMPrototypeNode>(Node))
 	{
 		// add_prototype_node(notation, position=[0.0, 0.0], node_name='', undo=True)
-		Commands.Add(FString::Printf(TEXT("blueprint.get_controller_by_name('%s').add_prototype_node('%s', unreal.Vector2D(%f, %f), '%s')"),
+		Commands.Add(FString::Printf(TEXT("blueprint.get_controller_by_name('%s').add_prototype_node('%s', %s, '%s')"),
 						*GetGraph()->GetGraphName(),
 						*PrototypeNode->GetNotation().ToString(),
-						PrototypeNode->GetPosition().X, 
-						PrototypeNode->GetPosition().Y,
+						*RigVMPythonUtils::Vector2DToPythonString(PrototypeNode->GetPosition()),
 						*PrototypeNode->GetName()));
 	}
 	else if (URigVMEnumNode* EnumNode = Cast<URigVMEnumNode>(Node))
 	{
 		// add_enum_node(cpp_type_object_path, position=[0.0, 0.0], node_name='', undo=True)
-		Commands.Add(FString::Printf(TEXT("blueprint.get_controller_by_name('%s').add_enum_node('%s', unreal.Vector2D(%f, %f), '%s')"),
+		Commands.Add(FString::Printf(TEXT("blueprint.get_controller_by_name('%s').add_enum_node('%s', %s, '%s')"),
 						*GetGraph()->GetGraphName(),
 						*EnumNode->GetCPPTypeObject()->GetPathName(),
-						EnumNode->GetPosition().X, 
-						EnumNode->GetPosition().Y,
+						*RigVMPythonUtils::Vector2DToPythonString(EnumNode->GetPosition()),
 						*EnumNode->GetName()));
 	}
 	else if (URigVMLibraryNode* LibraryNode = Cast<URigVMLibraryNode>(Node))
 	{
 		// AddFunctionReferenceNode(URigVMLibraryNode* InFunctionDefinition, const FVector2D& InNodePosition = FVector2D::ZeroVector, const FString& InNodeName = TEXT(""), bool bSetupUndoRedo = true);
-		Commands.Add(FString::Printf(TEXT("blueprint.get_controller_by_name('%s').add_function_reference_node(function_%s, unreal.Vector2D(%f, %f), '%s')"),
+		URigVMFunctionLibrary* Library = LibraryNode->GetLibrary();
+		if (!Library || Library == GetGraph()->GetDefaultFunctionLibrary())
+		{
+			Commands.Add(FString::Printf(TEXT("blueprint.get_controller_by_name('%s').add_function_reference_node(function_%s, %s, '%s')"),
 						*GetGraph()->GetGraphName(),
 						*RigVMPythonUtils::NameToPep8(LibraryNode->GetContainedGraph()->GetGraphName()),
-						LibraryNode->GetPosition().X, 
-						LibraryNode->GetPosition().Y,
+						*RigVMPythonUtils::Vector2DToPythonString(LibraryNode->GetPosition()), 
+						*LibraryNode->GetName()));
+		}
+		else
+		{
+			Commands.Add(FString::Printf(TEXT("function_blueprint = unreal.load_object(name = '%s', outer = None)"),
+				*Library->GetOuter()->GetPathName()));
+			Commands.Add(FString::Printf(TEXT("blueprint.get_controller_by_name('%s').add_function_reference_node(function_blueprint.get_local_function_library().find_function('%s'), %s, '%s')"),
+						*GetGraph()->GetGraphName(),
+						*LibraryNode->GetName(),
+						*RigVMPythonUtils::Vector2DToPythonString(LibraryNode->GetPosition()), 
 						*LibraryNode->GetName()));
+		}
 
 		if (Node->IsA<URigVMCollapseNode>())
 		{
@@ -599,8 +596,8 @@ TArray<FString> URigVMController::GetAddNodePythonCommands(URigVMNode* Node) con
 				continue;
 			}
 			
-			const FString DefaultValue =Pin->GetDefaultValue();
-			if (!DefaultValue.IsEmpty())
+			const FString DefaultValue = Pin->GetDefaultValue();
+			if (!DefaultValue.IsEmpty() && DefaultValue != TEXT("()"))
 			{
 				Commands.Add(FString::Printf(TEXT("blueprint.get_controller_by_name('%s').set_pin_default_value('%s', '%s')"),
 							*GetGraph()->GetGraphName(),
@@ -615,7 +612,7 @@ TArray<FString> URigVMController::GetAddNodePythonCommands(URigVMNode* Node) con
 
 #if WITH_EDITOR
 
-URigVMUnitNode* URigVMController::AddUnitNode(UScriptStruct* InScriptStruct, const FName& InMethodName, const FVector2D& InPosition, const FString& InNodeName, bool bSetupUndoRedo)
+URigVMUnitNode* URigVMController::AddUnitNode(UScriptStruct* InScriptStruct, const FName& InMethodName, const FVector2D& InPosition, const FString& InNodeName, bool bSetupUndoRedo, bool bPrintPythonCommand)
 {
 	if(!IsValidGraph())
 	{
@@ -723,10 +720,19 @@ URigVMUnitNode* URigVMController::AddUnitNode(UScriptStruct* InScriptStruct, con
 		ActionStack->EndAction(Action);
 	}
 
+	if (bPrintPythonCommand)
+	{
+		TArray<FString> Commands = GetAddNodePythonCommands(Node);
+		for (const FString& Command : Commands)
+		{
+			RigVMPythonUtils::Print(FString::Printf(TEXT("%s"), *Command));
+		}
+	}
+
 	return Node;
 }
 
-URigVMUnitNode* URigVMController::AddUnitNodeFromStructPath(const FString& InScriptStructPath, const FName& InMethodName, const FVector2D& InPosition, const FString& InNodeName, bool bSetupUndoRedo)
+URigVMUnitNode* URigVMController::AddUnitNodeFromStructPath(const FString& InScriptStructPath, const FName& InMethodName, const FVector2D& InPosition, const FString& InNodeName, bool bSetupUndoRedo, bool bPrintPythonCommand)
 {
 	if(!IsValidGraph())
 	{
@@ -740,10 +746,10 @@ URigVMUnitNode* URigVMController::AddUnitNodeFromStructPath(const FString& InScr
 		return nullptr;
 	}
 
-	return AddUnitNode(ScriptStruct, InMethodName, InPosition, InNodeName, bSetupUndoRedo);
+	return AddUnitNode(ScriptStruct, InMethodName, InPosition, InNodeName, bSetupUndoRedo, bPrintPythonCommand);
 }
 
-URigVMVariableNode* URigVMController::AddVariableNode(const FName& InVariableName, const FString& InCPPType, UObject* InCPPTypeObject, bool bIsGetter, const FString& InDefaultValue, const FVector2D& InPosition, const FString& InNodeName, bool bSetupUndoRedo)
+URigVMVariableNode* URigVMController::AddVariableNode(const FName& InVariableName, const FString& InCPPType, UObject* InCPPTypeObject, bool bIsGetter, const FString& InDefaultValue, const FVector2D& InPosition, const FString& InNodeName, bool bSetupUndoRedo, bool bPrintPythonCommand)
 {
 	if (!IsValidGraph())
 	{
@@ -900,10 +906,19 @@ URigVMVariableNode* URigVMController::AddVariableNode(const FName& InVariableNam
 		ActionStack->EndAction(Action);
 	}
 
+	if (bPrintPythonCommand)
+	{
+		TArray<FString> Commands = GetAddNodePythonCommands(Node);
+		for (const FString& Command : Commands)
+		{
+			RigVMPythonUtils::Print(FString::Printf(TEXT("%s"), *Command));
+		}
+	}
+
 	return Node;
 }
 
-URigVMVariableNode* URigVMController::AddVariableNodeFromObjectPath(const FName& InVariableName, const FString& InCPPType, const FString& InCPPTypeObjectPath, bool bIsGetter, const FString& InDefaultValue, const FVector2D& InPosition, const FString& InNodeName, bool bSetupUndoRedo)
+URigVMVariableNode* URigVMController::AddVariableNodeFromObjectPath(const FName& InVariableName, const FString& InCPPType, const FString& InCPPTypeObjectPath, bool bIsGetter, const FString& InDefaultValue, const FVector2D& InPosition, const FString& InNodeName, bool bSetupUndoRedo, bool bPrintPythonCommand)
 {
 	if (!IsValidGraph())
 	{
@@ -921,7 +936,7 @@ URigVMVariableNode* URigVMController::AddVariableNodeFromObjectPath(const FName&
 		}
 	}
 
-	return AddVariableNode(InVariableName, InCPPType, CPPTypeObject, bIsGetter, InDefaultValue, InPosition, InNodeName, bSetupUndoRedo);
+	return AddVariableNode(InVariableName, InCPPType, CPPTypeObject, bIsGetter, InDefaultValue, InPosition, InNodeName, bSetupUndoRedo, bPrintPythonCommand);
 }
 
 void URigVMController::RefreshVariableNode(const FName& InNodeName, const FName& InVariableName, const FString& InCPPType, UObject* InCPPTypeObject, bool bSetupUndoRedo, bool bSetupOrphanPins)
@@ -990,7 +1005,7 @@ void URigVMController::RefreshVariableNode(const FName& InNodeName, const FName&
 	}
 }
 
-void URigVMController::OnExternalVariableRemoved(const FName& InVarName, bool bSetupUndoRedo)
+void URigVMController::OnExternalVariableRemoved(const FName& InVarName, bool bSetupUndoRedo, bool bPrintPythonCommand)
 {
 	if (!IsValidGraph())
 	{
@@ -1075,9 +1090,16 @@ void URigVMController::OnExternalVariableRemoved(const FName& InVarName, bool bS
 	{
 		CloseUndoBracket();
 	}
+
+	if (bPrintPythonCommand)
+	{
+		RigVMPythonUtils::Print(FString::Printf(TEXT("blueprint.get_controller_by_name('%s').on_external_variable_removed('%s')"),
+				*GetGraph()->GetGraphName(),
+				*InVarName.ToString()));
+	}
 }
 
-void URigVMController::OnExternalVariableRenamed(const FName& InOldVarName, const FName& InNewVarName, bool bSetupUndoRedo)
+void URigVMController::OnExternalVariableRenamed(const FName& InOldVarName, const FName& InNewVarName, bool bSetupUndoRedo, bool bPrintPythonCommand)
 {
 	if (!IsValidGraph())
 	{
@@ -1165,9 +1187,17 @@ void URigVMController::OnExternalVariableRenamed(const FName& InOldVarName, cons
 	{
 		CloseUndoBracket();
 	}
+
+	if (bPrintPythonCommand)
+	{
+		RigVMPythonUtils::Print(FString::Printf(TEXT("blueprint.get_controller_by_name('%s').on_external_variable_renamed('%s', '%s')"),
+				*GetGraph()->GetGraphName(),
+				*InOldVarName.ToString(),
+				*InNewVarName.ToString()));
+	}
 }
 
-void URigVMController::OnExternalVariableTypeChanged(const FName& InVarName, const FString& InCPPType, UObject* InCPPTypeObject, bool bSetupUndoRedo)
+void URigVMController::OnExternalVariableTypeChanged(const FName& InVarName, const FString& InCPPType, UObject* InCPPTypeObject, bool bSetupUndoRedo, bool bPrintPythonCommand)
 {
 	if (!IsValidGraph())
 	{
@@ -1253,6 +1283,37 @@ void URigVMController::OnExternalVariableTypeChanged(const FName& InVarName, con
 	{
 		CloseUndoBracket();
 	}
+	
+	if (bPrintPythonCommand)
+	{
+		//OnExternalVariableTypeChanged(const FName& InVarName, const FString& InCPPType, UObject* InCPPTypeObject
+		RigVMPythonUtils::Print(FString::Printf(TEXT("blueprint.get_controller_by_name('%s').on_external_variable_type_changed_from_object_path('%s', '%s', '%s')"),
+				*GetGraph()->GetGraphName(),
+				*InVarName.ToString(),
+				*InCPPType,
+				*InCPPTypeObject->GetPathName()));
+	}
+}
+
+void URigVMController::OnExternalVariableTypeChangedFromObjectPath(const FName& InVarName, const FString& InCPPType, const FString& InCPPTypeObjectPath, bool bSetupUndoRedo)
+{
+	if (!IsValidGraph())
+	{
+		return;
+	}
+
+	UObject* CPPTypeObject = nullptr;
+	if (!InCPPTypeObjectPath.IsEmpty())
+	{
+		CPPTypeObject = URigVMPin::FindObjectFromCPPTypeObjectPath<UObject>(InCPPTypeObjectPath);
+		if (CPPTypeObject == nullptr)
+		{
+			ReportErrorf(TEXT("Cannot find cpp type object for path '%s'."), *InCPPTypeObjectPath);
+			return;
+		}
+	}
+
+	OnExternalVariableTypeChanged(InVarName, InCPPType, CPPTypeObject, bSetupUndoRedo);
 }
 
 URigVMVariableNode* URigVMController::ReplaceParameterNodeWithVariable(const FName& InNodeName, const FName& InVariableName, const FString& InCPPType, UObject* InCPPTypeObject, bool bSetupUndoRedo)
@@ -1302,7 +1363,7 @@ URigVMVariableNode* URigVMController::ReplaceParameterNodeWithVariable(const FNa
 	return nullptr;
 }
 
-URigVMParameterNode* URigVMController::AddParameterNode(const FName& InParameterName, const FString& InCPPType, UObject* InCPPTypeObject, bool bIsInput, const FString& InDefaultValue, const FVector2D& InPosition, const FString& InNodeName, bool bSetupUndoRedo)
+URigVMParameterNode* URigVMController::AddParameterNode(const FName& InParameterName, const FString& InCPPType, UObject* InCPPTypeObject, bool bIsInput, const FString& InDefaultValue, const FVector2D& InPosition, const FString& InNodeName, bool bSetupUndoRedo, bool bPrintPythonCommand)
 {
 	if (!IsValidGraph())
 	{
@@ -1464,10 +1525,19 @@ URigVMParameterNode* URigVMController::AddParameterNode(const FName& InParameter
 		ActionStack->EndAction(Action);
 	}
 
+	if (bPrintPythonCommand)
+	{
+		TArray<FString> Commands = GetAddNodePythonCommands(Node);
+		for (const FString& Command : Commands)
+		{
+			RigVMPythonUtils::Print(FString::Printf(TEXT("%s"), *Command));
+		}
+	}
+
 	return Node;
 }
 
-URigVMParameterNode* URigVMController::AddParameterNodeFromObjectPath(const FName& InParameterName, const FString& InCPPType, const FString& InCPPTypeObjectPath, bool bIsInput, const FString& InDefaultValue, const FVector2D& InPosition, const FString& InNodeName, bool bSetupUndoRedo)
+URigVMParameterNode* URigVMController::AddParameterNodeFromObjectPath(const FName& InParameterName, const FString& InCPPType, const FString& InCPPTypeObjectPath, bool bIsInput, const FString& InDefaultValue, const FVector2D& InPosition, const FString& InNodeName, bool bSetupUndoRedo, bool bPrintPythonCommand)
 {
 	if (!IsValidGraph())
 	{
@@ -1485,10 +1555,10 @@ URigVMParameterNode* URigVMController::AddParameterNodeFromObjectPath(const FNam
 		}
 	}
 
-	return AddParameterNode(InParameterName, InCPPType, CPPTypeObject, bIsInput, InDefaultValue, InPosition, InNodeName, bSetupUndoRedo);
+	return AddParameterNode(InParameterName, InCPPType, CPPTypeObject, bIsInput, InDefaultValue, InPosition, InNodeName, bSetupUndoRedo, bPrintPythonCommand);
 }
 
-URigVMCommentNode* URigVMController::AddCommentNode(const FString& InCommentText, const FVector2D& InPosition, const FVector2D& InSize, const FLinearColor& InColor, const FString& InNodeName, bool bSetupUndoRedo)
+URigVMCommentNode* URigVMController::AddCommentNode(const FString& InCommentText, const FVector2D& InPosition, const FVector2D& InSize, const FLinearColor& InColor, const FString& InNodeName, bool bSetupUndoRedo, bool bPrintPythonCommand)
 {
 	if (!IsValidGraph())
 	{
@@ -1532,10 +1602,19 @@ URigVMCommentNode* URigVMController::AddCommentNode(const FString& InCommentText
 		ActionStack->EndAction(Action);
 	}
 
+	if (bPrintPythonCommand)
+	{
+		TArray<FString> Commands = GetAddNodePythonCommands(Node);
+		for (const FString& Command : Commands)
+		{
+			RigVMPythonUtils::Print(FString::Printf(TEXT("%s"), *Command));
+		}
+	}
+
 	return Node;
 }
 
-URigVMRerouteNode* URigVMController::AddRerouteNodeOnLink(URigVMLink* InLink, bool bShowAsFullNode, const FVector2D& InPosition, const FString& InNodeName, bool bSetupUndoRedo)
+URigVMRerouteNode* URigVMController::AddRerouteNodeOnLink(URigVMLink* InLink, bool bShowAsFullNode, const FVector2D& InPosition, const FString& InNodeName, bool bSetupUndoRedo, bool bPrintPythonCommand)
 {
 	if(!IsValidLinkForGraph(InLink))
 	{
@@ -1578,10 +1657,21 @@ URigVMRerouteNode* URigVMController::AddRerouteNodeOnLink(URigVMLink* InLink, bo
 		ActionStack->EndAction(Action);
 	}
 
+	if (bPrintPythonCommand)
+	{
+		RigVMPythonUtils::Print(FString::Printf(TEXT("blueprint.get_controller_by_name('%s').add_reroute_node_on_link_path('%s', %s, %s, '%s')"),
+				*GetGraph()->GetGraphName(),
+				*InLink->GetPinPathRepresentation(),
+				(bShowAsFullNode) ? TEXT("True") : TEXT("False"),
+				*RigVMPythonUtils::Vector2DToPythonString(Node->GetPosition()),
+				*Node->GetName()));
+	}
+
 	return Node;
 }
 
-URigVMRerouteNode* URigVMController::AddRerouteNodeOnLinkPath(const FString& InLinkPinPathRepresentation, bool bShowAsFullNode, const FVector2D& InPosition, const FString& InNodeName, bool bSetupUndoRedo)
+URigVMRerouteNode* URigVMController::AddRerouteNodeOnLinkPath(const FString& InLinkPinPathRepresentation, bool bShowAsFullNode, const FVector2D& InPosition, const FString&
+                                                              InNodeName, bool bSetupUndoRedo, bool bPrintPythonCommand)
 {
 	if (!IsValidGraph())
 	{
@@ -1592,10 +1682,10 @@ URigVMRerouteNode* URigVMController::AddRerouteNodeOnLinkPath(const FString& InL
 	check(Graph);
 
 	URigVMLink* Link = Graph->FindLink(InLinkPinPathRepresentation);
-	return AddRerouteNodeOnLink(Link, bShowAsFullNode, InPosition, InNodeName, bSetupUndoRedo);
+	return AddRerouteNodeOnLink(Link, bShowAsFullNode, InPosition, InNodeName, bSetupUndoRedo, bPrintPythonCommand);
 }
 
-URigVMRerouteNode* URigVMController::AddRerouteNodeOnPin(const FString& InPinPath, bool bAsInput, bool bShowAsFullNode, const FVector2D& InPosition, const FString& InNodeName, bool bSetupUndoRedo)
+URigVMRerouteNode* URigVMController::AddRerouteNodeOnPin(const FString& InPinPath, bool bAsInput, bool bShowAsFullNode, const FVector2D& InPosition, const FString& InNodeName, bool bSetupUndoRedo, bool bPrintPythonCommand)
 {
 	if (!IsValidGraph())
 	{
@@ -1682,10 +1772,22 @@ URigVMRerouteNode* URigVMController::AddRerouteNodeOnPin(const FString& InPinPat
 		ActionStack->EndAction(Action);
 	}
 
+	if (bPrintPythonCommand)
+	{
+		// AddRerouteNodeOnPin(const FString& InPinPath, bool bAsInput, bool bShowAsFullNode, const FVector2D& InPosition, const FString& InNodeName, bool bSetupUndoRedo, bool bPrintPythonCommand)
+		RigVMPythonUtils::Print(FString::Printf(TEXT("blueprint.get_controller_by_name('%s').add_reroute_node_on_pin('%s', %s, %s, %s '%s')"),
+				*GetGraph()->GetGraphName(),
+				*InPinPath,
+				(bAsInput) ? TEXT("True") : TEXT("False"),
+				(bShowAsFullNode) ? TEXT("True") : TEXT("False"),
+				*RigVMPythonUtils::Vector2DToPythonString(Node->GetPosition()),
+				*Node->GetName()));
+	}
+
 	return Node;
 }
 
-URigVMInjectionInfo* URigVMController::AddInjectedNode(const FString& InPinPath, bool bAsInput, UScriptStruct* InScriptStruct, const FName& InMethodName, const FName& InInputPinName, const FName& InOutputPinName, const FString& InNodeName, bool bSetupUndoRedo)
+URigVMInjectionInfo* URigVMController::AddInjectedNode(const FString& InPinPath, bool bAsInput, UScriptStruct* InScriptStruct, const FName& InMethodName, const FName& InInputPinName, const FName& InOutputPinName, const FString& InNodeName, bool bSetupUndoRedo, bool bPrintPythonCommand)
 {
 	if (!IsValidGraph())
 	{
@@ -1877,6 +1979,19 @@ URigVMInjectionInfo* URigVMController::AddInjectedNode(const FString& InPinPath,
 		ActionStack->EndAction(Action);
 	}
 
+	if (bPrintPythonCommand)
+	{
+		RigVMPythonUtils::Print(FString::Printf(TEXT("blueprint.get_controller_by_name('%s').add_injected_node_from_struct_path('%s', %s, '%s', '%s', '%s', '%s', '%s')"),
+				*GetGraph()->GetGraphName(),
+				*InPinPath,
+				(bAsInput) ? TEXT("True") : TEXT("False"),
+				*InScriptStruct->GetPathName(),
+				*InMethodName.ToString(),
+				*InInputPinName.ToString(),
+				*InOutputPinName.ToString(),
+				*InNodeName));
+	}
+
 	return InjectionInfo;
 
 }
@@ -1898,7 +2013,7 @@ URigVMInjectionInfo* URigVMController::AddInjectedNodeFromStructPath(const FStri
 	return AddInjectedNode(InPinPath, bAsInput, ScriptStruct, InMethodName, InInputPinName, InOutputPinName, InNodeName, bSetupUndoRedo);
 }
 
-URigVMNode* URigVMController::EjectNodeFromPin(const FString& InPinPath, bool bSetupUndoRedo)
+URigVMNode* URigVMController::EjectNodeFromPin(const FString& InPinPath, bool bSetupUndoRedo, bool bPrintPythonCommand)
 {
 	if (!IsValidGraph())
 	{
@@ -2019,6 +2134,13 @@ URigVMNode* URigVMController::EjectNodeFromPin(const FString& InPinPath, bool bS
 		ActionStack->EndAction(Action);
 	}
 
+	if (bPrintPythonCommand)
+	{
+		RigVMPythonUtils::Print(FString::Printf(TEXT("blueprint.get_controller_by_name('%s').eject_node_from_pin('%s')"),
+				*GetGraph()->GetGraphName(),
+				*InPinPath));
+	}
+
 	return EjectedNode;
 }
 
@@ -2383,7 +2505,8 @@ TArray<FName> URigVMController::ImportNodesFromText(const FString& InText, bool
 URigVMLibraryNode* URigVMController::LocalizeFunction(
 	URigVMLibraryNode* InFunctionDefinition,
 	bool bLocalizeDependentPrivateFunctions,
-	bool bSetupUndoRedo)
+	bool bSetupUndoRedo,
+	bool bPrintPythonCommand)
 {
 	if(!IsValidGraph())
 	{
@@ -2398,7 +2521,7 @@ URigVMLibraryNode* URigVMController::LocalizeFunction(
 	TArray<URigVMLibraryNode*> FunctionsToLocalize;
 	FunctionsToLocalize.Add(InFunctionDefinition);
 
-	TMap<URigVMLibraryNode*, URigVMLibraryNode*> Results = LocalizeFunctions(FunctionsToLocalize, bLocalizeDependentPrivateFunctions, bSetupUndoRedo);
+	TMap<URigVMLibraryNode*, URigVMLibraryNode*> Results = LocalizeFunctions(FunctionsToLocalize, bLocalizeDependentPrivateFunctions, bSetupUndoRedo, bPrintPythonCommand);
 
 	URigVMLibraryNode** LocalizedFunctionPtr = Results.Find(FunctionsToLocalize[0]);
 	if(LocalizedFunctionPtr)
@@ -2409,9 +2532,9 @@ URigVMLibraryNode* URigVMController::LocalizeFunction(
 }
 
 TMap<URigVMLibraryNode*, URigVMLibraryNode*> URigVMController::LocalizeFunctions(
-    TArray<URigVMLibraryNode*> InFunctionDefinitions,
-    bool bLocalizeDependentPrivateFunctions,
-    bool bSetupUndoRedo)
+	TArray<URigVMLibraryNode*> InFunctionDefinitions,
+	bool bLocalizeDependentPrivateFunctions,
+	bool bSetupUndoRedo, bool bPrintPythonCommand)
 {
 	TMap<URigVMLibraryNode*, URigVMLibraryNode*> LocalizedFunctions;
 
@@ -2585,6 +2708,27 @@ TMap<URigVMLibraryNode*, URigVMLibraryNode*> URigVMController::LocalizeFunctions
 		CloseUndoBracket();
 	}
 
+	if (bPrintPythonCommand)
+	{
+		FString FunctionNames = TEXT("[");
+		for (auto It = InFunctionDefinitions.CreateConstIterator(); It; ++It)
+		{
+			FunctionNames += FString::Printf(TEXT("unreal.load_object(name = '%s', outer = None).get_local_function_library().find_function('%s')"),
+				*(*It)->GetLibrary()->GetOuter()->GetPathName(),
+				*(*It)->GetName());
+			if (It.GetIndex() < InFunctionDefinitions.Num() - 1)
+			{
+				FunctionNames += TEXT(", ");
+			}
+		}
+		FunctionNames += TEXT("]");
+		
+		RigVMPythonUtils::Print(FString::Printf(TEXT("blueprint.get_controller_by_name('%s').localize_functions(%s, %s)"),
+				*GetGraph()->GetGraphName(),
+				*FunctionNames,
+				(bLocalizeDependentPrivateFunctions) ? TEXT("True") : TEXT("False")));
+	}
+
 	return LocalizedFunctions;
 }
 
@@ -2601,7 +2745,7 @@ FName URigVMController::GetUniqueName(const FName& InName, TFunction<bool(const
 	return *Name;
 }
 
-URigVMCollapseNode* URigVMController::CollapseNodes(const TArray<FName>& InNodeNames, const FString& InCollapseNodeName, bool bSetupUndoRedo)
+URigVMCollapseNode* URigVMController::CollapseNodes(const TArray<FName>& InNodeNames, const FString& InCollapseNodeName, bool bSetupUndoRedo, bool bPrintPythonCommand)
 {
 	if (!IsValidGraph())
 	{
@@ -2623,10 +2767,30 @@ URigVMCollapseNode* URigVMController::CollapseNodes(const TArray<FName>& InNodeN
 		Nodes.AddUnique(Node);
 	}
 
-	return CollapseNodes(Nodes, InCollapseNodeName, bSetupUndoRedo);
+	URigVMCollapseNode* Node = CollapseNodes(Nodes, InCollapseNodeName, bSetupUndoRedo);
+	if (Node && bPrintPythonCommand)
+	{
+		FString ArrayStr = TEXT("[");
+		for (auto It = InNodeNames.CreateConstIterator(); It; ++It)
+		{
+			ArrayStr += TEXT("'") + It->ToString() + TEXT("'");
+			if (It.GetIndex() < InNodeNames.Num() - 1)
+			{
+				ArrayStr += TEXT(", ");
+			}
+		}
+		ArrayStr += TEXT("]");
+		
+		RigVMPythonUtils::Print(FString::Printf(TEXT("blueprint.get_controller_by_name('%s').collapse_nodes(%s, '%s')"),
+				*GetGraph()->GetGraphName(),
+				*ArrayStr,
+				*InCollapseNodeName));
+	}
+
+	return Node;
 }
 
-TArray<URigVMNode*> URigVMController::ExpandLibraryNode(const FName& InNodeName, bool bSetupUndoRedo)
+TArray<URigVMNode*> URigVMController::ExpandLibraryNode(const FName& InNodeName, bool bSetupUndoRedo, bool bPrintPythonCommand)
 {
 	if (!IsValidGraph())
 	{
@@ -2650,7 +2814,16 @@ TArray<URigVMNode*> URigVMController::ExpandLibraryNode(const FName& InNodeName,
 		return TArray<URigVMNode*>();
 	}
 
-	return ExpandLibraryNode(LibNode, bSetupUndoRedo);
+	TArray<URigVMNode*> Nodes = ExpandLibraryNode(LibNode, bSetupUndoRedo);
+
+	if (!Nodes.IsEmpty() && bPrintPythonCommand)
+	{
+		RigVMPythonUtils::Print(FString::Printf(TEXT("blueprint.get_controller_by_name('%s').expand_library_node('%s')"),
+				*GetGraph()->GetGraphName(),
+				*InNodeName.ToString()));
+	}
+
+	return Nodes;
 }
 
 #endif
@@ -3758,7 +3931,7 @@ TArray<URigVMNode*> URigVMController::ExpandLibraryNode(URigVMLibraryNode* InNod
 	return ExpandedNodes;
 }
 
-FName URigVMController::PromoteCollapseNodeToFunctionReferenceNode(const FName& InNodeName, bool bSetupUndoRedo)
+FName URigVMController::PromoteCollapseNodeToFunctionReferenceNode(const FName& InNodeName, bool bSetupUndoRedo, bool bPrintPythonCommand)
 {
 	if (!IsValidGraph())
 	{
@@ -3771,12 +3944,19 @@ FName URigVMController::PromoteCollapseNodeToFunctionReferenceNode(const FName&
 	URigVMNode* Result = PromoteCollapseNodeToFunctionReferenceNode(Cast<URigVMCollapseNode>(Graph->FindNodeByName(InNodeName)), bSetupUndoRedo);
 	if (Result)
 	{
+		if (bPrintPythonCommand)
+		{
+			RigVMPythonUtils::Print(FString::Printf(TEXT("blueprint.get_controller_by_name('%s').promote_collapse_node_to_function_reference_node('%s')"),
+					*GetGraph()->GetGraphName(),
+					*InNodeName.ToString()));
+		}
+		
 		return Result->GetFName();
 	}
 	return NAME_None;
 }
 
-FName URigVMController::PromoteFunctionReferenceNodeToCollapseNode(const FName& InNodeName, bool bSetupUndoRedo)
+FName URigVMController::PromoteFunctionReferenceNodeToCollapseNode(const FName& InNodeName, bool bSetupUndoRedo, bool bPrintPythonCommand)
 {
 	if (!IsValidGraph())
 	{
@@ -3974,7 +4154,7 @@ void URigVMController::RefreshFunctionPins(URigVMNode* InNode, bool bNotify)
 	}
 }
 
-bool URigVMController::RemoveNode(URigVMNode* InNode, bool bSetupUndoRedo, bool bRecursive)
+bool URigVMController::RemoveNode(URigVMNode* InNode, bool bSetupUndoRedo, bool bRecursive, bool bPrintPythonCommand)
 {
 	if (!IsValidNodeForGraph(InNode))
 	{
@@ -4194,10 +4374,17 @@ bool URigVMController::RemoveNode(URigVMNode* InNode, bool bSetupUndoRedo, bool
 		ActionStack->EndAction(Action);
 	}
 
+	if (bPrintPythonCommand)
+	{
+		RigVMPythonUtils::Print(FString::Printf(TEXT("blueprint.get_controller_by_name('%s').remove_node_by_name('%s')"),
+				*GetGraph()->GetGraphName(),
+				*InNode->GetNodePath()));
+	}
+
 	return true;
 }
 
-bool URigVMController::RemoveNodeByName(const FName& InNodeName, bool bSetupUndoRedo, bool bRecursive)
+bool URigVMController::RemoveNodeByName(const FName& InNodeName, bool bSetupUndoRedo, bool bRecursive, bool bPrintPythonCommand)
 {
 	if (!IsValidGraph())
 	{
@@ -4207,7 +4394,7 @@ bool URigVMController::RemoveNodeByName(const FName& InNodeName, bool bSetupUndo
 	URigVMGraph* Graph = GetGraph();
 	check(Graph);
 
-	return RemoveNode(Graph->FindNodeByName(InNodeName), bSetupUndoRedo, bRecursive);
+	return RemoveNode(Graph->FindNodeByName(InNodeName), bSetupUndoRedo, bRecursive, bPrintPythonCommand);
 }
 
 bool URigVMController::RenameNode(URigVMNode* InNode, const FName& InNewName, bool bSetupUndoRedo)
@@ -4274,7 +4461,7 @@ bool URigVMController::RenameNode(URigVMNode* InNode, const FName& InNewName, bo
 	return true;
 }
 
-bool URigVMController::SelectNode(URigVMNode* InNode, bool bSelect, bool bSetupUndoRedo)
+bool URigVMController::SelectNode(URigVMNode* InNode, bool bSelect, bool bSetupUndoRedo, bool bPrintPythonCommand)
 {
 	if (!IsValidNodeForGraph(InNode))
 	{
@@ -4299,7 +4486,7 @@ bool URigVMController::SelectNode(URigVMNode* InNode, bool bSelect, bool bSetupU
 		NewSelection.Remove(InNode->GetFName());
 	}
 
-	return SetNodeSelection(NewSelection, bSetupUndoRedo);
+	return SetNodeSelection(NewSelection, bSetupUndoRedo, bPrintPythonCommand);
 }
 
 bool URigVMController::SelectNodeByName(const FName& InNodeName, bool bSelect, bool bSetupUndoRedo)
@@ -4315,17 +4502,17 @@ bool URigVMController::SelectNodeByName(const FName& InNodeName, bool bSelect, b
 	return SelectNode(Graph->FindNodeByName(InNodeName), bSelect, bSetupUndoRedo);
 }
 
-bool URigVMController::ClearNodeSelection(bool bSetupUndoRedo)
+bool URigVMController::ClearNodeSelection(bool bSetupUndoRedo, bool bPrintPythonCommand)
 {
 	if (!IsValidGraph())
 	{
 		return false;
 	}
 
-	return SetNodeSelection(TArray<FName>(), bSetupUndoRedo);
+	return SetNodeSelection(TArray<FName>(), bSetupUndoRedo, bPrintPythonCommand);
 }
 
-bool URigVMController::SetNodeSelection(const TArray<FName>& InNodeNames, bool bSetupUndoRedo)
+bool URigVMController::SetNodeSelection(const TArray<FName>& InNodeNames, bool bSetupUndoRedo, bool bPrintPythonCommand)
 {
 	if (!IsValidGraph())
 	{
@@ -4404,10 +4591,28 @@ bool URigVMController::SetNodeSelection(const TArray<FName>& InNodeNames, bool b
 		Notify(ERigVMGraphNotifType::NodeSelectionChanged, nullptr);
 	}
 
+	if (bPrintPythonCommand)
+	{
+		FString ArrayStr = TEXT("[");
+		for (auto It = InNodeNames.CreateConstIterator(); It; ++It)
+		{
+			ArrayStr += TEXT("'") + It->ToString() + TEXT("'");
+			if (It.GetIndex() < InNodeNames.Num() - 1)
+			{
+				ArrayStr += TEXT(", ");
+			}
+		}
+		ArrayStr += TEXT("]");
+		
+		RigVMPythonUtils::Print(FString::Printf(TEXT("blueprint.get_controller_by_name('%s').set_node_selection(%s)"),
+				*GetGraph()->GetGraphName(),
+				*ArrayStr));
+	}
+
 	return bSelectedSomething;
 }
 
-bool URigVMController::SetNodePosition(URigVMNode* InNode, const FVector2D& InPosition, bool bSetupUndoRedo, bool bMergeUndoAction)
+bool URigVMController::SetNodePosition(URigVMNode* InNode, const FVector2D& InPosition, bool bSetupUndoRedo, bool bMergeUndoAction, bool bPrintPythonCommand)
 {
 	if (!IsValidNodeForGraph(InNode))
 	{
@@ -4435,10 +4640,19 @@ bool URigVMController::SetNodePosition(URigVMNode* InNode, const FVector2D& InPo
 		ActionStack->EndAction(Action, bMergeUndoAction);
 	}
 
+	if (bPrintPythonCommand)
+	{
+		RigVMPythonUtils::Print(FString::Printf(TEXT("blueprint.get_controller_by_name('%s').set_node_position_by_name('%s', %s)"),
+				*GetGraph()->GetGraphName(),
+				*InNode->GetNodePath(),
+				*RigVMPythonUtils::Vector2DToPythonString(InPosition)));
+	}
+
 	return true;
 }
 
-bool URigVMController::SetNodePositionByName(const FName& InNodeName, const FVector2D& InPosition, bool bSetupUndoRedo, bool bMergeUndoAction)
+bool URigVMController::SetNodePositionByName(const FName& InNodeName, const FVector2D& InPosition, bool bSetupUndoRedo, bool bMergeUndoAction, bool
+                                             bPrintPythonCommand)
 {
 	if(!IsValidGraph())
 	{
@@ -4449,10 +4663,10 @@ bool URigVMController::SetNodePositionByName(const FName& InNodeName, const FVec
 	check(Graph);
 
 	URigVMNode* Node = Graph->FindNodeByName(InNodeName);
-	return SetNodePosition(Node, InPosition, bSetupUndoRedo, bMergeUndoAction);
+	return SetNodePosition(Node, InPosition, bSetupUndoRedo, bMergeUndoAction, bPrintPythonCommand);
 }
 
-bool URigVMController::SetNodeSize(URigVMNode* InNode, const FVector2D& InSize, bool bSetupUndoRedo, bool bMergeUndoAction)
+bool URigVMController::SetNodeSize(URigVMNode* InNode, const FVector2D& InSize, bool bSetupUndoRedo, bool bMergeUndoAction, bool bPrintPythonCommand)
 {
 	if (!IsValidNodeForGraph(InNode))
 	{
@@ -4480,10 +4694,18 @@ bool URigVMController::SetNodeSize(URigVMNode* InNode, const FVector2D& InSize,
 		ActionStack->EndAction(Action, bMergeUndoAction);
 	}
 
+	if (bPrintPythonCommand)
+	{
+		RigVMPythonUtils::Print(FString::Printf(TEXT("blueprint.get_controller_by_name('%s').set_node_size_by_name('%s', %s)"),
+				*GetGraph()->GetGraphName(),
+				*InNode->GetNodePath(),
+				*RigVMPythonUtils::Vector2DToPythonString(InSize)));
+	}
+
 	return true;
 }
 
-bool URigVMController::SetNodeSizeByName(const FName& InNodeName, const FVector2D& InSize, bool bSetupUndoRedo, bool bMergeUndoAction)
+bool URigVMController::SetNodeSizeByName(const FName& InNodeName, const FVector2D& InSize, bool bSetupUndoRedo, bool bMergeUndoAction, bool bPrintPythonCommand)
 {
 	if(!IsValidGraph())
 	{
@@ -4494,10 +4716,10 @@ bool URigVMController::SetNodeSizeByName(const FName& InNodeName, const FVector2
 	check(Graph);
 
 	URigVMNode* Node = Graph->FindNodeByName(InNodeName);
-	return SetNodeSize(Node, InSize, bSetupUndoRedo, bMergeUndoAction);
+	return SetNodeSize(Node, InSize, bSetupUndoRedo, bMergeUndoAction, bPrintPythonCommand);
 }
 
-bool URigVMController::SetNodeColor(URigVMNode* InNode, const FLinearColor& InColor, bool bSetupUndoRedo, bool bMergeUndoAction)
+bool URigVMController::SetNodeColor(URigVMNode* InNode, const FLinearColor& InColor, bool bSetupUndoRedo, bool bMergeUndoAction, bool bPrintPythonCommand)
 {
 	if (!IsValidNodeForGraph(InNode))
 	{
@@ -4537,6 +4759,14 @@ bool URigVMController::SetNodeColor(URigVMNode* InNode, const FLinearColor& InCo
 		ActionStack->EndAction(Action, bMergeUndoAction);
 	}
 
+	if (bPrintPythonCommand)
+	{
+		RigVMPythonUtils::Print(FString::Printf(TEXT("blueprint.get_controller_by_name('%s').set_node_color_by_name('%s', %s)"),
+				*GetGraph()->GetGraphName(),
+				*InNode->GetNodePath(),
+				*RigVMPythonUtils::LinearColorToPythonString(InColor)));
+	}
+
 	return true;
 }
 
@@ -4554,7 +4784,7 @@ bool URigVMController::SetNodeColorByName(const FName& InNodeName, const FLinear
 	return SetNodeColor(Node, InColor, bSetupUndoRedo, bMergeUndoAction);
 }
 
-bool URigVMController::SetNodeCategory(URigVMCollapseNode* InNode, const FString& InCategory, bool bSetupUndoRedo, bool bMergeUndoAction)
+bool URigVMController::SetNodeCategory(URigVMCollapseNode* InNode, const FString& InCategory, bool bSetupUndoRedo, bool bMergeUndoAction, bool bPrintPythonCommand)
 {
 	if (!IsValidNodeForGraph(InNode))
 	{
@@ -4582,6 +4812,14 @@ bool URigVMController::SetNodeCategory(URigVMCollapseNode* InNode, const FString
 		ActionStack->EndAction(Action, bMergeUndoAction);
 	}
 
+	if (bPrintPythonCommand)
+	{
+		RigVMPythonUtils::Print(FString::Printf(TEXT("blueprint.get_controller_by_name('%s').set_node_category_by_name('%s', '%s')"),
+				*GetGraph()->GetGraphName(),
+				*InNode->GetNodePath(),
+				*InCategory));
+	}
+
 	return true;
 }
 
@@ -4599,7 +4837,7 @@ bool URigVMController::SetNodeCategoryByName(const FName& InNodeName, const FStr
 	return SetNodeCategory(Node, InCategory, bSetupUndoRedo, bMergeUndoAction);
 }
 
-bool URigVMController::SetNodeKeywords(URigVMCollapseNode* InNode, const FString& InKeywords, bool bSetupUndoRedo, bool bMergeUndoAction)
+bool URigVMController::SetNodeKeywords(URigVMCollapseNode* InNode, const FString& InKeywords, bool bSetupUndoRedo, bool bMergeUndoAction, bool bPrintPythonCommand)
 {
 	if (!IsValidNodeForGraph(InNode))
 	{
@@ -4627,6 +4865,14 @@ bool URigVMController::SetNodeKeywords(URigVMCollapseNode* InNode, const FString
 		ActionStack->EndAction(Action, bMergeUndoAction);
 	}
 
+	if (bPrintPythonCommand)
+	{
+		RigVMPythonUtils::Print(FString::Printf(TEXT("blueprint.get_controller_by_name('%s').set_node_keywords_by_name('%s', '%s')"),
+				*GetGraph()->GetGraphName(),
+				*InNode->GetNodePath(),
+				*InKeywords));
+	}
+
 	return true;
 }
 
@@ -4644,7 +4890,7 @@ bool URigVMController::SetNodeKeywordsByName(const FName& InNodeName, const FStr
 	return SetNodeKeywords(Node, InKeywords, bSetupUndoRedo, bMergeUndoAction);
 }
 
-bool URigVMController::SetNodeDescription(URigVMCollapseNode* InNode, const FString& InDescription, bool bSetupUndoRedo, bool bMergeUndoAction)
+bool URigVMController::SetNodeDescription(URigVMCollapseNode* InNode, const FString& InDescription, bool bSetupUndoRedo, bool bMergeUndoAction, bool bPrintPythonCommand)
 {
 	if (!IsValidNodeForGraph(InNode))
 	{
@@ -4672,6 +4918,14 @@ bool URigVMController::SetNodeDescription(URigVMCollapseNode* InNode, const FStr
 		ActionStack->EndAction(Action, bMergeUndoAction);
 	}
 
+	if (bPrintPythonCommand)
+	{
+		RigVMPythonUtils::Print(FString::Printf(TEXT("blueprint.get_controller_by_name('%s').set_node_description_by_name('%s', '%s')"),
+				*GetGraph()->GetGraphName(),
+				*InNode->GetNodePath(),
+				*InDescription));
+	}
+
 	return true;
 }
 
@@ -4689,7 +4943,7 @@ bool URigVMController::SetNodeDescriptionByName(const FName& InNodeName, const F
 	return SetNodeDescription(Node, InDescription, bSetupUndoRedo, bMergeUndoAction);
 }
 
-bool URigVMController::SetCommentText(URigVMNode* InNode, const FString& InCommentText, bool bSetupUndoRedo)
+bool URigVMController::SetCommentText(URigVMNode* InNode, const FString& InCommentText, bool bSetupUndoRedo, bool bPrintPythonCommand)
 {
 	if (!IsValidNodeForGraph(InNode))
 	{
@@ -4719,13 +4973,21 @@ bool URigVMController::SetCommentText(URigVMNode* InNode, const FString& InComme
 			ActionStack->EndAction(Action);
 		}
 
+		if (bPrintPythonCommand)
+		{
+			RigVMPythonUtils::Print(FString::Printf(TEXT("blueprint.get_controller_by_name('%s').set_comment_text_by_name('%s', '%s')"),
+				*GetGraph()->GetGraphName(),
+				*CommentNode->GetNodePath(),
+				*InCommentText));
+		}
+
 		return true;
 	}
 
 	return false;
 }
 
-bool URigVMController::SetCommentTextByName(const FName& InNodeName, const FString& InCommentText, bool bSetupUndoRedo)
+bool URigVMController::SetCommentTextByName(const FName& InNodeName, const FString& InCommentText, bool bSetupUndoRedo, bool bPrintPythonCommand)
 {
 	if(!IsValidGraph())
 	{
@@ -4736,10 +4998,10 @@ bool URigVMController::SetCommentTextByName(const FName& InNodeName, const FStri
 	check(Graph);
 
 	URigVMNode* Node = Graph->FindNodeByName(InNodeName);
-	return SetCommentText(Node, InCommentText, bSetupUndoRedo);
+	return SetCommentText(Node, InCommentText, bSetupUndoRedo, bPrintPythonCommand);
 }
 
-bool URigVMController::SetRerouteCompactness(URigVMNode* InNode, bool bShowAsFullNode, bool bSetupUndoRedo)
+bool URigVMController::SetRerouteCompactness(URigVMNode* InNode, bool bShowAsFullNode, bool bSetupUndoRedo, bool bPrintPythonCommand)
 {
 	if (!IsValidNodeForGraph(InNode))
 	{
@@ -4961,7 +5223,7 @@ void URigVMController::UpdateRerouteNodeAfterChangingLinks(URigVMPin* PinChanged
 	SetRerouteCompactness(Node, bShowAsFullNode, bSetupUndoRedo);
 }
 
-bool URigVMController::SetPinExpansion(const FString& InPinPath, bool bIsExpanded, bool bSetupUndoRedo)
+bool URigVMController::SetPinExpansion(const FString& InPinPath, bool bIsExpanded, bool bSetupUndoRedo, bool bPrintPythonCommand)
 {
 	if (!IsValidGraph())
 	{
@@ -4978,7 +5240,16 @@ bool URigVMController::SetPinExpansion(const FString& InPinPath, bool bIsExpande
 		return false;
 	}
 
-	return SetPinExpansion(Pin, bIsExpanded, bSetupUndoRedo);
+	const bool bSuccess = SetPinExpansion(Pin, bIsExpanded, bSetupUndoRedo);
+	if (bSuccess && bPrintPythonCommand)
+	{
+		RigVMPythonUtils::Print(FString::Printf(TEXT("blueprint.get_controller_by_name('%s').set_pin_expansion('%s', %s)"),
+			*GetGraph()->GetGraphName(),
+			*InPinPath,
+			(bIsExpanded) ? TEXT("True") : TEXT("False")));
+	}
+
+	return bSuccess;
 }
 
 bool URigVMController::SetPinExpansion(URigVMPin* InPin, bool bIsExpanded, bool bSetupUndoRedo)
@@ -5111,7 +5382,7 @@ FString URigVMController::GetPinDefaultValue(const FString& InPinPath)
 	return Pin->GetDefaultValue();
 }
 
-bool URigVMController::SetPinDefaultValue(const FString& InPinPath, const FString& InDefaultValue, bool bResizeArrays, bool bSetupUndoRedo, bool bMergeUndoAction)
+bool URigVMController::SetPinDefaultValue(const FString& InPinPath, const FString& InDefaultValue, bool bResizeArrays, bool bSetupUndoRedo, bool bMergeUndoAction, bool bPrintPythonCommand)
 {
 	if (!IsValidGraph())
 	{
@@ -5157,6 +5428,15 @@ bool URigVMController::SetPinDefaultValue(const FString& InPinPath, const FStrin
 		}
 	}
 
+	if (bPrintPythonCommand)
+	{
+		RigVMPythonUtils::Print(FString::Printf(TEXT("blueprint.get_controller_by_name('%s').set_pin_default_value('%s', '%s', %s)"),
+			*GetGraph()->GetGraphName(),
+			*InPinPath,
+			*InDefaultValue,
+			(bResizeArrays) ? TEXT("True") : TEXT("False")));
+	}
+
 	return true;
 }
 
@@ -5268,7 +5548,7 @@ bool URigVMController::SetPinDefaultValue(URigVMPin* InPin, const FString& InDef
 	return true;
 }
 
-bool URigVMController::ResetPinDefaultValue(const FString& InPinPath, bool bSetupUndoRedo)
+bool URigVMController::ResetPinDefaultValue(const FString& InPinPath, bool bSetupUndoRedo, bool bPrintPythonCommand)
 {
 	if (!IsValidGraph())
 	{
@@ -5291,7 +5571,15 @@ bool URigVMController::ResetPinDefaultValue(const FString& InPinPath, bool bSetu
 		return false;
 	}
 
-	return ResetPinDefaultValue(Pin, bSetupUndoRedo);
+	const bool bSuccess = ResetPinDefaultValue(Pin, bSetupUndoRedo);
+	if (bSuccess && bPrintPythonCommand)
+	{
+		RigVMPythonUtils::Print(FString::Printf(TEXT("blueprint.get_controller_by_name('%s').reset_pin_default_value('%s')"),
+			*GetGraph()->GetGraphName(),
+			*InPinPath));
+	}
+
+	return bSuccess;
 }
 
 bool URigVMController::ResetPinDefaultValue(URigVMPin* InPin, bool bSetupUndoRedo)
@@ -5366,12 +5654,12 @@ bool URigVMController::ResetPinDefaultValue(URigVMPin* InPin, bool bSetupUndoRed
 
 
 
-FString URigVMController::AddArrayPin(const FString& InArrayPinPath, const FString& InDefaultValue, bool bSetupUndoRedo)
+FString URigVMController::AddArrayPin(const FString& InArrayPinPath, const FString& InDefaultValue, bool bSetupUndoRedo, bool bPrintPythonCommand)
 {
-	return InsertArrayPin(InArrayPinPath, INDEX_NONE, InDefaultValue, bSetupUndoRedo);
+	return InsertArrayPin(InArrayPinPath, INDEX_NONE, InDefaultValue, bSetupUndoRedo, bPrintPythonCommand);
 }
 
-FString URigVMController::DuplicateArrayPin(const FString& InArrayElementPinPath, bool bSetupUndoRedo)
+FString URigVMController::DuplicateArrayPin(const FString& InArrayElementPinPath, bool bSetupUndoRedo, bool bPrintPythonCommand)
 {
 	if (!IsValidGraph())
 	{
@@ -5399,10 +5687,10 @@ FString URigVMController::DuplicateArrayPin(const FString& InArrayElementPinPath
 	ensure(ArrayPin->IsArray());
 
 	FString DefaultValue = ElementPin->GetDefaultValue();
-	return InsertArrayPin(ArrayPin->GetPinPath(), ElementPin->GetPinIndex() + 1, DefaultValue, bSetupUndoRedo);
+	return InsertArrayPin(ArrayPin->GetPinPath(), ElementPin->GetPinIndex() + 1, DefaultValue, bSetupUndoRedo, bPrintPythonCommand);
 }
 
-FString URigVMController::InsertArrayPin(const FString& InArrayPinPath, int32 InIndex, const FString& InDefaultValue, bool bSetupUndoRedo)
+FString URigVMController::InsertArrayPin(const FString& InArrayPinPath, int32 InIndex, const FString& InDefaultValue, bool bSetupUndoRedo, bool bPrintPythonCommand)
 {
 	if (!IsValidGraph())
 	{
@@ -5422,6 +5710,15 @@ FString URigVMController::InsertArrayPin(const FString& InArrayPinPath, int32 In
 	URigVMPin* ElementPin = InsertArrayPin(ArrayPin, InIndex, InDefaultValue, bSetupUndoRedo);
 	if (ElementPin)
 	{
+		if (bPrintPythonCommand)
+		{
+			RigVMPythonUtils::Print(FString::Printf(TEXT("blueprint.get_controller_by_name('%s').insert_array_pin('%s', %d, '%s')"),
+				*GetGraph()->GetGraphName(),
+				*InArrayPinPath,
+				InIndex,
+				*InDefaultValue));
+		}
+		
 		return ElementPin->GetPinPath();
 	}
 
@@ -5509,7 +5806,7 @@ URigVMPin* URigVMController::InsertArrayPin(URigVMPin* ArrayPin, int32 InIndex,
 	return Pin;
 }
 
-bool URigVMController::RemoveArrayPin(const FString& InArrayElementPinPath, bool bSetupUndoRedo)
+bool URigVMController::RemoveArrayPin(const FString& InArrayElementPinPath, bool bSetupUndoRedo, bool bPrintPythonCommand)
 {
 	if(!IsValidGraph())
 	{
@@ -5567,6 +5864,13 @@ bool URigVMController::RemoveArrayPin(const FString& InArrayElementPinPath, bool
 		ActionStack->EndAction(Action);
 	}
 
+	if (bPrintPythonCommand)
+	{
+		RigVMPythonUtils::Print(FString::Printf(TEXT("blueprint.get_controller_by_name('%s').remove_array_pin('%s')"),
+			*GetGraph()->GetGraphName(),
+			*InArrayElementPinPath));
+	}
+
 	return true;
 }
 
@@ -5608,12 +5912,12 @@ bool URigVMController::RemovePin(URigVMPin* InPinToRemove, bool bSetupUndoRedo,
 	return true;
 }
 
-bool URigVMController::ClearArrayPin(const FString& InArrayPinPath, bool bSetupUndoRedo)
+bool URigVMController::ClearArrayPin(const FString& InArrayPinPath, bool bSetupUndoRedo, bool bPrintPythonCommand)
 {
-	return SetArrayPinSize(InArrayPinPath, 0, FString(), bSetupUndoRedo);
+	return SetArrayPinSize(InArrayPinPath, 0, FString(), bSetupUndoRedo, bPrintPythonCommand);
 }
 
-bool URigVMController::SetArrayPinSize(const FString& InArrayPinPath, int32 InSize, const FString& InDefaultValue, bool bSetupUndoRedo)
+bool URigVMController::SetArrayPinSize(const FString& InArrayPinPath, int32 InSize, const FString& InDefaultValue, bool bSetupUndoRedo, bool bPrintPythonCommand)
 {
 	if (!IsValidGraph())
 	{
@@ -5698,7 +6002,7 @@ bool URigVMController::SetArrayPinSize(const FString& InArrayPinPath, int32 InSi
 	return RemovedPins > 0 || AddedPins > 0;
 }
 
-bool URigVMController::BindPinToVariable(const FString& InPinPath, const FString& InNewBoundVariablePath, bool bSetupUndoRedo)
+bool URigVMController::BindPinToVariable(const FString& InPinPath, const FString& InNewBoundVariablePath, bool bSetupUndoRedo, bool bPrintPythonCommand)
 {
 	if (!IsValidGraph())
 	{
@@ -5715,12 +6019,21 @@ bool URigVMController::BindPinToVariable(const FString& InPinPath, const FString
 		return false;
 	}
 
-	return BindPinToVariable(Pin, InNewBoundVariablePath, bSetupUndoRedo);
+	const bool bSuccess = BindPinToVariable(Pin, InNewBoundVariablePath, bSetupUndoRedo);
+	if (bSuccess && bPrintPythonCommand)
+	{
+		RigVMPythonUtils::Print(FString::Printf(TEXT("blueprint.get_controller_by_name('%s').bind_pin_to_variable('%s', '%s')"),
+			*GetGraph()->GetGraphName(),
+			*InPinPath,
+			*InNewBoundVariablePath));
+	}
+	
+	return bSuccess;
 }
 
-bool URigVMController::UnbindPinFromVariable(const FString& InPinPath, bool bSetupUndoRedo)
+bool URigVMController::UnbindPinFromVariable(const FString& InPinPath, bool bSetupUndoRedo, bool bPrintPythonCommand)
 {
-	return BindPinToVariable(InPinPath, FString(), bSetupUndoRedo);
+	return BindPinToVariable(InPinPath, FString(), bSetupUndoRedo, bPrintPythonCommand);
 }
 
 bool URigVMController::BindPinToVariable(URigVMPin* InPin, const FString& InNewBoundVariablePath, bool bSetupUndoRedo)
@@ -5804,7 +6117,7 @@ bool URigVMController::BindPinToVariable(URigVMPin* InPin, const FString& InNewB
 	return true;
 }
 
-bool URigVMController::MakeBindingsFromVariableNode(const FName& InNodeName, bool bSetupUndoRedo)
+bool URigVMController::MakeBindingsFromVariableNode(const FName& InNodeName, bool bSetupUndoRedo, bool bPrintPythonCommand)
 {
 	if (!IsValidGraph())
 	{
@@ -5903,12 +6216,12 @@ bool URigVMController::MakeBindingsFromVariableNode(URigVMVariableNode* InNode,
 
 }
 
-bool URigVMController::MakeVariableNodeFromBinding(const FString& InPinPath, const FVector2D& InNodePosition, bool bSetupUndoRedo)
+bool URigVMController::MakeVariableNodeFromBinding(const FString& InPinPath, const FVector2D& InNodePosition, bool bSetupUndoRedo, bool bPrintPythonCommand)
 {
-	return PromotePinToVariable(InPinPath, true, InNodePosition, bSetupUndoRedo);
+	return PromotePinToVariable(InPinPath, true, InNodePosition, bSetupUndoRedo, bPrintPythonCommand);
 }
 
-bool URigVMController::PromotePinToVariable(const FString& InPinPath, bool bCreateVariableNode, const FVector2D& InNodePosition, bool bSetupUndoRedo)
+bool URigVMController::PromotePinToVariable(const FString& InPinPath, bool bCreateVariableNode, const FVector2D& InNodePosition, bool bSetupUndoRedo, bool bPrintPythonCommand)
 {
 	if (!IsValidGraph())
 	{
@@ -5925,7 +6238,17 @@ bool URigVMController::PromotePinToVariable(const FString& InPinPath, bool bCrea
 		return false;
 	}
 
-	return PromotePinToVariable(Pin, bCreateVariableNode, InNodePosition, bSetupUndoRedo);
+	const bool bSuccess = PromotePinToVariable(Pin, bCreateVariableNode, InNodePosition, bSetupUndoRedo);
+	if (bSuccess && bPrintPythonCommand)
+	{
+		RigVMPythonUtils::Print(FString::Printf(TEXT("blueprint.get_controller_by_name('%s').promote_pin_to_variable('%s', %s, %s)"),
+			*GetGraph()->GetGraphName(),
+			*InPinPath,
+			(bCreateVariableNode) ? TEXT("True") : TEXT("False"),
+			*RigVMPythonUtils::Vector2DToPythonString(InNodePosition)));
+	}
+	
+	return bSuccess;
 }
 
 bool URigVMController::PromotePinToVariable(URigVMPin* InPin, bool bCreateVariableNode, const FVector2D& InNodePosition, bool bSetupUndoRedo)
@@ -6006,7 +6329,7 @@ bool URigVMController::PromotePinToVariable(URigVMPin* InPin, bool bCreateVariab
 	return false;
 }
 
-bool URigVMController::AddLink(const FString& InOutputPinPath, const FString& InInputPinPath, bool bSetupUndoRedo)
+bool URigVMController::AddLink(const FString& InOutputPinPath, const FString& InInputPinPath, bool bSetupUndoRedo, bool bPrintPythonCommand)
 {
 	if(!IsValidGraph())
 	{
@@ -6044,7 +6367,16 @@ bool URigVMController::AddLink(const FString& InOutputPinPath, const FString& In
 	}
 	InputPin = InputPin->GetPinForLink();
 
-	return AddLink(OutputPin, InputPin, bSetupUndoRedo);
+	const bool bSuccess = AddLink(OutputPin, InputPin, bSetupUndoRedo);
+	if (bSuccess && bPrintPythonCommand)
+	{		
+		RigVMPythonUtils::Print(FString::Printf(TEXT("blueprint.get_controller_by_name('%s').add_link('%s', '%s')"),
+			*GetGraph()->GetGraphName(),
+			*OutputPin->GetPinPath(),
+			*InputPin->GetPinPath()));
+	}
+	
+	return bSuccess;
 }
 
 bool URigVMController::AddLink(URigVMPin* OutputPin, URigVMPin* InputPin, bool bSetupUndoRedo)
@@ -6149,7 +6481,7 @@ bool URigVMController::AddLink(URigVMPin* OutputPin, URigVMPin* InputPin, bool b
 	return true;
 }
 
-bool URigVMController::BreakLink(const FString& InOutputPinPath, const FString& InInputPinPath, bool bSetupUndoRedo)
+bool URigVMController::BreakLink(const FString& InOutputPinPath, const FString& InInputPinPath, bool bSetupUndoRedo, bool bPrintPythonCommand)
 {
 	if(!IsValidGraph())
 	{
@@ -6175,7 +6507,15 @@ bool URigVMController::BreakLink(const FString& InOutputPinPath, const FString&
 	}
 	InputPin = InputPin->GetPinForLink();
 
-	return BreakLink(OutputPin, InputPin, bSetupUndoRedo);
+	const bool bSuccess = BreakLink(OutputPin, InputPin, bSetupUndoRedo);
+	if (bSuccess && bPrintPythonCommand)
+	{
+		RigVMPythonUtils::Print(FString::Printf(TEXT("blueprint.get_controller_by_name('%s').break_link('%s', '%s')"),
+			*GetGraph()->GetGraphName(),
+			*OutputPin->GetPinPath(),
+			*InputPin->GetPinPath()));
+	}
+	return bSuccess;
 }
 
 bool URigVMController::BreakLink(URigVMPin* OutputPin, URigVMPin* InputPin, bool bSetupUndoRedo)
@@ -6239,7 +6579,7 @@ bool URigVMController::BreakLink(URigVMPin* OutputPin, URigVMPin* InputPin, bool
 	return false;
 }
 
-bool URigVMController::BreakAllLinks(const FString& InPinPath, bool bAsInput, bool bSetupUndoRedo)
+bool URigVMController::BreakAllLinks(const FString& InPinPath, bool bAsInput, bool bSetupUndoRedo, bool bPrintPythonCommand)
 {
 	if(!IsValidGraph())
 	{
@@ -6262,7 +6602,15 @@ bool URigVMController::BreakAllLinks(const FString& InPinPath, bool bAsInput, bo
 		return false;
 	}
 
-	return BreakAllLinks(Pin, bAsInput, bSetupUndoRedo);
+	const bool bSuccess = BreakAllLinks(Pin, bAsInput, bSetupUndoRedo);
+	if (bSuccess && bPrintPythonCommand)
+	{
+		RigVMPythonUtils::Print(FString::Printf(TEXT("blueprint.get_controller_by_name('%s').break_all_links('%s', %s)"),
+			*GetGraph()->GetGraphName(),
+			*Pin->GetPinPath(),
+			bAsInput ? TEXT("True") : TEXT("False")));
+	}
+	return bSuccess;
 }
 
 bool URigVMController::BreakAllLinks(URigVMPin* Pin, bool bAsInput, bool bSetupUndoRedo)
@@ -6331,7 +6679,7 @@ void URigVMController::BreakAllLinksRecursive(URigVMPin* Pin, bool bAsInput, boo
 	}
 }
 
-FName URigVMController::AddExposedPin(const FName& InPinName, ERigVMPinDirection InDirection, const FString& InCPPType, const FName& InCPPTypeObjectPath, const FString& InDefaultValue, bool bSetupUndoRedo)
+FName URigVMController::AddExposedPin(const FName& InPinName, ERigVMPinDirection InDirection, const FString& InCPPType, const FName& InCPPTypeObjectPath, const FString& InDefaultValue, bool bSetupUndoRedo, bool bPrintPythonCommand)
 {
 	if (!IsValidGraph())
 	{
@@ -6435,11 +6783,33 @@ FName URigVMController::AddExposedPin(const FName& InPinName, ERigVMPinDirection
 	{
 		ActionStack->EndAction(Action);
 	}
+
+	if (bPrintPythonCommand)
+	{
+		FString DirectionStr;
+		switch (InDirection)
+		{
+			case ERigVMPinDirection::Hidden: DirectionStr = TEXT("unreal.RigVMPinDirection.HIDDEN"); break;
+			case ERigVMPinDirection::Input: DirectionStr = TEXT("unreal.RigVMPinDirection.INPUT"); break;
+			case ERigVMPinDirection::Output: DirectionStr = TEXT("unreal.RigVMPinDirection.OUTPUT"); break;
+			case ERigVMPinDirection::Visible: DirectionStr = TEXT("unreal.RigVMPinDirection.VISIBLE"); break;
+			case ERigVMPinDirection::IO: DirectionStr = TEXT("unreal.RigVMPinDirection.IO"); break;
+		}
+		
+		//AddExposedPin(const FName& InPinName, ERigVMPinDirection InDirection, const FString& InCPPType, const FName& InCPPTypeObjectPath, const FString& InDefaultValue, bool bSetupUndoRedo, bool bPrintPythonCommand)
+		RigVMPythonUtils::Print(FString::Printf(TEXT("blueprint.get_controller_by_name('%s').add_exposed_pin('%s', %s, '%s', '%s', '%s')"),
+				*GetGraph()->GetGraphName(),
+				*InPinName.ToString(),
+				*DirectionStr,
+				*InCPPType,
+				*InCPPTypeObjectPath.ToString(),
+				*InDefaultValue));
+	}
 	
 	return PinName;
 }
 
-bool URigVMController::RemoveExposedPin(const FName& InPinName, bool bSetupUndoRedo)
+bool URigVMController::RemoveExposedPin(const FName& InPinName, bool bSetupUndoRedo, bool bPrintPythonCommand)
 {
 	if (!IsValidGraph())
 	{
@@ -6511,10 +6881,17 @@ bool URigVMController::RemoveExposedPin(const FName& InPinName, bool bSetupUndoR
 		}
 	}
 
+	if (bSuccessfullyRemovedPin && bPrintPythonCommand)
+	{
+		RigVMPythonUtils::Print(FString::Printf(TEXT("blueprint.get_controller_by_name('%s').remove_exposed_pin('%s')"),
+				*GetGraph()->GetGraphName(),
+				*InPinName.ToString()));
+	}
+
 	return bSuccessfullyRemovedPin;
 }
 
-bool URigVMController::RenameExposedPin(const FName& InOldPinName, const FName& InNewPinName, bool bSetupUndoRedo)
+bool URigVMController::RenameExposedPin(const FName& InOldPinName, const FName& InNewPinName, bool bSetupUndoRedo, bool bPrintPythonCommand)
 {
 	if (!IsValidGraph())
 	{
@@ -6652,10 +7029,18 @@ bool URigVMController::RenameExposedPin(const FName& InOldPinName, const FName&
 		ActionStack->EndAction(Action);
 	}
 
+	if (bPrintPythonCommand)
+	{
+		RigVMPythonUtils::Print(FString::Printf(TEXT("blueprint.get_controller_by_name('%s').rename_exposed_pin('%s', '%s')"),
+				*GetGraph()->GetGraphName(),
+				*InOldPinName.ToString(),
+				*InNewPinName.ToString()));
+	}
+
 	return true;
 }
 
-bool URigVMController::ChangeExposedPinType(const FName& InPinName, const FString& InCPPType, const FName& InCPPTypeObjectPath, bool& bSetupUndoRedo, bool bSetupOrphanPins)
+bool URigVMController::ChangeExposedPinType(const FName& InPinName, const FString& InCPPType, const FName& InCPPTypeObjectPath, bool& bSetupUndoRedo, bool bSetupOrphanPins, bool bPrintPythonCommand)
 {
 	if (!IsValidGraph())
 	{
@@ -6755,10 +7140,19 @@ bool URigVMController::ChangeExposedPinType(const FName& InPinName, const FStrin
 		ActionStack->EndAction(Action);
 	}
 
+	if (bPrintPythonCommand)
+	{
+		RigVMPythonUtils::Print(FString::Printf(TEXT("blueprint.get_controller_by_name('%s').change_exposed_pin_type('%s', '%s', '%s')"),
+				*GetGraph()->GetGraphName(),
+				*InPinName.ToString(),
+				*InCPPType,
+				*InCPPTypeObjectPath.ToString()));
+	}
+
 	return true;
 }
 
-bool URigVMController::SetExposedPinIndex(const FName& InPinName, int32 InNewIndex, bool bSetupUndoRedo)
+bool URigVMController::SetExposedPinIndex(const FName& InPinName, int32 InNewIndex, bool bSetupUndoRedo, bool bPrintPythonCommand)
 {
 	if (!IsValidGraph())
 	{
@@ -6818,10 +7212,18 @@ bool URigVMController::SetExposedPinIndex(const FName& InPinName, int32 InNewInd
 		ActionStack->AddAction(PinIndexAction);
 	}
 
+	if (bPrintPythonCommand)
+	{
+		RigVMPythonUtils::Print(FString::Printf(TEXT("blueprint.get_controller_by_name('%s').set_exposed_pin_index('%s', %d)"),
+				*GetGraph()->GetGraphName(),
+				*InPinName.ToString(),
+				InNewIndex));
+	}
+
 	return true;
 }
 
-URigVMFunctionReferenceNode* URigVMController::AddFunctionReferenceNode(URigVMLibraryNode* InFunctionDefinition, const FVector2D& InNodePosition, const FString& InNodeName, bool bSetupUndoRedo)
+URigVMFunctionReferenceNode* URigVMController::AddFunctionReferenceNode(URigVMLibraryNode* InFunctionDefinition, const FVector2D& InNodePosition, const FString& InNodeName, bool bSetupUndoRedo, bool bPrintPythonCommand)
 {
 	if (!IsValidGraph())
 	{
@@ -6898,6 +7300,29 @@ URigVMFunctionReferenceNode* URigVMController::AddFunctionReferenceNode(URigVMLi
 		ActionStack->EndAction(InverseAction);
 	}
 
+	if (bPrintPythonCommand)
+	{
+		if (InFunctionDefinition->GetLibrary() == GetGraph()->GetDefaultFunctionLibrary())
+		{
+			RigVMPythonUtils::Print(FString::Printf(TEXT("blueprint.get_controller_by_name('%s').add_function_reference_node(library.find_function('%s'), %s, '%s')"),
+						*GetGraph()->GetGraphName(),
+						*InFunctionDefinition->GetName(),
+						*RigVMPythonUtils::Vector2DToPythonString(InNodePosition),
+						*InNodeName));
+		}
+		else
+		{
+			RigVMPythonUtils::Print(FString::Printf(TEXT("function_blueprint = unreal.load_object(name = '%s', outer = None)"),
+				*InFunctionDefinition->GetLibrary()->GetOuter()->GetPathName()));
+			RigVMPythonUtils::Print(FString::Printf(TEXT("blueprint.get_controller_by_name('%s').add_function_reference_node(function_blueprint.get_local_function_library().find_function('%s'), %s, '%s')"),
+						*GetGraph()->GetGraphName(),
+						*InFunctionDefinition->GetName(),
+						*RigVMPythonUtils::Vector2DToPythonString(InFunctionDefinition->GetPosition()), 
+						*InFunctionDefinition->GetName()));
+		}
+		
+	}
+
 	return FunctionRefNode;
 }
 
@@ -6988,7 +7413,7 @@ bool URigVMController::SetRemappedVariable(URigVMFunctionReferenceNode* InFuncti
 	return true;
 }
 
-URigVMLibraryNode* URigVMController::AddFunctionToLibrary(const FName& InFunctionName, bool bMutable, const FVector2D& InNodePosition, bool bSetupUndoRedo)
+URigVMLibraryNode* URigVMController::AddFunctionToLibrary(const FName& InFunctionName, bool bMutable, const FVector2D& InNodePosition, bool bSetupUndoRedo, bool bPrintPythonCommand)
 {
 	if (!IsValidGraph())
 	{
@@ -7053,6 +7478,16 @@ URigVMLibraryNode* URigVMController::AddFunctionToLibrary(const FName& InFunctio
 		ActionStack->EndAction(InverseAction);
 	}
 
+	if (bPrintPythonCommand)
+	{
+		//AddFunctionToLibrary(const FName& InFunctionName, bool bMutable, const FVector2D& InNodePosition, bool bSetupUndoRedo, bool bPrintPythonCommand)
+		RigVMPythonUtils::Print(FString::Printf(TEXT("library_controller.add_function_to_library('%s', %s, %s)"),
+				*GetGraph()->GetGraphName(),
+				*InFunctionName.ToString(),
+				(bMutable) ? TEXT("True") : TEXT("False"),
+				*RigVMPythonUtils::Vector2DToPythonString(InNodePosition)));
+	}
+
 	return CollapseNode;
 }
 
@@ -7076,8 +7511,7 @@ bool URigVMController::RemoveFunctionFromLibrary(const FName& InFunctionName, bo
 	return RemoveNodeByName(InFunctionName, bSetupUndoRedo);
 }
 
-FRigVMGraphVariableDescription URigVMController::AddLocalVariable(const FName& InVariableName,
-	const FString& InCPPType, UObject* InCPPTypeObject, const FString& InDefaultValue, bool bSetupUndoRedo)
+FRigVMGraphVariableDescription URigVMController::AddLocalVariable(const FName& InVariableName, const FString& InCPPType, UObject* InCPPTypeObject, const FString& InDefaultValue, bool bSetupUndoRedo, bool bPrintPythonCommand)
 {
 	FRigVMGraphVariableDescription NewVariable;
 	if (!IsValidGraph())
@@ -7121,6 +7555,16 @@ FRigVMGraphVariableDescription URigVMController::AddLocalVariable(const FName& I
 		Graph->MarkPackageDirty();
 	}
 
+	if (bPrintPythonCommand)
+	{
+		RigVMPythonUtils::Print(FString::Printf(TEXT("blueprint.get_controller_by_name('%s').add_local_variable_from_object_path('%s', '%s', '%s', '%s')"),
+				*GetGraph()->GetGraphName(),
+				*NewVariable.Name.ToString(),
+				*NewVariable.CPPType,
+				(NewVariable.CPPTypeObject) ? *NewVariable.CPPTypeObject->GetPathName() : *FString(),
+				*NewVariable.DefaultValue));
+	}
+
 	return NewVariable;
 }
 
@@ -7146,7 +7590,7 @@ FRigVMGraphVariableDescription URigVMController::AddLocalVariableFromObjectPath(
 	return AddLocalVariable(InVariableName, InCPPType, CPPTypeObject, InDefaultValue, bSetupUndoRedo);
 }
 
-bool URigVMController::RemoveLocalVariable(const FName& InVariableName, bool bSetupUndoRedo)
+bool URigVMController::RemoveLocalVariable(const FName& InVariableName, bool bSetupUndoRedo, bool bPrintPythonCommand)
 {
 	if (!IsValidGraph())
 	{
@@ -7211,14 +7655,20 @@ bool URigVMController::RemoveLocalVariable(const FName& InVariableName, bool bSe
 		}	
 		
 		LocalVariables.RemoveAt(FoundIndex);
+
+		if (bPrintPythonCommand)
+		{
+			RigVMPythonUtils::Print(FString::Printf(TEXT("blueprint.get_controller_by_name('%s').remove_local_variable('%s')"),
+					*GetGraph()->GetGraphName(),
+					*InVariableName.ToString()));
+		}
 		return true;
 	}
 
 	return false;
 }
 
-bool URigVMController::RenameLocalVariable(const FName& InVariableName, const FName& InNewVariableName,
-	bool bSetupUndoRedo)
+bool URigVMController::RenameLocalVariable(const FName& InVariableName, const FName& InNewVariableName, bool bSetupUndoRedo, bool bPrintPythonCommand)
 {
 	if (!IsValidGraph())
 	{
@@ -7302,11 +7752,19 @@ bool URigVMController::RenameLocalVariable(const FName& InVariableName, const FN
 		}
 	}
 
+	if (bPrintPythonCommand)
+	{
+		RigVMPythonUtils::Print(FString::Printf(TEXT("blueprint.get_controller_by_name('%s').rename_local_variable('%s', '%s')"),
+				*GetGraph()->GetGraphName(),
+				*InVariableName.ToString(),
+				*InNewVariableName.ToString()));
+	}
+
 	return true;
 }
 
 bool URigVMController::SetLocalVariableType(const FName& InVariableName, const FString& InCPPType,
-	UObject* InCPPTypeObject, bool bSetupUndoRedo)
+                                            UObject* InCPPTypeObject, bool bSetupUndoRedo, bool bPrintPythonCommand)
 {
 	if (!IsValidGraph())
 	{
@@ -7383,10 +7841,36 @@ bool URigVMController::SetLocalVariableType(const FName& InVariableName, const F
 			}
 		}
 	}
+
+	if (bPrintPythonCommand)
+	{
+		//bool URigVMController::SetLocalVariableType(const FName& InVariableName, const FString& InCPPType, UObject* InCPPTypeObject, bool bSetupUndoRedo, bool bPrintPythonCommand)
+		RigVMPythonUtils::Print(FString::Printf(TEXT("blueprint.get_controller_by_name('%s').set_local_variable_type_from_object_path('%s', '%s', '%s')"),
+				*GetGraph()->GetGraphName(),
+				*InVariableName.ToString(),
+				*InCPPType,
+				(InCPPTypeObject) ? *InCPPTypeObject->GetPathName() : *FString()));
+	}
 	
 	return true;
 }
 
+bool URigVMController::SetLocalVariableTypeFromObjectPath(const FName& InVariableName, const FString& InCPPType, const FString& InCPPTypeObjectPath, bool bSetupUndoRedo, bool bPrintPythonCommand)
+{
+	UObject* CPPTypeObject = nullptr;
+	if (!InCPPTypeObjectPath.IsEmpty())
+	{
+		CPPTypeObject = URigVMPin::FindObjectFromCPPTypeObjectPath<UObject>(InCPPTypeObjectPath);
+		if (CPPTypeObject == nullptr)
+		{
+			ReportErrorf(TEXT("Cannot find cpp type object for path '%s'."), *InCPPTypeObjectPath);
+			return false;
+		}
+	}
+
+	return SetLocalVariableType(InVariableName, InCPPType, CPPTypeObject, bSetupUndoRedo, bPrintPythonCommand);
+}
+
 TArray<TSoftObjectPtr<URigVMFunctionReferenceNode>> URigVMController::GetAffectedReferences(ERigVMControllerBulkEditType InEditType, bool bForceLoad, bool bNotify)
 {
 	TArray<TSoftObjectPtr<URigVMFunctionReferenceNode>> FunctionReferencePtrs;
@@ -7747,7 +8231,7 @@ URigVMRerouteNode* URigVMController::AddFreeRerouteNode(bool bShowAsFullNode, co
 	return Node;
 }
 
-URigVMBranchNode* URigVMController::AddBranchNode(const FVector2D& InPosition, const FString& InNodeName, bool bSetupUndoRedo)
+URigVMBranchNode* URigVMController::AddBranchNode(const FVector2D& InPosition, const FString& InNodeName, bool bSetupUndoRedo, bool bPrintPythonCommand)
 {
 	if(!IsValidGraph())
 	{
@@ -7797,10 +8281,19 @@ URigVMBranchNode* URigVMController::AddBranchNode(const FVector2D& InPosition, c
 		ActionStack->AddAction(FRigVMAddBranchNodeAction(Node));
 	}
 
+	if (bPrintPythonCommand)
+	{
+		TArray<FString> Commands = GetAddNodePythonCommands(Node);
+		for (const FString& Command : Commands)
+		{
+			RigVMPythonUtils::Print(FString::Printf(TEXT("%s"), *Command));
+		}
+	}
+
 	return Node;
 }
 
-URigVMIfNode* URigVMController::AddIfNode(const FString& InCPPType, const FName& InCPPTypeObjectPath, const FVector2D& InPosition, const FString& InNodeName, bool bSetupUndoRedo)
+URigVMIfNode* URigVMController::AddIfNode(const FString& InCPPType, const FName& InCPPTypeObjectPath, const FVector2D& InPosition, const FString& InNodeName, bool  bSetupUndoRedo, bool bPrintPythonCommand)
 {
 	if(!IsValidGraph())
 	{
@@ -7892,11 +8385,20 @@ URigVMIfNode* URigVMController::AddIfNode(const FString& InCPPType, const FName&
 	{
 		ActionStack->AddAction(FRigVMAddIfNodeAction(Node));
 	}
+	
+	if (bPrintPythonCommand)
+	{
+		TArray<FString> Commands = GetAddNodePythonCommands(Node);
+		for (const FString& Command : Commands)
+		{
+			RigVMPythonUtils::Print(FString::Printf(TEXT("%s"), *Command));
+		}
+	}
 
 	return Node;
 }
 
-URigVMSelectNode* URigVMController::AddSelectNode(const FString& InCPPType, const FName& InCPPTypeObjectPath, const FVector2D& InPosition, const FString& InNodeName, bool bSetupUndoRedo)
+URigVMSelectNode* URigVMController::AddSelectNode(const FString& InCPPType, const FName& InCPPTypeObjectPath, const FVector2D& InPosition, const FString& InNodeName, bool bSetupUndoRedo, bool bPrintPythonCommand)
 {
 	if (!IsValidGraph())
 	{
@@ -7968,6 +8470,15 @@ URigVMSelectNode* URigVMController::AddSelectNode(const FString& InCPPType, cons
 		ActionStack->AddAction(FRigVMAddSelectNodeAction(Node));
 	}
 
+	if (bPrintPythonCommand)
+	{
+		TArray<FString> Commands = GetAddNodePythonCommands(Node);
+		for (const FString& Command : Commands)
+		{
+			RigVMPythonUtils::Print(FString::Printf(TEXT("%s"), *Command));
+		}
+	}
+
 	return Node;
 }
 
@@ -7982,7 +8493,7 @@ URigVMSelectNode* URigVMController::AddSelectNodeFromStruct(UScriptStruct* InScr
 	return AddSelectNode(InScriptStruct->GetStructCPPName(), FName(InScriptStruct->GetPathName()), InPosition, InNodeName, bSetupUndoRedo);
 }
 
-URigVMPrototypeNode* URigVMController::AddPrototypeNode(const FName& InNotation, const FVector2D& InPosition, const FString& InNodeName, bool bSetupUndoRedo)
+URigVMPrototypeNode* URigVMController::AddPrototypeNode(const FName& InNotation, const FVector2D& InPosition, const FString& InNodeName, bool bSetupUndoRedo, bool bPrintPythonCommand)
 {
 	if (!IsValidGraph())
 	{
@@ -8035,11 +8546,20 @@ URigVMPrototypeNode* URigVMController::AddPrototypeNode(const FName& InNotation,
 		ActionStack->AddAction(FRigVMAddPrototypeNodeAction(Node));
 	}
 
+	if (bPrintPythonCommand)
+	{
+		TArray<FString> Commands = GetAddNodePythonCommands(Node);
+		for (const FString& Command : Commands)
+		{
+			RigVMPythonUtils::Print(FString::Printf(TEXT("%s"), *Command));
+		}
+	}
+
 	return Node;
 }
 
 
-URigVMEnumNode* URigVMController::AddEnumNode(const FName& InCPPTypeObjectPath, const FVector2D& InPosition, const FString& InNodeName, bool bSetupUndoRedo)
+URigVMEnumNode* URigVMController::AddEnumNode(const FName& InCPPTypeObjectPath, const FVector2D& InPosition, const FString& InNodeName, bool bSetupUndoRedo, bool bPrintPythonCommand)
 {
 	if (!IsValidGraph())
 	{
@@ -8090,6 +8610,15 @@ URigVMEnumNode* URigVMController::AddEnumNode(const FName& InCPPTypeObjectPath,
 		ActionStack->AddAction(FRigVMAddEnumNodeAction(Node));
 	}
 
+	if (bPrintPythonCommand)
+	{
+		TArray<FString> Commands = GetAddNodePythonCommands(Node);
+		for (const FString& Command : Commands)
+		{
+			RigVMPythonUtils::Print(FString::Printf(TEXT("%s"), *Command));
+		}
+	}
+
 	return Node;
 }
 
diff --git a/Engine/Source/Developer/RigVMDeveloper/Private/RigVMPythonUtils.cpp b/Engine/Source/Developer/RigVMDeveloper/Private/RigVMPythonUtils.cpp
new file mode 100644
--- /dev/null
+++ b/Engine/Source/Developer/RigVMDeveloper/Private/RigVMPythonUtils.cpp
@@ -0,0 +1,86 @@
+﻿// Copyright Epic Games, Inc. All Rights Reserved.
+
+#include "RigVMPythonUtils.h"
+
+#if WITH_EDITOR
+#include "MessageLogModule.h"
+#include "IMessageLogListing.h"
+#endif
+
+#define LOCTEXT_NAMESPACE "RigVMDeveloperModule"
+
+
+FString RigVMPythonUtils::NameToPep8(const FString& Name)
+{
+	// Wish we could use PyGenUtil::PythonizeName, but unfortunately it's private
+
+	const FString NameNoSpaces = Name.Replace(TEXT(" "), TEXT("_"));
+	FString Result;
+
+	for (const TCHAR& Char : NameNoSpaces)
+	{
+		if (FChar::IsUpper(Char))
+		{
+			if (!Result.IsEmpty() && !Result.EndsWith(TEXT("_")))
+			{
+				Result.AppendChar(TEXT('_'));
+			}
+			Result.AppendChar(FChar::ToLower(Char));
+		}
+		else
+		{
+			Result.AppendChar(Char);
+		}
+	}
+	return Result;
+}
+
+FString RigVMPythonUtils::TransformToPythonString(const FTransform& Transform)
+{
+	return FString::Printf(TEXT("unreal.Transform(location=[%f,%f,%f],rotation=[%f,%f,%f],scale=[%f,%f,%f])"),
+	                       Transform.GetLocation().X,
+	                       Transform.GetLocation().Y,
+	                       Transform.GetLocation().Z,
+	                       Transform.Rotator().Pitch,
+	                       Transform.Rotator().Yaw,
+	                       Transform.Rotator().Roll,
+	                       Transform.GetScale3D().X,
+	                       Transform.GetScale3D().Y,
+	                       Transform.GetScale3D().Z);
+}
+
+FString RigVMPythonUtils::Vector2DToPythonString(const FVector2D& Vector)
+{
+	return FString::Printf(TEXT("unreal.Vector2D(%f, %f)"),
+	                       Vector.X,
+	                       Vector.Y);
+}
+
+FString RigVMPythonUtils::LinearColorToPythonString(const FLinearColor& Color)
+{
+	return FString::Printf(TEXT("unreal.LinearColor(%f, %f, %f, %f)"),
+	                       Color.R, Color.G, Color.B, Color.A);
+}
+
+#if WITH_EDITOR
+void RigVMPythonUtils::Print(const FString& InMessage)
+{
+	//UE_LOG(LogRigVMDeveloper, Display, TEXT("%s"), *InMessage);
+	FMessageLogModule& MessageLogModule = FModuleManager::LoadModuleChecked<FMessageLogModule>("MessageLog");
+
+	if (!MessageLogModule.IsRegisteredLogListing("RigVMPythonLog"))
+	{
+		FMessageLogInitializationOptions InitOptions;
+		InitOptions.bShowFilters = true;
+		InitOptions.bShowPages = true;
+		InitOptions.bAllowClear = true;
+		MessageLogModule.RegisterLogListing("RigVMPythonLog", LOCTEXT("RigVMPythonLog", "RigVM Python Log"), InitOptions);
+	}
+	TSharedRef<IMessageLogListing> PythonLog = MessageLogModule.GetLogListing( TEXT("RigVMPythonLog") );
+
+	TSharedRef<FTokenizedMessage> Token = FTokenizedMessage::Create(EMessageSeverity::Info, FText::FromString(InMessage));
+	PythonLog->AddMessage(Token, false);			
+	
+}
+
+#endif
diff --git a/Engine/Source/Developer/RigVMDeveloper/Public/RigVMModel/RigVMController.h b/Engine/Source/Developer/RigVMDeveloper/Public/RigVMModel/RigVMController.h
--- a/Engine/Source/Developer/RigVMDeveloper/Public/RigVMModel/RigVMController.h
+++ b/Engine/Source/Developer/RigVMDeveloper/Public/RigVMModel/RigVMController.h
@@ -147,27 +147,27 @@ public:
 	// UnitNode represent a RIGVM_METHOD declaration on a USTRUCT.
 	// This causes a NodeAdded modified event.
 	UFUNCTION(BlueprintCallable, Category = RigVMController)
-	URigVMUnitNode* AddUnitNode(UScriptStruct* InScriptStruct, const FName& InMethodName = TEXT("Execute"), const FVector2D& InPosition = FVector2D::ZeroVector, const FString& InNodeName = TEXT(""), bool bSetupUndoRedo = true);
+	URigVMUnitNode* AddUnitNode(UScriptStruct* InScriptStruct, const FName& InMethodName = TEXT("Execute"), const FVector2D& InPosition = FVector2D::ZeroVector, const FString& InNodeName = TEXT(""), bool bSetupUndoRedo = true, bool bPrintPythonCommand = false);
 
 	// Adds a Function / Struct Node to the edited Graph given its struct object path name.
 	// UnitNode represent a RIGVM_METHOD declaration on a USTRUCT.
 	// This causes a NodeAdded modified event.
 	UFUNCTION(BlueprintCallable, Category = RigVMController)
-	URigVMUnitNode* AddUnitNodeFromStructPath(const FString& InScriptStructPath, const FName& InMethodName = TEXT("Execute"), const FVector2D& InPosition = FVector2D::ZeroVector, const FString& InNodeName = TEXT(""), bool bSetupUndoRedo = true);
+	URigVMUnitNode* AddUnitNodeFromStructPath(const FString& InScriptStructPath, const FName& InMethodName = TEXT("Execute"), const FVector2D& InPosition = FVector2D::ZeroVector, const FString& InNodeName = TEXT(""), bool bSetupUndoRedo = true, bool bPrintPythonCommand = false);
 
 	// Adds a Variable Node to the edited Graph.
 	// Variables represent local work state for the function and
 	// can be read from and written to. 
 	// This causes a NodeAdded modified event.
 	UFUNCTION(BlueprintCallable, Category = RigVMController)
-	URigVMVariableNode* AddVariableNode(const FName& InVariableName, const FString& InCPPType, UObject* InCPPTypeObject, bool bIsGetter, const FString& InDefaultValue, const FVector2D& InPosition = FVector2D::ZeroVector, const FString& InNodeName = TEXT(""), bool bSetupUndoRedo = true);
+	URigVMVariableNode* AddVariableNode(const FName& InVariableName, const FString& InCPPType, UObject* InCPPTypeObject, bool bIsGetter, const FString& InDefaultValue, const FVector2D& InPosition = FVector2D::ZeroVector, const FString& InNodeName = TEXT(""), bool bSetupUndoRedo = true, bool bPrintPythonCommand = false);
 
 	// Adds a Variable Node to the edited Graph given a struct object path name.
 	// Variables represent local work state for the function and
 	// can be read from (bIsGetter == true) or written to (bIsGetter == false). 
 	// This causes a NodeAdded modified event.
 	UFUNCTION(BlueprintCallable, Category = RigVMController)
-	URigVMVariableNode* AddVariableNodeFromObjectPath(const FName& InVariableName, const FString& InCPPType, const FString& InCPPTypeObjectPath, bool bIsGetter, const FString& InDefaultValue, const FVector2D& InPosition = FVector2D::ZeroVector, const FString& InNodeName = TEXT(""), bool bSetupUndoRedo = true);
+	URigVMVariableNode* AddVariableNodeFromObjectPath(const FName& InVariableName, const FString& InCPPType, const FString& InCPPTypeObjectPath, bool bIsGetter, const FString& InDefaultValue, const FVector2D& InPosition = FVector2D::ZeroVector, const FString& InNodeName = TEXT(""), bool bSetupUndoRedo = true, bool bPrintPythonCommand = false);
 
 	// Refreshes the variable node with the new data
 	UFUNCTION(BlueprintCallable, Category = RigVMController)
@@ -175,15 +175,17 @@ public:
 
 	// Removes all nodes related to a given variable
 	UFUNCTION(BlueprintCallable, Category = RigVMController)
-	void OnExternalVariableRemoved(const FName& InVarName, bool bSetupUndoRedo);
+	void OnExternalVariableRemoved(const FName& InVarName, bool bSetupUndoRedo, bool bPrintPythonCommand = false);
 
 	// Renames the variable name in all relevant nodes
 	UFUNCTION(BlueprintCallable, Category = RigVMController)
-	void OnExternalVariableRenamed(const FName& InOldVarName, const FName& InNewVarName, bool bSetupUndoRedo);
+	void OnExternalVariableRenamed(const FName& InOldVarName, const FName& InNewVarName, bool bSetupUndoRedo, bool bPrintPythonCommand = false);
 
 	// Changes the data type of all nodes matching a given variable name
 	UFUNCTION(BlueprintCallable, Category = RigVMController)
-	void OnExternalVariableTypeChanged(const FName& InVarName, const FString& InCPPType, UObject* InCPPTypeObject, bool bSetupUndoRedo);
+	void OnExternalVariableTypeChanged(const FName& InVarName, const FString& InCPPType, UObject* InCPPTypeObject, bool bSetupUndoRedo, bool bPrintPythonCommand = false);
+	UFUNCTION(BlueprintCallable, Category = RigVMController)
+	void OnExternalVariableTypeChangedFromObjectPath(const FName& InVarName, const FString& InCPPType, const FString& InCPPTypeObjectPath, bool bSetupUndoRedo);
 
 	// Refreshes the variable node with the new data
 	UFUNCTION(BlueprintCallable, Category = RigVMController)
@@ -194,41 +196,41 @@ public:
 	// Input Parameters are constant values / literals.
 	// This causes a NodeAdded modified event.
 	UFUNCTION(BlueprintCallable, Category = RigVMController)
-	URigVMParameterNode* AddParameterNode(const FName& InParameterName, const FString& InCPPType, UObject* InCPPTypeObject, bool bIsInput, const FString& InDefaultValue, const FVector2D& InPosition = FVector2D::ZeroVector, const FString& InNodeName = TEXT(""), bool bSetupUndoRedo = true);
+	URigVMParameterNode* AddParameterNode(const FName& InParameterName, const FString& InCPPType, UObject* InCPPTypeObject, bool bIsInput, const FString& InDefaultValue, const FVector2D& InPosition = FVector2D::ZeroVector, const FString& InNodeName = TEXT(""), bool bSetupUndoRedo = true, bool bPrintPythonCommand = false);
 
 	// Adds a Parameter Node to the edited Graph given a struct object path name.
 	// Parameters represent input or output arguments to the Graph / Function.
 	// Input Parameters are constant values / literals.
 	// This causes a NodeAdded modified event.
 	UFUNCTION(BlueprintCallable, Category = RigVMController)
-	URigVMParameterNode* AddParameterNodeFromObjectPath(const FName& InParameterName, const FString& InCPPType, const FString& InCPPTypeObjectPath, bool bIsInput, const FString& InDefaultValue, const FVector2D& InPosition = FVector2D::ZeroVector, const FString& InNodeName = TEXT(""), bool bSetupUndoRedo = true);
+	URigVMParameterNode* AddParameterNodeFromObjectPath(const FName& InParameterName, const FString& InCPPType, const FString& InCPPTypeObjectPath, bool bIsInput, const FString& InDefaultValue, const FVector2D& InPosition = FVector2D::ZeroVector, const FString& InNodeName = TEXT(""), bool bSetupUndoRedo = true, bool bPrintPythonCommand = false);
 
 	// Adds a Comment Node to the edited Graph.
 	// Comments can be used to annotate the Graph.
 	// This causes a NodeAdded modified event.
 	UFUNCTION(BlueprintCallable, Category = RigVMController)
-	URigVMCommentNode* AddCommentNode(const FString& InCommentText, const FVector2D& InPosition = FVector2D::ZeroVector, const FVector2D& InSize = FVector2D(400.f, 300.f), const FLinearColor& InColor = FLinearColor::Black, const FString& InNodeName = TEXT(""), bool bSetupUndoRedo = true);
+	URigVMCommentNode* AddCommentNode(const FString& InCommentText, const FVector2D& InPosition = FVector2D::ZeroVector, const FVector2D& InSize = FVector2D(400.f, 300.f), const FLinearColor& InColor = FLinearColor::Black, const FString& InNodeName = TEXT(""), bool bSetupUndoRedo = true, bool bPrintPythonCommand = false);
 
 	// Adds a Reroute Node on an existing Link to the edited Graph.
 	// Reroute Nodes can be used to visually improve the data flow,
 	// they don't require any additional memory though and are purely
 	// cosmetic. This causes a NodeAdded modified event.
 	UFUNCTION(BlueprintCallable, Category = RigVMController)
-	URigVMRerouteNode* AddRerouteNodeOnLink(URigVMLink* InLink, bool bShowAsFullNode, const FVector2D& InPosition = FVector2D::ZeroVector, const FString& InNodeName = TEXT(""), bool bSetupUndoRedo = true);
+	URigVMRerouteNode* AddRerouteNodeOnLink(URigVMLink* InLink, bool bShowAsFullNode, const FVector2D& InPosition = FVector2D::ZeroVector, const FString& InNodeName = TEXT(""), bool bSetupUndoRedo = true, bool bPrintPythonCommand = false);
 
 	// Adds a Reroute Node on an existing Link to the edited Graph given the Link's string representation.
 	// Reroute Nodes can be used to visually improve the data flow,
 	// they don't require any additional memory though and are purely
 	// cosmetic. This causes a NodeAdded modified event.
 	UFUNCTION(BlueprintCallable, Category = RigVMController)
-	URigVMRerouteNode* AddRerouteNodeOnLinkPath(const FString& InLinkPinPathRepresentation, bool bShowAsFullNode, const FVector2D& InPosition = FVector2D::ZeroVector, const FString& InNodeName = TEXT(""), bool bSetupUndoRedo = true);
+	URigVMRerouteNode* AddRerouteNodeOnLinkPath(const FString& InLinkPinPathRepresentation, bool bShowAsFullNode, const FVector2D& InPosition = FVector2D::ZeroVector, const FString& InNodeName = TEXT(""), bool bSetupUndoRedo = true, bool bPrintPythonCommand = false);
 
 	// Adds a Reroute Node on an existing Pin to the editor Graph.
 	// Reroute Nodes can be used to visually improve the data flow,
 	// they don't require any additional memory though and are purely
 	// cosmetic. This causes a NodeAdded modified event.
 	UFUNCTION(BlueprintCallable, Category = RigVMController)
-	URigVMRerouteNode* AddRerouteNodeOnPin(const FString& InPinPath, bool bAsInput, bool bShowAsFullNode, const FVector2D& InPosition = FVector2D::ZeroVector, const FString& InNodeName = TEXT(""), bool bSetupUndoRedo = true);
+	URigVMRerouteNode* AddRerouteNodeOnPin(const FString& InPinPath, bool bAsInput, bool bShowAsFullNode, const FVector2D& InPosition = FVector2D::ZeroVector, const FString& InNodeName = TEXT(""), bool bSetupUndoRedo = true, bool bPrintPythonCommand = false);
 
 	// Adds a free Reroute Node
 	UFUNCTION(BlueprintCallable, Category = RigVMController)
@@ -238,30 +240,30 @@ public:
 	// Branch nodes can be used to split the execution of into multiple branches,
 	// allowing to drive behavior by logic.
 	UFUNCTION(BlueprintCallable, Category = RigVMController)
-	URigVMBranchNode* AddBranchNode(const FVector2D& InPosition = FVector2D::ZeroVector, const FString& InNodeName = TEXT(""), bool bSetupUndoRedo = true);
+	URigVMBranchNode* AddBranchNode(const FVector2D& InPosition = FVector2D::ZeroVector, const FString& InNodeName = TEXT(""), bool bSetupUndoRedo = true, bool bPrintPythonCommand = false);
 
 	// Adds an if node to the graph.
 	// If nodes can be used to pick between two values based on a condition.
 	UFUNCTION(BlueprintCallable, Category = RigVMController)
-	URigVMIfNode* AddIfNode(const FString& InCPPType, const FName& InCPPTypeObjectPath, const FVector2D& InPosition = FVector2D::ZeroVector, const FString& InNodeName = TEXT(""), bool bSetupUndoRedo = true);
+	URigVMIfNode* AddIfNode(const FString& InCPPType, const FName& InCPPTypeObjectPath, const FVector2D& InPosition = FVector2D::ZeroVector, const FString& InNodeName = TEXT(""), bool bSetupUndoRedo = true, bool bPrintPythonCommand = false);
 
 	// Adds a select node to the graph.
 	// Select nodes can be used to pick between multiple values based on an index.
 	UFUNCTION(BlueprintCallable, Category = RigVMController)
-	URigVMSelectNode* AddSelectNode(const FString& InCPPType, const FName& InCPPTypeObjectPath, const FVector2D& InPosition = FVector2D::ZeroVector, const FString& InNodeName = TEXT(""), bool bSetupUndoRedo = true);
+	URigVMSelectNode* AddSelectNode(const FString& InCPPType, const FName& InCPPTypeObjectPath, const FVector2D& InPosition = FVector2D::ZeroVector, const FString& InNodeName = TEXT(""), bool bSetupUndoRedo = true, bool bPrintPythonCommand = false);
 	UFUNCTION(BlueprintCallable, Category = RigVMController)
 	URigVMSelectNode* AddSelectNodeFromStruct(UScriptStruct* InScriptStruct, const FVector2D& InPosition = FVector2D::ZeroVector, const FString& InNodeName = TEXT(""), bool bSetupUndoRedo = true);
 
 
 	// Adds a prototype node to the graph.
 	UFUNCTION(BlueprintCallable, Category = RigVMController)
-	URigVMPrototypeNode* AddPrototypeNode(const FName& InNotation, const FVector2D& InPosition = FVector2D::ZeroVector, const FString& InNodeName = TEXT(""), bool bSetupUndoRedo = true);
+	URigVMPrototypeNode* AddPrototypeNode(const FName& InNotation, const FVector2D& InPosition = FVector2D::ZeroVector, const FString& InNodeName = TEXT(""), bool bSetupUndoRedo = true, bool bPrintPythonCommand = false);
 
 	// Adds a Function / Struct Node to the edited Graph as an injected node
 	// UnitNode represent a RIGVM_METHOD declaration on a USTRUCT.
 	// This causes a NodeAdded modified event.
 	UFUNCTION(BlueprintCallable, Category = RigVMController)
-	URigVMInjectionInfo* AddInjectedNode(const FString& InPinPath, bool bAsInput, UScriptStruct* InScriptStruct, const FName& InMethodName, const FName& InInputPinName, const FName& InOutputPinName, const FString& InNodeName = TEXT(""), bool bSetupUndoRedo = true);
+	URigVMInjectionInfo* AddInjectedNode(const FString& InPinPath, bool bAsInput, UScriptStruct* InScriptStruct, const FName& InMethodName, const FName& InInputPinName, const FName& InOutputPinName, const FString& InNodeName = TEXT(""), bool bSetupUndoRedo = true, bool bPrintPythonCommand = false);
 
 	// Adds a Function / Struct Node to the edited Graph as an injected node
 	// UnitNode represent a RIGVM_METHOD declaration on a USTRUCT.
@@ -271,12 +273,12 @@ public:
 
 	// Ejects the last injected node on a pin
 	UFUNCTION(BlueprintCallable, Category = RigVMController)
-	URigVMNode* EjectNodeFromPin(const FString& InPinPath, bool bSetupUndoRedo = true);
+	URigVMNode* EjectNodeFromPin(const FString& InPinPath, bool bSetupUndoRedo = true, bool bPrintPythonCommand = false);
 
 	// Adds an enum node to the graph
 	// Enum nodes can be used to represent constant enum values within the graph
 	UFUNCTION(BlueprintCallable, Category = RigVMController)
-	URigVMEnumNode* AddEnumNode(const FName& InCPPTypeObjectPath, const FVector2D& InPosition = FVector2D::ZeroVector, const FString& InNodeName = TEXT(""), bool bSetupUndoRedo = true);
+	URigVMEnumNode* AddEnumNode(const FName& InCPPTypeObjectPath, const FVector2D& InPosition = FVector2D::ZeroVector, const FString& InNodeName = TEXT(""), bool bSetupUndoRedo = true, bool bPrintPythonCommand = false);
 
 	// Un-does the last action on the stack.
 	// Note: This should really only be used for unit tests,
@@ -330,44 +332,47 @@ public:
     URigVMLibraryNode* LocalizeFunction(
     	URigVMLibraryNode* InFunctionDefinition,
 		bool bLocalizeDependentPrivateFunctions = true,
-    	bool bSetupUndoRedo = true);
+    	bool bSetupUndoRedo = true,
+    	bool bPrintPythonCommand = false);
 
 	// Copies a series of function declaratioms into this graph's local function library
+	UFUNCTION(BlueprintCallable, Category = RigVMController)
     TMap<URigVMLibraryNode*, URigVMLibraryNode*> LocalizeFunctions(
         TArray<URigVMLibraryNode*> InFunctionDefinitions,
         bool bLocalizeDependentPrivateFunctions = true,
-        bool bSetupUndoRedo = true);
+        bool bSetupUndoRedo = true,
+        bool bPrintPythonCommand = false);
 
 	// Returns a unique name
 	static FName GetUniqueName(const FName& InName, TFunction<bool(const FName&)> IsNameAvailableFunction);
 
 	// Turns a series of nodes into a Collapse node
 	UFUNCTION(BlueprintCallable, Category = RigVMController)
-	URigVMCollapseNode* CollapseNodes(const TArray<FName>& InNodeNames, const FString& InCollapseNodeName = TEXT(""), bool bSetupUndoRedo = true);
+	URigVMCollapseNode* CollapseNodes(const TArray<FName>& InNodeNames, const FString& InCollapseNodeName = TEXT(""), bool bSetupUndoRedo = true, bool bPrintPythonCommand = false);
 
 	// Turns a library node into its contained nodes
 	UFUNCTION(BlueprintCallable, Category = RigVMController)
-	TArray<URigVMNode*> ExpandLibraryNode(const FName& InNodeName, bool bSetupUndoRedo = true);
+	TArray<URigVMNode*> ExpandLibraryNode(const FName& InNodeName, bool bSetupUndoRedo = true, bool bPrintPythonCommand = false);
 
 	// Turns a collapse node into a function node
 	UFUNCTION(BlueprintCallable, Category = RigVMController)
-	FName PromoteCollapseNodeToFunctionReferenceNode(const FName& InNodeName, bool bSetupUndoRedo = true);
+	FName PromoteCollapseNodeToFunctionReferenceNode(const FName& InNodeName, bool bSetupUndoRedo = true, bool bPrintPythonCommand = false);
 
 	// Turns a collapse node into a function node
 	UFUNCTION(BlueprintCallable, Category = RigVMController)
-	FName PromoteFunctionReferenceNodeToCollapseNode(const FName& InNodeName, bool bSetupUndoRedo = true);
+	FName PromoteFunctionReferenceNodeToCollapseNode(const FName& InNodeName, bool bSetupUndoRedo = true, bool bPrintPythonCommand = false);
 
 #endif
 
 	// Removes a node from the graph
 	// This causes a NodeRemoved modified event.
 	UFUNCTION(BlueprintCallable, Category = RigVMController)
-	bool RemoveNode(URigVMNode* InNode, bool bSetupUndoRedo = true, bool bRecursive = false);
+	bool RemoveNode(URigVMNode* InNode, bool bSetupUndoRedo = true, bool bRecursive = false, bool bPrintPythonCommand = false);
 
 	// Removes a node from the graph given the node's name.
 	// This causes a NodeRemoved modified event.
 	UFUNCTION(BlueprintCallable, Category = RigVMController)
-	bool RemoveNodeByName(const FName& InNodeName, bool bSetupUndoRedo = true, bool bRecursive = false);
+	bool RemoveNodeByName(const FName& InNodeName, bool bSetupUndoRedo = true, bool bRecursive = false, bool bPrintPythonCommand = false);
 
 	// Renames a node in the graph
 	// This causes a NodeRenamed modified event.
@@ -377,7 +382,7 @@ public:
 	// Selects a single node in the graph.
 	// This causes a NodeSelected / NodeDeselected modified event.
 	UFUNCTION(BlueprintCallable, Category = RigVMController)
-	bool SelectNode(URigVMNode* InNode, bool bSelect = true, bool bSetupUndoRedo = true);
+	bool SelectNode(URigVMNode* InNode, bool bSelect = true, bool bSetupUndoRedo = true, bool bPrintPythonCommand = false);
 
 	// Selects a single node in the graph by name.
 	// This causes a NodeSelected / NodeDeselected modified event.
@@ -387,37 +392,37 @@ public:
 	// Deselects all currently selected nodes in the graph.
 	// This might cause several NodeDeselected modified event.
 	UFUNCTION(BlueprintCallable, Category = RigVMController)
-	bool ClearNodeSelection(bool bSetupUndoRedo = true);
+	bool ClearNodeSelection(bool bSetupUndoRedo = true, bool bPrintPythonCommand = false);
 
 	// Selects the nodes given the selection
 	// This might cause several NodeDeselected modified event.
 	UFUNCTION(BlueprintCallable, Category = RigVMController)
-	bool SetNodeSelection(const TArray<FName>& InNodeNames, bool bSetupUndoRedo = true);
+	bool SetNodeSelection(const TArray<FName>& InNodeNames, bool bSetupUndoRedo = true, bool bPrintPythonCommand = false);
 
 	// Sets the position of a node in the graph.
 	// This causes a NodePositionChanged modified event.
 	UFUNCTION(BlueprintCallable, Category = RigVMController)
-	bool SetNodePosition(URigVMNode* InNode, const FVector2D& InPosition, bool bSetupUndoRedo = true, bool bMergeUndoAction = false);
+	bool SetNodePosition(URigVMNode* InNode, const FVector2D& InPosition, bool bSetupUndoRedo = true, bool bMergeUndoAction = false, bool bPrintPythonCommand = false);
 
 	// Sets the position of a node in the graph by name.
 	// This causes a NodePositionChanged modified event.
 	UFUNCTION(BlueprintCallable, Category = RigVMController)
-	bool SetNodePositionByName(const FName& InNodeName, const FVector2D& InPosition, bool bSetupUndoRedo = true, bool bMergeUndoAction = false);
+	bool SetNodePositionByName(const FName& InNodeName, const FVector2D& InPosition, bool bSetupUndoRedo = true, bool bMergeUndoAction = false, bool bPrintPythonCommand = false);
 
 	// Sets the size of a node in the graph.
 	// This causes a NodeSizeChanged modified event.
 	UFUNCTION(BlueprintCallable, Category = RigVMController)
-	bool SetNodeSize(URigVMNode* InNode, const FVector2D& InSize, bool bSetupUndoRedo = true, bool bMergeUndoAction = false);
+	bool SetNodeSize(URigVMNode* InNode, const FVector2D& InSize, bool bSetupUndoRedo = true, bool bMergeUndoAction = false, bool bPrintPythonCommand = false);
 
 	// Sets the size of a node in the graph by name.
 	// This causes a NodeSizeChanged modified event.
 	UFUNCTION(BlueprintCallable, Category = RigVMController)
-	bool SetNodeSizeByName(const FName& InNodeName, const FVector2D& InSize, bool bSetupUndoRedo = true, bool bMergeUndoAction = false);
+	bool SetNodeSizeByName(const FName& InNodeName, const FVector2D& InSize, bool bSetupUndoRedo = true, bool bMergeUndoAction = false, bool bPrintPythonCommand = false);
 
 	// Sets the color of a node in the graph.
 	// This causes a NodeColorChanged modified event.
 	UFUNCTION(BlueprintCallable, Category = RigVMController)
-	bool SetNodeColor(URigVMNode* InNode, const FLinearColor& InColor, bool bSetupUndoRedo = true, bool bMergeUndoAction = false);
+	bool SetNodeColor(URigVMNode* InNode, const FLinearColor& InColor, bool bSetupUndoRedo = true, bool bMergeUndoAction = false, bool bPrintPythonCommand = false);
 
 	// Sets the color of a node in the graph by name.
 	// This causes a NodeColorChanged modified event.
@@ -427,7 +432,7 @@ public:
 	// Sets the category of a node in the graph.
 	// This causes a NodeCategoryChanged modified event.
 	UFUNCTION(BlueprintCallable, Category = RigVMController)
-	bool SetNodeCategory(URigVMCollapseNode* InNode, const FString& InCategory, bool bSetupUndoRedo = true, bool bMergeUndoAction = false);
+	bool SetNodeCategory(URigVMCollapseNode* InNode, const FString& InCategory, bool bSetupUndoRedo = true, bool bMergeUndoAction = false, bool bPrintPythonCommand = false);
 
 	// Sets the category of a node in the graph.
 	// This causes a NodeCategoryChanged modified event.
@@ -437,7 +442,7 @@ public:
 	// Sets the keywords of a node in the graph.
 	// This causes a NodeKeywordsChanged modified event.
 	UFUNCTION(BlueprintCallable, Category = RigVMController)
-	bool SetNodeKeywords(URigVMCollapseNode* InNode, const FString& InKeywords, bool bSetupUndoRedo = true, bool bMergeUndoAction = false);
+	bool SetNodeKeywords(URigVMCollapseNode* InNode, const FString& InKeywords, bool bSetupUndoRedo = true, bool bMergeUndoAction = false, bool bPrintPythonCommand = false);
 
 	// Sets the keywords of a node in the graph.
 	// This causes a NodeKeywordsChanged modified event.
@@ -447,7 +452,7 @@ public:
 	// Sets the function description of a node in the graph.
 	// This causes a NodeDescriptionChanged modified event.
 	UFUNCTION(BlueprintCallable, Category = RigVMController)
-	bool SetNodeDescription(URigVMCollapseNode* InNode, const FString& InDescription, bool bSetupUndoRedo = true, bool bMergeUndoAction = false);
+	bool SetNodeDescription(URigVMCollapseNode* InNode, const FString& InDescription, bool bSetupUndoRedo = true, bool bMergeUndoAction = false, bool bPrintPythonCommand = false);
 
 	// Sets the keywords of a node in the graph.
 	// This causes a NodeDescriptionChanged modified event.
@@ -457,17 +462,17 @@ public:
 	// Sets the comment text of a comment node in the graph.
 	// This causes a CommentTextChanged modified event.
 	UFUNCTION(BlueprintCallable, Category = RigVMController)
-	bool SetCommentText(URigVMNode* InNode, const FString& InCommentText, bool bSetupUndoRedo = true);
+	bool SetCommentText(URigVMNode* InNode, const FString& InCommentText, bool bSetupUndoRedo = true, bool bPrintPythonCommand = false);
 
 	// Sets the comment text of a comment node in the graph by name.
 	// This causes a CommentTextChanged modified event.
 	UFUNCTION(BlueprintCallable, Category = RigVMController)
-	bool SetCommentTextByName(const FName& InNodeName, const FString& InCommentText, bool bSetupUndoRedo = true);
+	bool SetCommentTextByName(const FName& InNodeName, const FString& InCommentText, bool bSetupUndoRedo = true, bool bPrintPythonCommand = false);
 
 	// Sets the compactness of a reroute node in the graph.
 	// This causes a RerouteCompactnessChanged modified event.
 	UFUNCTION(BlueprintCallable, Category = RigVMController)
-	bool SetRerouteCompactness(URigVMNode* InNode, bool bShowAsFullNode, bool bSetupUndoRedo = true);
+	bool SetRerouteCompactness(URigVMNode* InNode, bool bShowAsFullNode, bool bSetupUndoRedo = true, bool bPrintPythonCommand = false);
 
 	// Sets the compactness of a reroute node in the graph by name.
 	// This causes a RerouteCompactnessChanged modified event.
@@ -487,7 +492,7 @@ public:
 	// Sets the pin to be expanded or not
 	// This causes a PinExpansionChanged modified event.
 	UFUNCTION(BlueprintCallable, Category = RigVMController)
-	bool SetPinExpansion(const FString& InPinPath, bool bIsExpanded, bool bSetupUndoRedo = true);
+	bool SetPinExpansion(const FString& InPinPath, bool bIsExpanded, bool bSetupUndoRedo = true, bool bPrintPythonCommand = false);
 
 	// Sets the pin to be watched (or not)
 	// This causes a PinWatchedChanged modified event.
@@ -501,103 +506,103 @@ public:
 	// Sets the default value of a pin given its pinpath.
 	// This causes a PinDefaultValueChanged modified event.
 	UFUNCTION(BlueprintCallable, Category = RigVMController)
-	bool SetPinDefaultValue(const FString& InPinPath, const FString& InDefaultValue, bool bResizeArrays = true, bool bSetupUndoRedo = true, bool bMergeUndoAction = false);
+	bool SetPinDefaultValue(const FString& InPinPath, const FString& InDefaultValue, bool bResizeArrays = true, bool bSetupUndoRedo = true, bool bMergeUndoAction = false, bool bPrintPythonCommand = false);
 
 	// Resets the default value of a pin given its pinpath.
 	// This causes a PinDefaultValueChanged modified event.
 	UFUNCTION(BlueprintCallable, Category = RigVMController)
-	bool ResetPinDefaultValue(const FString& InPinPath, bool bSetupUndoRedo = true);
+	bool ResetPinDefaultValue(const FString& InPinPath, bool bSetupUndoRedo = true, bool bPrintPythonCommand = false);
 
 	// Adds an array element pin to the end of an array pin.
 	// This causes a PinArraySizeChanged modified event.
 	UFUNCTION(BlueprintCallable, Category = RigVMController)
-	FString AddArrayPin(const FString& InArrayPinPath, const FString& InDefaultValue = TEXT(""), bool bSetupUndoRedo = true);
+	FString AddArrayPin(const FString& InArrayPinPath, const FString& InDefaultValue = TEXT(""), bool bSetupUndoRedo = true, bool bPrintPythonCommand = false);
 
 	// Duplicates an array element pin.
 	// This causes a PinArraySizeChanged modified event.
 	UFUNCTION(BlueprintCallable, Category = RigVMController)
-	FString DuplicateArrayPin(const FString& InArrayElementPinPath, bool bSetupUndoRedo = true);
+	FString DuplicateArrayPin(const FString& InArrayElementPinPath, bool bSetupUndoRedo = true, bool bPrintPythonCommand = false);
 
 	// Inserts an array element pin into an array pin.
 	// This causes a PinArraySizeChanged modified event.
 	UFUNCTION(BlueprintCallable, Category = RigVMController)
-	FString InsertArrayPin(const FString& InArrayPinPath, int32 InIndex = -1, const FString& InDefaultValue = TEXT(""), bool bSetupUndoRedo = true);
+	FString InsertArrayPin(const FString& InArrayPinPath, int32 InIndex = -1, const FString& InDefaultValue = TEXT(""), bool bSetupUndoRedo = true, bool bPrintPythonCommand = false);
 
 	// Removes an array element pin from an array pin.
 	// This causes a PinArraySizeChanged modified event.
 	UFUNCTION(BlueprintCallable, Category = RigVMController)
-	bool RemoveArrayPin(const FString& InArrayElementPinPath, bool bSetupUndoRedo = true);
+	bool RemoveArrayPin(const FString& InArrayElementPinPath, bool bSetupUndoRedo = true, bool bPrintPythonCommand = false);
 
 	// Removes all (but one) array element pin from an array pin.
 	// This causes a PinArraySizeChanged modified event.
 	UFUNCTION(BlueprintCallable, Category = RigVMController)
-	bool ClearArrayPin(const FString& InArrayPinPath, bool bSetupUndoRedo = true);
+	bool ClearArrayPin(const FString& InArrayPinPath, bool bSetupUndoRedo = true, bool bPrintPythonCommand = false);
 
 	// Sets the size of the array pin
 	// This causes a PinArraySizeChanged modified event.
 	UFUNCTION(BlueprintCallable, Category = RigVMController)
-	bool SetArrayPinSize(const FString& InArrayPinPath, int32 InSize, const FString& InDefaultValue = TEXT(""), bool bSetupUndoRedo = true);
+	bool SetArrayPinSize(const FString& InArrayPinPath, int32 InSize, const FString& InDefaultValue = TEXT(""), bool bSetupUndoRedo = true, bool bPrintPythonCommand = false);
 
 	// Binds a pin to a variable (or removes the binding given NAME_None)
 	// This causes a PinBoundVariableChanged modified event.
 	UFUNCTION(BlueprintCallable, Category = RigVMController)
-	bool BindPinToVariable(const FString& InPinPath, const FString& InNewBoundVariablePath, bool bSetupUndoRedo = true);
+	bool BindPinToVariable(const FString& InPinPath, const FString& InNewBoundVariablePath, bool bSetupUndoRedo = true, bool bPrintPythonCommand = false);
 
 	// Removes the binging of a pin to a variable
 	// This causes a PinBoundVariableChanged modified event.
 	UFUNCTION(BlueprintCallable, Category = RigVMController)
-	bool UnbindPinFromVariable(const FString& InPinPath, bool bSetupUndoRedo = true);
+	bool UnbindPinFromVariable(const FString& InPinPath, bool bSetupUndoRedo = true, bool bPrintPythonCommand = false);
 
 	// Turns a variable node into one or more bindings
 	UFUNCTION(BlueprintCallable, Category = RigVMController)
-	bool MakeBindingsFromVariableNode(const FName& InNodeName, bool bSetupUndoRedo = true);
+	bool MakeBindingsFromVariableNode(const FName& InNodeName, bool bSetupUndoRedo = true, bool bPrintPythonCommand = false);
 
 	// Turns a binding to a variable node
 	UFUNCTION(BlueprintCallable, Category = RigVMController)
-	bool MakeVariableNodeFromBinding(const FString& InPinPath, const FVector2D& InNodePosition = FVector2D::ZeroVector, bool bSetupUndoRedo = true);
+	bool MakeVariableNodeFromBinding(const FString& InPinPath, const FVector2D& InNodePosition = FVector2D::ZeroVector, bool bSetupUndoRedo = true, bool bPrintPythonCommand = false);
 
 	// Promotes a pin to a variable
 	UFUNCTION(BlueprintCallable, Category = RigVMController)
-	bool PromotePinToVariable(const FString& InPinPath, bool bCreateVariableNode, const FVector2D& InNodePosition = FVector2D::ZeroVector, bool bSetupUndoRedo = true);
+	bool PromotePinToVariable(const FString& InPinPath, bool bCreateVariableNode, const FVector2D& InNodePosition = FVector2D::ZeroVector, bool bSetupUndoRedo = true, bool bPrintPythonCommand = false);
 
 	// Adds a link to the graph.
 	// This causes a LinkAdded modified event.
 	UFUNCTION(BlueprintCallable, Category = RigVMController)
-	bool AddLink(const FString& InOutputPinPath, const FString& InInputPinPath, bool bSetupUndoRedo = true);
+	bool AddLink(const FString& InOutputPinPath, const FString& InInputPinPath, bool bSetupUndoRedo = true, bool bPrintPythonCommand = false);
 
 	// Removes a link from the graph.
 	// This causes a LinkRemoved modified event.
 	UFUNCTION(BlueprintCallable, Category = RigVMController)
-	bool BreakLink(const FString& InOutputPinPath, const FString& InInputPinPath, bool bSetupUndoRedo = true);
+	bool BreakLink(const FString& InOutputPinPath, const FString& InInputPinPath, bool bSetupUndoRedo = true, bool bPrintPythonCommand = false);
 
 	// Removes all links on a given pin from the graph.
 	// This might cause multiple LinkRemoved modified event.
 	UFUNCTION(BlueprintCallable, Category = RigVMController)
-	bool BreakAllLinks(const FString& InPinPath, bool bAsInput = true, bool bSetupUndoRedo = true);
+	bool BreakAllLinks(const FString& InPinPath, bool bAsInput = true, bool bSetupUndoRedo = true, bool bPrintPythonCommand = false);
 
 	// Adds an exposed pin to the graph controlled by this
 	UFUNCTION(BlueprintCallable, Category = RigVMController)
-	FName AddExposedPin(const FName& InPinName, ERigVMPinDirection InDirection, const FString& InCPPType, const FName& InCPPTypeObjectPath, const FString& InDefaultValue, bool bSetupUndoRedo = true);
+	FName AddExposedPin(const FName& InPinName, ERigVMPinDirection InDirection, const FString& InCPPType, const FName& InCPPTypeObjectPath, const FString& InDefaultValue, bool bSetupUndoRedo = true, bool bPrintPythonCommand = false);
 
 	// Removes an exposed pin from the graph controlled by this
 	UFUNCTION(BlueprintCallable, Category = RigVMController)
-	bool RemoveExposedPin(const FName& InPinName, bool bSetupUndoRedo = true);
+	bool RemoveExposedPin(const FName& InPinName, bool bSetupUndoRedo = true, bool bPrintPythonCommand = false);
 
 	// Renames an exposed pin in the graph controlled by this
 	UFUNCTION(BlueprintCallable, Category = RigVMController)
-	bool RenameExposedPin(const FName& InOldPinName, const FName& InNewPinName, bool bSetupUndoRedo = true);
+	bool RenameExposedPin(const FName& InOldPinName, const FName& InNewPinName, bool bSetupUndoRedo = true, bool bPrintPythonCommand = false);
 
 	// Changes the type of an exposed pin in the graph controlled by this
 	UFUNCTION(BlueprintCallable, Category = RigVMController)
-	bool ChangeExposedPinType(const FName& InPinName, const FString& InCPPType, const FName& InCPPTypeObjectPath, UPARAM(ref) bool& bSetupUndoRedo, bool bSetupOrphanPins = true);
+	bool ChangeExposedPinType(const FName& InPinName, const FString& InCPPType, const FName& InCPPTypeObjectPath, UPARAM(ref) bool& bSetupUndoRedo, bool bSetupOrphanPins = true, bool bPrintPythonCommand = false);
 
 	// Sets the index for an exposed pin. This can be used to move the pin up and down on the node.
 	UFUNCTION(BlueprintCallable, Category = RigVMController)
-	bool SetExposedPinIndex(const FName& InPinName, int32 InNewIndex, bool bSetupUndoRedo = true);
+	bool SetExposedPinIndex(const FName& InPinName, int32 InNewIndex, bool bSetupUndoRedo = true, bool bPrintPythonCommand = false);
 
 	// Adds a function reference / invocation to the graph
 	UFUNCTION(BlueprintCallable, Category = RigVMController)
-	URigVMFunctionReferenceNode* AddFunctionReferenceNode(URigVMLibraryNode* InFunctionDefinition, const FVector2D& InNodePosition = FVector2D::ZeroVector, const FString& InNodeName = TEXT(""), bool bSetupUndoRedo = true);
+	URigVMFunctionReferenceNode* AddFunctionReferenceNode(URigVMLibraryNode* InFunctionDefinition, const FVector2D& InNodePosition = FVector2D::ZeroVector, const FString& InNodeName = TEXT(""), bool bSetupUndoRedo = true, bool bPrintPythonCommand = false);
 
 	// Sets the remapped variable on a function reference node
 	UFUNCTION(BlueprintCallable, Category = RigVMController)
@@ -605,7 +610,7 @@ public:
 
 	// Adds a function definition to a function library graph
 	UFUNCTION(BlueprintCallable, Category = RigVMController)
-	URigVMLibraryNode* AddFunctionToLibrary(const FName& InFunctionName, bool bMutable, const FVector2D& InNodePosition = FVector2D::ZeroVector, bool bSetupUndoRedo = true);
+	URigVMLibraryNode* AddFunctionToLibrary(const FName& InFunctionName, bool bMutable, const FVector2D& InNodePosition = FVector2D::ZeroVector, bool bSetupUndoRedo = true, bool bPrintPythonCommand = false);
 
 	// Removes a function from a function library graph
 	UFUNCTION(BlueprintCallable, Category = RigVMController)
@@ -613,7 +618,7 @@ public:
 
 	// Add a local variable to the graph
 	UFUNCTION(BlueprintCallable, Category = RigVMController)
-	FRigVMGraphVariableDescription AddLocalVariable(const FName& InVariableName, const FString& InCPPType, UObject* InCPPTypeObject, const FString& InDefaultValue, bool bSetupUndoRedo = true);
+	FRigVMGraphVariableDescription AddLocalVariable(const FName& InVariableName, const FString& InCPPType, UObject* InCPPTypeObject, const FString& InDefaultValue, bool bSetupUndoRedo = true, bool bPrintPythonCommand = false);
 
 	// Add a local variable to the graph given a struct object path name.
 	UFUNCTION(BlueprintCallable, Category = RigVMController)
@@ -621,15 +626,17 @@ public:
 
 	// Remove a local variable from the graph
 	UFUNCTION(BlueprintCallable, Category = RigVMController)
-	bool RemoveLocalVariable(const FName& InVariableName, bool bSetupUndoRedo = true);
+	bool RemoveLocalVariable(const FName& InVariableName, bool bSetupUndoRedo = true, bool bPrintPythonCommand = false);
 
 	// Rename a local variable from the graph
 	UFUNCTION(BlueprintCallable, Category = RigVMController)
-	bool RenameLocalVariable(const FName& InVariableName, const FName& InNewVariableName, bool bSetupUndoRedo = true);
+	bool RenameLocalVariable(const FName& InVariableName, const FName& InNewVariableName, bool bSetupUndoRedo = true, bool bPrintPythonCommand = false);
 
 	// Sets the type of the local variable
 	UFUNCTION(BlueprintCallable, Category = RigVMController)
-	bool SetLocalVariableType(const FName& InVariableName, const FString& InCPPType, UObject* InCPPTypeObject, bool bSetupUndoRedo = true);
+	bool SetLocalVariableType(const FName& InVariableName, const FString& InCPPType, UObject* InCPPTypeObject, bool bSetupUndoRedo = true, bool bPrintPythonCommand = false);
+	UFUNCTION(BlueprintCallable, Category = RigVMController)
+	bool SetLocalVariableTypeFromObjectPath(const FName& InVariableName, const FString& InCPPType, const FString& InCPPTypeObjectPath, bool bSetupUndoRedo = true, bool bPrintPythonCommand = false);
 
 	// Determine affected function references for a potential bulk edit on a library node
 	TArray<TSoftObjectPtr<URigVMFunctionReferenceNode>> GetAffectedReferences(ERigVMControllerBulkEditType InEditType, bool bForceLoad = false, bool bNotify = true);
diff --git a/Engine/Source/Developer/RigVMDeveloper/Public/RigVMPythonUtils.h b/Engine/Source/Developer/RigVMDeveloper/Public/RigVMPythonUtils.h
new file mode 100644
--- /dev/null
+++ b/Engine/Source/Developer/RigVMDeveloper/Public/RigVMPythonUtils.h
@@ -0,0 +1,21 @@
+﻿// Copyright Epic Games, Inc. All Rights Reserved.
+
+#pragma once
+
+#include "RigVMDeveloperModule.h"
+
+namespace RigVMPythonUtils
+{
+	RIGVMDEVELOPER_API FString NameToPep8(const FString& Name);
+
+	RIGVMDEVELOPER_API FString TransformToPythonString(const FTransform& Transform);
+
+	RIGVMDEVELOPER_API FString Vector2DToPythonString(const FVector2D& Vector);
+
+	RIGVMDEVELOPER_API FString LinearColorToPythonString(const FLinearColor& Color);
+	
+#if WITH_EDITOR
+	RIGVMDEVELOPER_API void Print(const FString& InMessage);
+
+#endif
+}
diff --git a/Engine/Source/Developer/RigVMDeveloper/RigVMDeveloper.Build.cs b/Engine/Source/Developer/RigVMDeveloper/RigVMDeveloper.Build.cs
--- a/Engine/Source/Developer/RigVMDeveloper/RigVMDeveloper.Build.cs
+++ b/Engine/Source/Developer/RigVMDeveloper/RigVMDeveloper.Build.cs
@@ -25,6 +25,7 @@ public class RigVMDeveloper : ModuleRules
 					"Slate",
 					"SlateCore",
 					"EditorStyle",
+					"MessageLog",
 				}
 			);
         }
diff --git a/Engine/Source/Runtime/RigVM/Public/RigVMCore/RigVMPythonUtils.h b/Engine/Source/Runtime/RigVM/Public/RigVMCore/RigVMPythonUtils.h
deleted file mode 100644
