Control Rig: Fix Fornite compilation errors

diff --git a/Engine/Source/Developer/RigVMDeveloper/Private/RigVMPythonUtils.cpp b/Engine/Source/Developer/RigVMDeveloper/Private/RigVMPythonUtils.cpp
--- a/Engine/Source/Developer/RigVMDeveloper/Private/RigVMPythonUtils.cpp
+++ b/Engine/Source/Developer/RigVMDeveloper/Private/RigVMPythonUtils.cpp
@@ -3,6 +3,7 @@
 #include "RigVMPythonUtils.h"
 
 #if WITH_EDITOR
+#include "Modules/ModuleManager.h"
 #include "MessageLogModule.h"
 #include "IMessageLogListing.h"
 #endif
@@ -65,7 +66,6 @@ FString RigVMPythonUtils::LinearColorToPythonString(const FLinearColor& Color)
 #if WITH_EDITOR
 void RigVMPythonUtils::Print(const FString& InMessage)
 {
-	//UE_LOG(LogRigVMDeveloper, Display, TEXT("%s"), *InMessage);
 	FMessageLogModule& MessageLogModule = FModuleManager::LoadModuleChecked<FMessageLogModule>("MessageLog");
 
 	if (!MessageLogModule.IsRegisteredLogListing("RigVMPythonLog"))
