Control Rig: Print python commands when removing and renaming functions

diff --git a/Engine/Plugins/Animation/ControlRig/Source/ControlRigDeveloper/Private/Graph/ControlRigGraphSchema.cpp b/Engine/Plugins/Animation/ControlRig/Source/ControlRigDeveloper/Private/Graph/ControlRigGraphSchema.cpp
--- a/Engine/Plugins/Animation/ControlRig/Source/ControlRigDeveloper/Private/Graph/ControlRigGraphSchema.cpp
+++ b/Engine/Plugins/Animation/ControlRig/Source/ControlRigDeveloper/Private/Graph/ControlRigGraphSchema.cpp
@@ -949,7 +949,7 @@ bool UControlRigGraphSchema::TryDeleteGraph(UEdGraph* GraphToDelete) const
 							}
 						}
 						
-						return Controller->RemoveNode(LibraryNode);
+						return Controller->RemoveNode(LibraryNode, true, false, true);
 					}
 				}
 			}
@@ -976,7 +976,7 @@ bool UControlRigGraphSchema::TryRenameGraph(UEdGraph* GraphToRename, const FName
 				{
 					if (URigVMController* Controller = RigBlueprint->GetOrCreateController(LibraryNode->GetGraph()))
 					{
-						Controller->RenameNode(LibraryNode, InNewName);
+						Controller->RenameNode(LibraryNode, InNewName, true, true);
 						return true;
 					}
 				}
diff --git a/Engine/Source/Developer/RigVMDeveloper/Private/RigVMModel/RigVMController.cpp b/Engine/Source/Developer/RigVMDeveloper/Private/RigVMModel/RigVMController.cpp
--- a/Engine/Source/Developer/RigVMDeveloper/Private/RigVMModel/RigVMController.cpp
+++ b/Engine/Source/Developer/RigVMDeveloper/Private/RigVMModel/RigVMController.cpp
@@ -4328,6 +4328,21 @@ bool URigVMController::RemoveNode(URigVMNode* InNode, bool bSetupUndoRedo, bool
 		}
 	}
 
+	if (bPrintPythonCommand)
+	{
+		if (Graph->IsA<URigVMFunctionLibrary>())
+		{
+			RigVMPythonUtils::Print(FString::Printf(TEXT("library_controller.remove_function_from_library('%s')"),
+				*InNode->GetFName().ToString()));
+		}
+		else
+		{
+			RigVMPythonUtils::Print(FString::Printf(TEXT("blueprint.get_controller_by_name('%s').remove_node_by_name('%s')"),
+				*GetGraph()->GetGraphName(),
+				*InNode->GetNodePath()));
+		}
+	}
+
 	if (URigVMVariableNode* VariableNode = Cast<URigVMVariableNode>(InNode))
 	{
 		Notify(ERigVMGraphNotifType::VariableRemoved, VariableNode);
@@ -4349,13 +4364,6 @@ bool URigVMController::RemoveNode(URigVMNode* InNode, bool bSetupUndoRedo, bool
 		ActionStack->EndAction(Action);
 	}
 
-	if (bPrintPythonCommand)
-	{
-		RigVMPythonUtils::Print(FString::Printf(TEXT("blueprint.get_controller_by_name('%s').remove_node_by_name('%s')"),
-				*GetGraph()->GetGraphName(),
-				*InNode->GetNodePath()));
-	}
-
 	return true;
 }
 
@@ -4372,7 +4380,7 @@ bool URigVMController::RemoveNodeByName(const FName& InNodeName, bool bSetupUndo
 	return RemoveNode(Graph->FindNodeByName(InNodeName), bSetupUndoRedo, bRecursive, bPrintPythonCommand);
 }
 
-bool URigVMController::RenameNode(URigVMNode* InNode, const FName& InNewName, bool bSetupUndoRedo)
+bool URigVMController::RenameNode(URigVMNode* InNode, const FName& InNewName, bool bSetupUndoRedo, bool bPrintPythonCommand)
 {
 	if (!IsValidNodeForGraph(InNode))
 	{
@@ -4385,6 +4393,7 @@ bool URigVMController::RenameNode(URigVMNode* InNode, const FName& InNewName, bo
 		return false;
 	}
 
+	const FString OldName = InNode->GetName();
 	FRigVMRenameNodeAction Action;
 	if (bSetupUndoRedo)
 	{
@@ -4433,6 +4442,13 @@ bool URigVMController::RenameNode(URigVMNode* InNode, const FName& InNewName, bo
 		ActionStack->EndAction(Action);
 	}
 
+	if (GetGraph()->IsA<URigVMFunctionLibrary>())
+	{
+		RigVMPythonUtils::Print(FString::Printf(TEXT("library_controller.rename_function('%s', '%s')"),
+			*OldName,
+			*InNewName.ToString()));
+	}
+
 	return true;
 }
 
@@ -7457,7 +7473,6 @@ URigVMLibraryNode* URigVMController::AddFunctionToLibrary(const FName& InFunctio
 	{
 		//AddFunctionToLibrary(const FName& InFunctionName, bool bMutable, const FVector2D& InNodePosition, bool bSetupUndoRedo, bool bPrintPythonCommand)
 		RigVMPythonUtils::Print(FString::Printf(TEXT("library_controller.add_function_to_library('%s', %s, %s)"),
-				*GetGraph()->GetGraphName(),
 				*InFunctionName.ToString(),
 				(bMutable) ? TEXT("True") : TEXT("False"),
 				*RigVMPythonUtils::Vector2DToPythonString(InNodePosition)));
@@ -7486,6 +7501,33 @@ bool URigVMController::RemoveFunctionFromLibrary(const FName& InFunctionName, bo
 	return RemoveNodeByName(InFunctionName, bSetupUndoRedo);
 }
 
+bool URigVMController::RenameFunction(const FName& InOldFunctionName, const FName& InNewFunctionName, bool bSetupUndoRedo)
+{
+	if (!IsValidGraph())
+	{
+		return false;
+	}
+
+	URigVMGraph* Graph = GetGraph();
+	check(Graph);
+
+
+	if (!Graph->IsA<URigVMFunctionLibrary>())
+	{
+		ReportError(TEXT("Can only remove function definitions from function library graphs."));
+		return false;
+	}
+
+	URigVMNode* Node = Graph->FindNode(InOldFunctionName.ToString());
+	if (!Node)
+	{
+		ReportErrorf(TEXT("Could not find function called '%s'."), *InOldFunctionName.ToString());
+		return false;
+	}
+
+	return RenameNode(Node, InNewFunctionName, bSetupUndoRedo);
+}
+
 FRigVMGraphVariableDescription URigVMController::AddLocalVariable(const FName& InVariableName, const FString& InCPPType, UObject* InCPPTypeObject, const FString& InDefaultValue, bool bSetupUndoRedo, bool bPrintPythonCommand)
 {
 	FRigVMGraphVariableDescription NewVariable;
diff --git a/Engine/Source/Developer/RigVMDeveloper/Public/RigVMModel/RigVMController.h b/Engine/Source/Developer/RigVMDeveloper/Public/RigVMModel/RigVMController.h
--- a/Engine/Source/Developer/RigVMDeveloper/Public/RigVMModel/RigVMController.h
+++ b/Engine/Source/Developer/RigVMDeveloper/Public/RigVMModel/RigVMController.h
@@ -377,7 +377,7 @@ public:
 	// Renames a node in the graph
 	// This causes a NodeRenamed modified event.
 	UFUNCTION(BlueprintCallable, Category = RigVMController)
-	bool RenameNode(URigVMNode* InNode, const FName& InNewName, bool bSetupUndoRedo = true);
+	bool RenameNode(URigVMNode* InNode, const FName& InNewName, bool bSetupUndoRedo = true, bool bPrintPythonCommand = false);
 
 	// Selects a single node in the graph.
 	// This causes a NodeSelected / NodeDeselected modified event.
@@ -616,6 +616,10 @@ public:
 	UFUNCTION(BlueprintCallable, Category = RigVMController)
 	bool RemoveFunctionFromLibrary(const FName& InFunctionName, bool bSetupUndoRedo = true);
 
+	// Renames a function in the function library
+	UFUNCTION(BlueprintCallable, Category = RigVMController)
+	bool RenameFunction(const FName& InOldFunctionName, const FName& InNewFunctionName, bool bSetupUndoRedo = true);
+
 	// Add a local variable to the graph
 	UFUNCTION(BlueprintCallable, Category = RigVMController)
 	FRigVMGraphVariableDescription AddLocalVariable(const FName& InVariableName, const FString& InCPPType, UObject* InCPPTypeObject, const FString& InDefaultValue, bool bSetupUndoRedo = true, bool bPrintPythonCommand = false);
