Control Rig: Guard against nullptr crash when reading orphan pins

diff --git a/Engine/Source/Developer/RigVMDeveloper/Private/RigVMModel/RigVMPin.cpp b/Engine/Source/Developer/RigVMDeveloper/Private/RigVMModel/RigVMPin.cpp
--- a/Engine/Source/Developer/RigVMDeveloper/Private/RigVMModel/RigVMPin.cpp
+++ b/Engine/Source/Developer/RigVMDeveloper/Private/RigVMModel/RigVMPin.cpp
@@ -876,7 +876,11 @@ bool URigVMPin::IsOrphanPin() const
 			return RootPin->IsOrphanPin();
 		}
 	}
-	return GetNode()->OrphanedPins.Contains(this); 
+	if(URigVMNode* Node = GetNode())
+	{
+		return Node->OrphanedPins.Contains(this);
+	}
+	return false;
 }
 
 void URigVMPin::UpdateCPPTypeObjectIfRequired() const
