Control Rig: Python log copy/pasting nodes

diff --git a/Engine/Plugins/Animation/ControlRig/Source/ControlRigDeveloper/Private/ControlRigBlueprint.cpp b/Engine/Plugins/Animation/ControlRig/Source/ControlRigDeveloper/Private/ControlRigBlueprint.cpp
--- a/Engine/Plugins/Animation/ControlRig/Source/ControlRigDeveloper/Private/ControlRigBlueprint.cpp
+++ b/Engine/Plugins/Animation/ControlRig/Source/ControlRigDeveloper/Private/ControlRigBlueprint.cpp
@@ -227,7 +227,7 @@ bool UControlRigBlueprint::TryImportGraphFromText(const FString& InClipboardText
             })
         );
 		
-		TArray<FName> ImportedNodeNames = FunctionLibraryController->ImportNodesFromText(InClipboardText, true);
+		TArray<FName> ImportedNodeNames = FunctionLibraryController->ImportNodesFromText(InClipboardText, true, true);
 		if (ImportedNodeNames.Num() == 0)
 		{
 			return false;
diff --git a/Engine/Plugins/Animation/ControlRig/Source/ControlRigEditor/Private/Editor/ControlRigEditor.cpp b/Engine/Plugins/Animation/ControlRig/Source/ControlRigEditor/Private/Editor/ControlRigEditor.cpp
--- a/Engine/Plugins/Animation/ControlRig/Source/ControlRigEditor/Private/Editor/ControlRigEditor.cpp
+++ b/Engine/Plugins/Animation/ControlRig/Source/ControlRigEditor/Private/Editor/ControlRigEditor.cpp
@@ -2166,7 +2166,7 @@ void FControlRigEditor::PasteNodes()
 		})
 	);
 	
-	TArray<FName> NodeNames = GetFocusedController()->ImportNodesFromText(TextToImport);
+	TArray<FName> NodeNames = GetFocusedController()->ImportNodesFromText(TextToImport, true, true);
 
 	if (NodeNames.Num() > 0)
 	{
diff --git a/Engine/Source/Developer/RigVMDeveloper/Private/RigVMModel/RigVMController.cpp b/Engine/Source/Developer/RigVMDeveloper/Private/RigVMModel/RigVMController.cpp
--- a/Engine/Source/Developer/RigVMDeveloper/Private/RigVMModel/RigVMController.cpp
+++ b/Engine/Source/Developer/RigVMDeveloper/Private/RigVMModel/RigVMController.cpp
@@ -2326,7 +2326,7 @@ bool URigVMController::CanImportNodesFromText(const FString& InText)
 	return Factory.CanCreateObjectsFromText(InText);
 }
 
-TArray<FName> URigVMController::ImportNodesFromText(const FString& InText, bool bSetupUndoRedo)
+TArray<FName> URigVMController::ImportNodesFromText(const FString& InText, bool bSetupUndoRedo, bool bPrintPythonCommands)
 {
 	TArray<FName> NodeNames;
 	if (!IsValidGraph())
@@ -2473,6 +2473,19 @@ TArray<FName> URigVMController::ImportNodesFromText(const FString& InText, bool
 	{
 		CloseUndoBracket();
 	}
+	
+#if WITH_EDITOR
+	if (bPrintPythonCommands && !NodeNames.IsEmpty())
+	{
+		FString PythonContent = InText.Replace(TEXT("\\\""), TEXT("\\\\\""));
+		PythonContent = InText.Replace(TEXT("'"), TEXT("\\'"));
+		PythonContent = PythonContent.Replace(TEXT("\r\n"), TEXT("\\r\\n'\r\n'"));
+		
+		RigVMPythonUtils::Print(FString::Printf(TEXT("blueprint.get_controller_by_name('%s').import_nodes_from_text('%s')"),
+			*GetGraph()->GetGraphName(),
+			*PythonContent));
+	}
+#endif
 
 	return NodeNames;
 }
diff --git a/Engine/Source/Developer/RigVMDeveloper/Private/RigVMPythonUtils.cpp b/Engine/Source/Developer/RigVMDeveloper/Private/RigVMPythonUtils.cpp
--- a/Engine/Source/Developer/RigVMDeveloper/Private/RigVMPythonUtils.cpp
+++ b/Engine/Source/Developer/RigVMDeveloper/Private/RigVMPythonUtils.cpp
@@ -68,15 +68,15 @@ void RigVMPythonUtils::Print(const FString& InMessage)
 {
 	FMessageLogModule& MessageLogModule = FModuleManager::LoadModuleChecked<FMessageLogModule>("MessageLog");
 
-	if (!MessageLogModule.IsRegisteredLogListing("RigVMPythonLog"))
+	if (!MessageLogModule.IsRegisteredLogListing("ControlRigPythonLog"))
 	{
 		FMessageLogInitializationOptions InitOptions;
 		InitOptions.bShowFilters = true;
 		InitOptions.bShowPages = true;
 		InitOptions.bAllowClear = true;
-		MessageLogModule.RegisterLogListing("RigVMPythonLog", LOCTEXT("RigVMPythonLog", "RigVM Python Log"), InitOptions);
+		MessageLogModule.RegisterLogListing("ControlRigPythonLog", LOCTEXT("ControlRigPythonLog", "Control Rig Python Log"), InitOptions);
 	}
-	TSharedRef<IMessageLogListing> PythonLog = MessageLogModule.GetLogListing( TEXT("RigVMPythonLog") );
+	TSharedRef<IMessageLogListing> PythonLog = MessageLogModule.GetLogListing( TEXT("ControlRigPythonLog") );
 
 	TSharedRef<FTokenizedMessage> Token = FTokenizedMessage::Create(EMessageSeverity::Info, FText::FromString(InMessage));
 	PythonLog->AddMessage(Token, false);			
diff --git a/Engine/Source/Developer/RigVMDeveloper/Public/RigVMModel/RigVMController.h b/Engine/Source/Developer/RigVMDeveloper/Public/RigVMModel/RigVMController.h
--- a/Engine/Source/Developer/RigVMDeveloper/Public/RigVMModel/RigVMController.h
+++ b/Engine/Source/Developer/RigVMDeveloper/Public/RigVMModel/RigVMController.h
@@ -325,7 +325,7 @@ public:
 
 	// Exports the given nodes as text
 	UFUNCTION(BlueprintCallable, Category = RigVMController)
-	TArray<FName> ImportNodesFromText(const FString& InText, bool bSetupUndoRedo = true);
+	TArray<FName> ImportNodesFromText(const FString& InText, bool bSetupUndoRedo = true, bool bPrintPythonCommands = false);
 
 	// Copies a function declaration into this graph's local function library
 	UFUNCTION(BlueprintCallable, Category = RigVMController)
