Control Rig: Python commands output to different message log pages depending on the blueprint

diff --git a/Engine/Plugins/Animation/ControlRig/Source/ControlRig/Private/Rigs/RigHierarchy.cpp b/Engine/Plugins/Animation/ControlRig/Source/ControlRig/Private/Rigs/RigHierarchy.cpp
--- a/Engine/Plugins/Animation/ControlRig/Source/ControlRig/Private/Rigs/RigHierarchy.cpp
+++ b/Engine/Plugins/Animation/ControlRig/Source/ControlRig/Private/Rigs/RigHierarchy.cpp
@@ -1478,7 +1478,8 @@ void URigHierarchy::SetTransform(FRigTransformElement* InTransformElement, const
 			}
 		}
 
-		RigVMPythonUtils::Print(FString::Printf(TEXT("hierarchy.%s(%s, %s, %s, %s)"),
+		RigVMPythonUtils::Print(GetOuter()->GetFName().ToString(),
+			FString::Printf(TEXT("hierarchy.%s(%s, %s, %s, %s)"),
 			*MethodName,
 			*InTransformElement->GetKey().ToPythonString(),
 			*RigVMPythonUtils::TransformToPythonString(InTransform),
@@ -1586,7 +1587,8 @@ void URigHierarchy::SetControlOffsetTransform(FRigControlElement* InControlEleme
 	if (bPrintPythonCommands)
 	{
 		//FRigElementKey InKey, FTransform InTransform, bool bInitial = false, bool bAffectChildren = true, bool bSetupUndo = false, bool bPrintPythonCommands = false)
-		RigVMPythonUtils::Print(FString::Printf(TEXT("hierarchy.set_control_offset_transform(%s, %s, %s, %s)"),
+		RigVMPythonUtils::Print(GetOuter()->GetFName().ToString(),
+			FString::Printf(TEXT("hierarchy.set_control_offset_transform(%s, %s, %s, %s)"),
 			*InControlElement->GetKey().ToPythonString(),
 			*RigVMPythonUtils::TransformToPythonString(InTransform),
 			(ERigTransformType::IsInitial(InTransformType)) ? TEXT("True") : TEXT("False"),
diff --git a/Engine/Plugins/Animation/ControlRig/Source/ControlRig/Private/Rigs/RigHierarchyController.cpp b/Engine/Plugins/Animation/ControlRig/Source/ControlRig/Private/Rigs/RigHierarchyController.cpp
--- a/Engine/Plugins/Animation/ControlRig/Source/ControlRig/Private/Rigs/RigHierarchyController.cpp
+++ b/Engine/Plugins/Animation/ControlRig/Source/ControlRig/Private/Rigs/RigHierarchyController.cpp
@@ -175,7 +175,8 @@ FRigElementKey URigHierarchyController::AddBone(FName InName, FRigElementKey InP
 		TArray<FString> Commands = GetAddBonePythonCommands(NewElement);
 		for (const FString& Command : Commands)
 		{
-			RigVMPythonUtils::Print(FString::Printf(TEXT("%s"), *Command));
+			RigVMPythonUtils::Print(GetOuter()->GetFName().ToString(),
+				FString::Printf(TEXT("%s"), *Command));
 		}
 	}
 #endif
@@ -226,7 +227,8 @@ FRigElementKey URigHierarchyController::AddNull(FName InName, FRigElementKey InP
 		TArray<FString> Commands = GetAddNullPythonCommands(NewElement);
 		for (const FString& Command : Commands)
 		{
-			RigVMPythonUtils::Print(FString::Printf(TEXT("%s"), *Command));
+			RigVMPythonUtils::Print(GetOuter()->GetFName().ToString(),
+				FString::Printf(TEXT("%s"), *Command));
 		}
 	}
 #endif
@@ -284,7 +286,8 @@ FRigElementKey URigHierarchyController::AddControl(
 		TArray<FString> Commands = GetAddControlPythonCommands(NewElement);
 		for (const FString& Command : Commands)
 		{
-			RigVMPythonUtils::Print(FString::Printf(TEXT("%s"), *Command));
+			RigVMPythonUtils::Print(GetOuter()->GetFName().ToString(),
+				FString::Printf(TEXT("%s"), *Command));
 		}
 	}
 #endif
@@ -322,7 +325,8 @@ FRigElementKey URigHierarchyController::AddCurve(FName InName, float InValue, bo
 		TArray<FString> Commands = GetAddCurvePythonCommands(NewElement);
 		for (const FString& Command : Commands)
 		{
-			RigVMPythonUtils::Print(FString::Printf(TEXT("%s"), *Command));
+			RigVMPythonUtils::Print(GetOuter()->GetFName().ToString(),
+				FString::Printf(TEXT("%s"), *Command));
 		}
 	}
 #endif
@@ -364,7 +368,8 @@ FRigElementKey URigHierarchyController::AddRigidBody(FName InName, FRigElementKe
 		TArray<FString> Commands = GetAddRigidBodyPythonCommands(NewElement);
 		for (const FString& Command : Commands)
 		{
-			RigVMPythonUtils::Print(FString::Printf(TEXT("%s"), *Command));
+			RigVMPythonUtils::Print(GetOuter()->GetFName().ToString(),
+				FString::Printf(TEXT("%s"), *Command));
 		}
 	}
 #endif
@@ -648,7 +653,8 @@ TArray<FRigElementKey> URigHierarchyController::ImportBones(USkeleton* InSkeleto
 	{
 		if (bPrintPythonCommand)
 		{
-			RigVMPythonUtils::Print(FString::Printf(TEXT("hierarchy_controller.import_bones_from_asset('%s', '%s', %s, %s, %s)"),
+			RigVMPythonUtils::Print(GetOuter()->GetFName().ToString(),
+				FString::Printf(TEXT("hierarchy_controller.import_bones_from_asset('%s', '%s', %s, %s, %s)"),
 				*InSkeleton->GetPathName(),
 				*InNameSpace.ToString(),
 				(bReplaceExistingBones) ? TEXT("True") : TEXT("False"),
@@ -753,7 +759,8 @@ TArray<FRigElementKey> URigHierarchyController::ImportCurves(USkeleton* InSkelet
 	{
 		if (bPrintPythonCommand)
 		{
-			RigVMPythonUtils::Print(FString::Printf(TEXT("hierarchy_controller.import_curves_from_asset('%s', '%s', %s)"),
+			RigVMPythonUtils::Print(GetOuter()->GetFName().ToString(),
+				FString::Printf(TEXT("hierarchy_controller.import_curves_from_asset('%s', '%s', %s)"),
 				*InSkeleton->GetPathName(),
 				*InNameSpace.ToString(),
 				(bSelectCurves) ? TEXT("True") : TEXT("False")));
@@ -1092,7 +1099,8 @@ TArray<FRigElementKey> URigHierarchyController::ImportFromText(FString InContent
 	{
 		FString PythonContent = InContent.Replace(TEXT("\\\""), TEXT("\\\\\""));
 		
-		RigVMPythonUtils::Print(FString::Printf(TEXT("hierarchy_controller.import_from_text('%s', %s, %s)"),
+		RigVMPythonUtils::Print(GetOuter()->GetFName().ToString(),
+			FString::Printf(TEXT("hierarchy_controller.import_from_text('%s', %s, %s)"),
 			*PythonContent,
 			(bReplaceExistingElements) ? TEXT("True") : TEXT("False"),
 			(bSelectNewElements) ? TEXT("True") : TEXT("False")));
@@ -1567,7 +1575,8 @@ bool URigHierarchyController::RemoveElement(FRigElementKey InElement, bool bSetu
 
 	if (bRemoved && bPrintPythonCommand)
 	{
-		RigVMPythonUtils::Print(FString::Printf(TEXT("hierarchy_controller.remove_element(%s)"),
+		RigVMPythonUtils::Print(GetOuter()->GetFName().ToString(),
+			FString::Printf(TEXT("hierarchy_controller.remove_element(%s)"),
 			*InElement.ToPythonString()));
 	}
 #endif
@@ -1718,7 +1727,8 @@ FRigElementKey URigHierarchyController::RenameElement(FRigElementKey InElement,
 
 	if (bRenamed && bPrintPythonCommand)
 	{
-		RigVMPythonUtils::Print(FString::Printf(TEXT("hierarchy_controller.rename_element(%s, '%s')"),
+		RigVMPythonUtils::Print(GetOuter()->GetFName().ToString(), 
+			FString::Printf(TEXT("hierarchy_controller.rename_element(%s, '%s')"),
 			*InElement.ToPythonString(),
 			*InName.ToString()));
 	}
@@ -1973,7 +1983,8 @@ bool URigHierarchyController::RemoveParent(FRigElementKey InChild, FRigElementKe
 
 	if (bRemoved && bPrintPythonCommand)
 	{
-		RigVMPythonUtils::Print(FString::Printf(TEXT("hierarchy_controller.remove_parent(%s, %s, %s)"),
+		RigVMPythonUtils::Print(GetOuter()->GetFName().ToString(),
+			FString::Printf(TEXT("hierarchy_controller.remove_parent(%s, %s, %s)"),
 			*InChild.ToPythonString(),
 			*InParent.ToPythonString(),
 			(bMaintainGlobalTransform) ? TEXT("True") : TEXT("False")));
@@ -2144,7 +2155,8 @@ bool URigHierarchyController::RemoveAllParents(FRigElementKey InChild, bool bMai
 	
 	if (bRemoved && bPrintPythonCommand)
 	{
-		RigVMPythonUtils::Print(FString::Printf(TEXT("hierarchy_controller.remove_all_parents(%s, %s)"),
+		RigVMPythonUtils::Print(GetOuter()->GetFName().ToString(),
+			FString::Printf(TEXT("hierarchy_controller.remove_all_parents(%s, %s)"),
 			*InChild.ToPythonString(),
 			(bMaintainGlobalTransform) ? TEXT("True") : TEXT("False")));
 	}
@@ -2218,7 +2230,8 @@ bool URigHierarchyController::SetParent(FRigElementKey InChild, FRigElementKey I
 
 	if (bParentSet && bPrintPythonCommand)
 	{
-		RigVMPythonUtils::Print(FString::Printf(TEXT("hierarchy_controller.set_parent(%s, %s, %s)"),
+		RigVMPythonUtils::Print(GetOuter()->GetFName().ToString(),
+			FString::Printf(TEXT("hierarchy_controller.set_parent(%s, %s, %s)"),
 			*InChild.ToPythonString(),
 			*InParent.ToPythonString(),
 			(bMaintainGlobalTransform) ? TEXT("True") : TEXT("False")));
@@ -2247,7 +2260,8 @@ TArray<FRigElementKey> URigHierarchyController::DuplicateElements(TArray<FRigEle
 		}
 		ArrayStr += TEXT("]");		
 		
-		RigVMPythonUtils::Print(FString::Printf(TEXT("hierarchy_controller.duplicate_elements(%s, %s)"),
+		RigVMPythonUtils::Print(GetOuter()->GetFName().ToString(),
+			FString::Printf(TEXT("hierarchy_controller.duplicate_elements(%s, %s)"),
 			*ArrayStr,
 			(bSelectNewElements) ? TEXT("True") : TEXT("False")));
 	}
@@ -2362,7 +2376,8 @@ TArray<FRigElementKey> URigHierarchyController::MirrorElements(TArray<FRigElemen
 		}
 		ArrayStr += TEXT("]");
 
-		RigVMPythonUtils::Print(FString::Printf(TEXT("hierarchy_controller.mirror_elements(%s, unreal.RigMirrorSettings(unreal.AxisType.%s, unreal.AxisType.%s, '%s', '%s'), %s)"),
+		RigVMPythonUtils::Print(GetOuter()->GetFName().ToString(),
+			FString::Printf(TEXT("hierarchy_controller.mirror_elements(%s, unreal.RigMirrorSettings(unreal.AxisType.%s, unreal.AxisType.%s, '%s', '%s'), %s)"),
 			*ArrayStr,
 			(InSettings.MirrorAxis.GetValue() == EAxis::X) ? TEXT("X") : (InSettings.MirrorAxis.GetValue() == EAxis::Y) ? TEXT("Y") : TEXT("Z"),
 			(InSettings.AxisToFlip.GetValue() == EAxis::X) ? TEXT("X") : (InSettings.AxisToFlip.GetValue() == EAxis::Y) ? TEXT("Y") : TEXT("Z"),
diff --git a/Engine/Plugins/Animation/ControlRig/Source/ControlRigDeveloper/Private/ControlRigBlueprint.cpp b/Engine/Plugins/Animation/ControlRig/Source/ControlRigDeveloper/Private/ControlRigBlueprint.cpp
--- a/Engine/Plugins/Animation/ControlRig/Source/ControlRigDeveloper/Private/ControlRigBlueprint.cpp
+++ b/Engine/Plugins/Animation/ControlRig/Source/ControlRigDeveloper/Private/ControlRigBlueprint.cpp
@@ -3517,7 +3517,8 @@ void UControlRigBlueprint::OnVariableAdded(const FName& InVarName)
 
 	FRigVMExternalVariable ExternalVariable = UControlRig::GetExternalVariableFromDescription(Variable);
 
-	RigVMPythonUtils::Print(FString::Printf(TEXT("blueprint.add_member_variable('%s', '%s', %s, %s, '%s')"),
+	RigVMPythonUtils::Print(GetFName().ToString(),
+		FString::Printf(TEXT("blueprint.add_member_variable('%s', '%s', %s, %s, '%s')"),
 			*InVarName.ToString(),
 			*ExternalVariable.TypeName.ToString(),
 			(ExternalVariable.bIsPublic) ? TEXT("False") : TEXT("True"), 
@@ -3538,7 +3539,8 @@ void UControlRigBlueprint::OnVariableRemoved(const FName& InVarName)
 		}
 	}
 
-	RigVMPythonUtils::Print(FString::Printf(TEXT("blueprint.remove_member_variable('%s')"),
+	RigVMPythonUtils::Print(GetFName().ToString(),
+		FString::Printf(TEXT("blueprint.remove_member_variable('%s')"),
 			*InVarName.ToString()));
 	
 	BroadcastExternalVariablesChangedEvent();
@@ -3555,7 +3557,8 @@ void UControlRigBlueprint::OnVariableRenamed(const FName& InOldVarName, const FN
 		}
 	}
 
-	RigVMPythonUtils::Print(FString::Printf(TEXT("blueprint.rename_member_variable('%s', '%s')"),
+	RigVMPythonUtils::Print(GetFName().ToString(),
+		FString::Printf(TEXT("blueprint.rename_member_variable('%s', '%s')"),
 			*InOldVarName.ToString(),
 			*InNewVarName.ToString()));
 	
@@ -3610,7 +3613,8 @@ void UControlRigBlueprint::OnVariableTypeChanged(const FName& InVarName, FEdGrap
 		}
 	}
 
-	RigVMPythonUtils::Print(FString::Printf(TEXT("blueprint.change_member_variable_type('%s', '%s')"),
+	RigVMPythonUtils::Print(GetFName().ToString(),
+		FString::Printf(TEXT("blueprint.change_member_variable_type('%s', '%s')"),
 		*InVarName.ToString(),
 		*PinType));
 
diff --git a/Engine/Plugins/Animation/ControlRig/Source/ControlRigEditor/Private/ControlRigCompilerDetails.cpp b/Engine/Plugins/Animation/ControlRig/Source/ControlRigEditor/Private/ControlRigCompilerDetails.cpp
--- a/Engine/Plugins/Animation/ControlRig/Source/ControlRigEditor/Private/ControlRigCompilerDetails.cpp
+++ b/Engine/Plugins/Animation/ControlRig/Source/ControlRigEditor/Private/ControlRigCompilerDetails.cpp
@@ -182,7 +182,7 @@ FReply FRigVMCompileSettingsDetails::OnRunPythonContextClicked()
 
 		for (FString& Command : PyCommands)
 		{
-			RigVMPythonUtils::Print(Command);
+			RigVMPythonUtils::Print(BlueprintBeingCustomized->GetFName().ToString(), Command);
 		
 			// Run the Python commands
 			IPythonScriptPlugin::Get()->ExecPythonCommand(*Command);
diff --git a/Engine/Source/Developer/MessageLog/Private/Model/MessageLogListingModel.cpp b/Engine/Source/Developer/MessageLog/Private/Model/MessageLogListingModel.cpp
--- a/Engine/Source/Developer/MessageLog/Private/Model/MessageLogListingModel.cpp
+++ b/Engine/Source/Developer/MessageLog/Private/Model/MessageLogListingModel.cpp
@@ -118,6 +118,33 @@ void FMessageLogListingModel::NewPage( const FText& InTitle, uint32 InMaxPages )
 	}
 }
 
+bool FMessageLogListingModel::SwitchToPage(const FText& InTitle, uint32 InMaxPages )
+{
+	if (CurrentPage().Title.CompareTo(InTitle) == 0)
+	{
+		return false;
+	}
+	
+	for (TDoubleLinkedList<FPage>::TIterator it = begin(Pages); it != end(Pages); ++it)
+	{
+		if (it.GetNode()->GetValue().Title.CompareTo(InTitle) == 0)
+		{
+			TDoubleLinkedList<FPage>::TDoubleLinkedListNode* SelectedNode = it.GetNode();
+			Pages.RemoveNode(SelectedNode, false);
+			Pages.AddHead(SelectedNode);
+
+			// invalidate cache as all indices will change
+			CachedPage = nullptr;
+			return true;
+		}
+	}
+
+	// If the page does not exist yet, create a new one
+	NewPage(InTitle, InMaxPages);
+
+	return true;
+}
+
 uint32 FMessageLogListingModel::NumPages() const
 {
 	return Pages.Num();
diff --git a/Engine/Source/Developer/MessageLog/Private/Model/MessageLogListingModel.h b/Engine/Source/Developer/MessageLog/Private/Model/MessageLogListingModel.h
--- a/Engine/Source/Developer/MessageLog/Private/Model/MessageLogListingModel.h
+++ b/Engine/Source/Developer/MessageLog/Private/Model/MessageLogListingModel.h
@@ -81,6 +81,15 @@ public:
 	 */
 	void NewPage( const FText& InTitle, uint32 InMaxPages );
 
+	/**
+	* Switches the current page to the one specified by the title. If the page does not exist, creates a new page.
+	* @param	InTitle		The title of the page to search for, or create
+	* @param	InMaxPages	The maximum number of pages we keep around. If the count is exceeded, 
+	*						we discard the oldest page.
+	* @returns true if the page has changed (either switched or created a new one).
+	*/
+	bool SwitchToPage( const FText& InTitle, uint32 InMaxPages  );
+
 	/** Get the number of pages contained in this log */
 	uint32 NumPages() const;
 
diff --git a/Engine/Source/Developer/MessageLog/Private/Presentation/MessageLogListingViewModel.cpp b/Engine/Source/Developer/MessageLog/Private/Presentation/MessageLogListingViewModel.cpp
--- a/Engine/Source/Developer/MessageLog/Private/Presentation/MessageLogListingViewModel.cpp
+++ b/Engine/Source/Developer/MessageLog/Private/Presentation/MessageLogListingViewModel.cpp
@@ -277,6 +277,18 @@ void FMessageLogListingViewModel::NewPage( const FText& Title )
 	RefreshFilteredMessages();
 }
 
+void FMessageLogListingViewModel::SwitchToPage(const FText& Title)
+{
+	// we should take this as a suggestion we want to show pages!
+	bShowPages = true;
+
+	// switch page & refresh
+	MessageLogListingModel->SwitchToPage( Title, MaxPageCount );
+	
+	// reset so we always display the current page when we switch
+	SetCurrentPageIndex(0);
+}
+
 void FMessageLogListingViewModel::NotifyIfAnyMessages( const FText& Message, EMessageSeverity::Type SeverityFilter, bool bForce )
 {
 	// Note we use page 0 in this function, as that is the page that will
diff --git a/Engine/Source/Developer/MessageLog/Private/Presentation/MessageLogListingViewModel.h b/Engine/Source/Developer/MessageLog/Private/Presentation/MessageLogListingViewModel.h
--- a/Engine/Source/Developer/MessageLog/Private/Presentation/MessageLogListingViewModel.h
+++ b/Engine/Source/Developer/MessageLog/Private/Presentation/MessageLogListingViewModel.h
@@ -56,6 +56,7 @@ public:
 	virtual const TArray< TSharedRef< class FMessageFilter> >& GetMessageFilters() const override;
 	virtual void ExecuteToken( const TSharedRef<class IMessageToken>& Token ) const override;
 	virtual void NewPage( const FText& Title ) override;
+	virtual void SwitchToPage( const FText& Title ) override;
 	virtual void NotifyIfAnyMessages( const FText& Message, EMessageSeverity::Type SeverityFilter = EMessageSeverity::Info, bool bForce = false ) override;
 	virtual void Open() override;
 	virtual int32 NumMessages( EMessageSeverity::Type SeverityFilter ) override;
diff --git a/Engine/Source/Developer/RigVMDeveloper/Private/RigVMModel/RigVMController.cpp b/Engine/Source/Developer/RigVMDeveloper/Private/RigVMModel/RigVMController.cpp
--- a/Engine/Source/Developer/RigVMDeveloper/Private/RigVMModel/RigVMController.cpp
+++ b/Engine/Source/Developer/RigVMDeveloper/Private/RigVMModel/RigVMController.cpp
@@ -725,7 +725,8 @@ URigVMUnitNode* URigVMController::AddUnitNode(UScriptStruct* InScriptStruct, con
 		TArray<FString> Commands = GetAddNodePythonCommands(Node);
 		for (const FString& Command : Commands)
 		{
-			RigVMPythonUtils::Print(FString::Printf(TEXT("%s"), *Command));
+			RigVMPythonUtils::Print(GetGraph()->GetRootGraph()->GetOuter()->GetFName().ToString(),
+				FString::Printf(TEXT("%s"), *Command));
 		}
 	}
 
@@ -911,7 +912,8 @@ URigVMVariableNode* URigVMController::AddVariableNode(const FName& InVariableNam
 		TArray<FString> Commands = GetAddNodePythonCommands(Node);
 		for (const FString& Command : Commands)
 		{
-			RigVMPythonUtils::Print(FString::Printf(TEXT("%s"), *Command));
+			RigVMPythonUtils::Print(GetGraph()->GetRootGraph()->GetOuter()->GetFName().ToString(), 
+								FString::Printf(TEXT("%s"), *Command));
 		}
 	}
 
@@ -1505,7 +1507,8 @@ URigVMParameterNode* URigVMController::AddParameterNode(const FName& InParameter
 		TArray<FString> Commands = GetAddNodePythonCommands(Node);
 		for (const FString& Command : Commands)
 		{
-			RigVMPythonUtils::Print(FString::Printf(TEXT("%s"), *Command));
+			RigVMPythonUtils::Print(GetGraph()->GetRootGraph()->GetOuter()->GetFName().ToString(), 
+								FString::Printf(TEXT("%s"), *Command));
 		}
 	}
 
@@ -1582,7 +1585,8 @@ URigVMCommentNode* URigVMController::AddCommentNode(const FString& InCommentText
 		TArray<FString> Commands = GetAddNodePythonCommands(Node);
 		for (const FString& Command : Commands)
 		{
-			RigVMPythonUtils::Print(FString::Printf(TEXT("%s"), *Command));
+			RigVMPythonUtils::Print(GetGraph()->GetRootGraph()->GetOuter()->GetFName().ToString(), 
+								FString::Printf(TEXT("%s"), *Command));
 		}
 	}
 
@@ -1634,12 +1638,13 @@ URigVMRerouteNode* URigVMController::AddRerouteNodeOnLink(URigVMLink* InLink, bo
 
 	if (bPrintPythonCommand)
 	{
-		RigVMPythonUtils::Print(FString::Printf(TEXT("blueprint.get_controller_by_name('%s').add_reroute_node_on_link_path('%s', %s, %s, '%s')"),
-				*GetGraph()->GetGraphName(),
-				*InLink->GetPinPathRepresentation(),
-				(bShowAsFullNode) ? TEXT("True") : TEXT("False"),
-				*RigVMPythonUtils::Vector2DToPythonString(Node->GetPosition()),
-				*Node->GetName()));
+		RigVMPythonUtils::Print(GetGraph()->GetRootGraph()->GetOuter()->GetFName().ToString(), 
+							FString::Printf(TEXT("blueprint.get_controller_by_name('%s').add_reroute_node_on_link_path('%s', %s, %s, '%s')"),
+											*GetGraph()->GetGraphName(),
+											*InLink->GetPinPathRepresentation(),
+											(bShowAsFullNode) ? TEXT("True") : TEXT("False"),
+											*RigVMPythonUtils::Vector2DToPythonString(Node->GetPosition()),
+											*Node->GetName()));
 	}
 
 	return Node;
@@ -1750,13 +1755,14 @@ URigVMRerouteNode* URigVMController::AddRerouteNodeOnPin(const FString& InPinPat
 	if (bPrintPythonCommand)
 	{
 		// AddRerouteNodeOnPin(const FString& InPinPath, bool bAsInput, bool bShowAsFullNode, const FVector2D& InPosition, const FString& InNodeName, bool bSetupUndoRedo, bool bPrintPythonCommand)
-		RigVMPythonUtils::Print(FString::Printf(TEXT("blueprint.get_controller_by_name('%s').add_reroute_node_on_pin('%s', %s, %s, %s '%s')"),
-				*GetGraph()->GetGraphName(),
-				*InPinPath,
-				(bAsInput) ? TEXT("True") : TEXT("False"),
-				(bShowAsFullNode) ? TEXT("True") : TEXT("False"),
-				*RigVMPythonUtils::Vector2DToPythonString(Node->GetPosition()),
-				*Node->GetName()));
+		RigVMPythonUtils::Print(GetGraph()->GetRootGraph()->GetOuter()->GetFName().ToString(), 
+							FString::Printf(TEXT("blueprint.get_controller_by_name('%s').add_reroute_node_on_pin('%s', %s, %s, %s '%s')"),
+											*GetGraph()->GetGraphName(),
+											*InPinPath,
+											(bAsInput) ? TEXT("True") : TEXT("False"),
+											(bShowAsFullNode) ? TEXT("True") : TEXT("False"),
+											*RigVMPythonUtils::Vector2DToPythonString(Node->GetPosition()),
+											*Node->GetName()));
 	}
 
 	return Node;
@@ -1956,15 +1962,16 @@ URigVMInjectionInfo* URigVMController::AddInjectedNode(const FString& InPinPath,
 
 	if (bPrintPythonCommand)
 	{
-		RigVMPythonUtils::Print(FString::Printf(TEXT("blueprint.get_controller_by_name('%s').add_injected_node_from_struct_path('%s', %s, '%s', '%s', '%s', '%s', '%s')"),
-				*GetGraph()->GetGraphName(),
-				*InPinPath,
-				(bAsInput) ? TEXT("True") : TEXT("False"),
-				*InScriptStruct->GetPathName(),
-				*InMethodName.ToString(),
-				*InInputPinName.ToString(),
-				*InOutputPinName.ToString(),
-				*InNodeName));
+		RigVMPythonUtils::Print(GetGraph()->GetRootGraph()->GetOuter()->GetFName().ToString(), 
+							FString::Printf(TEXT("blueprint.get_controller_by_name('%s').add_injected_node_from_struct_path('%s', %s, '%s', '%s', '%s', '%s', '%s')"),
+											*GetGraph()->GetGraphName(),
+											*InPinPath,
+											(bAsInput) ? TEXT("True") : TEXT("False"),
+											*InScriptStruct->GetPathName(),
+											*InMethodName.ToString(),
+											*InInputPinName.ToString(),
+											*InOutputPinName.ToString(),
+											*InNodeName));
 	}
 
 	return InjectionInfo;
@@ -2111,9 +2118,10 @@ URigVMNode* URigVMController::EjectNodeFromPin(const FString& InPinPath, bool bS
 
 	if (bPrintPythonCommand)
 	{
-		RigVMPythonUtils::Print(FString::Printf(TEXT("blueprint.get_controller_by_name('%s').eject_node_from_pin('%s')"),
-				*GetGraph()->GetGraphName(),
-				*InPinPath));
+		RigVMPythonUtils::Print(GetGraph()->GetRootGraph()->GetOuter()->GetFName().ToString(), 
+							FString::Printf(TEXT("blueprint.get_controller_by_name('%s').eject_node_from_pin('%s')"),
+											*GetGraph()->GetGraphName(),
+											*InPinPath));
 	}
 
 	return EjectedNode;
@@ -2711,10 +2719,11 @@ TMap<URigVMLibraryNode*, URigVMLibraryNode*> URigVMController::LocalizeFunctions
 		}
 		FunctionNames += TEXT("]");
 		
-		RigVMPythonUtils::Print(FString::Printf(TEXT("blueprint.get_controller_by_name('%s').localize_functions(%s, %s)"),
-				*GetGraph()->GetGraphName(),
-				*FunctionNames,
-				(bLocalizeDependentPrivateFunctions) ? TEXT("True") : TEXT("False")));
+		RigVMPythonUtils::Print(GetGraph()->GetRootGraph()->GetOuter()->GetFName().ToString(), 
+							FString::Printf(TEXT("blueprint.get_controller_by_name('%s').localize_functions(%s, %s)"),
+											*GetGraph()->GetGraphName(),
+											*FunctionNames,
+											(bLocalizeDependentPrivateFunctions) ? TEXT("True") : TEXT("False")));
 	}
 
 	return LocalizedFunctions;
@@ -2769,10 +2778,11 @@ URigVMCollapseNode* URigVMController::CollapseNodes(const TArray<FName>& InNodeN
 		}
 		ArrayStr += TEXT("]");
 		
-		RigVMPythonUtils::Print(FString::Printf(TEXT("blueprint.get_controller_by_name('%s').collapse_nodes(%s, '%s')"),
-				*GetGraph()->GetGraphName(),
-				*ArrayStr,
-				*InCollapseNodeName));
+		RigVMPythonUtils::Print(GetGraph()->GetRootGraph()->GetOuter()->GetFName().ToString(), 
+							FString::Printf(TEXT("blueprint.get_controller_by_name('%s').collapse_nodes(%s, '%s')"),
+											*GetGraph()->GetGraphName(),
+											*ArrayStr,
+											*InCollapseNodeName));
 	}
 
 	return Node;
@@ -2806,9 +2816,10 @@ TArray<URigVMNode*> URigVMController::ExpandLibraryNode(const FName& InNodeName,
 
 	if (!Nodes.IsEmpty() && bPrintPythonCommand)
 	{
-		RigVMPythonUtils::Print(FString::Printf(TEXT("blueprint.get_controller_by_name('%s').expand_library_node('%s')"),
-				*GetGraph()->GetGraphName(),
-				*InNodeName.ToString()));
+		RigVMPythonUtils::Print(GetGraph()->GetRootGraph()->GetOuter()->GetFName().ToString(), 
+							FString::Printf(TEXT("blueprint.get_controller_by_name('%s').expand_library_node('%s')"),
+											*GetGraph()->GetGraphName(),
+											*InNodeName.ToString()));
 	}
 
 	return Nodes;
@@ -3934,9 +3945,10 @@ FName URigVMController::PromoteCollapseNodeToFunctionReferenceNode(const FName&
 	{
 		if (bPrintPythonCommand)
 		{
-			RigVMPythonUtils::Print(FString::Printf(TEXT("blueprint.get_controller_by_name('%s').promote_collapse_node_to_function_reference_node('%s')"),
-					*GetGraph()->GetGraphName(),
-					*InNodeName.ToString()));
+			RigVMPythonUtils::Print(GetGraph()->GetRootGraph()->GetOuter()->GetFName().ToString(), 
+								FString::Printf(TEXT("blueprint.get_controller_by_name('%s').promote_collapse_node_to_function_reference_node('%s')"),
+													*GetGraph()->GetGraphName(),
+													*InNodeName.ToString()));
 		}
 		
 		return Result->GetFName();
@@ -4345,14 +4357,16 @@ bool URigVMController::RemoveNode(URigVMNode* InNode, bool bSetupUndoRedo, bool
 	{
 		if (Graph->IsA<URigVMFunctionLibrary>())
 		{
-			RigVMPythonUtils::Print(FString::Printf(TEXT("library_controller.remove_function_from_library('%s')"),
-				*InNode->GetFName().ToString()));
+			RigVMPythonUtils::Print(GetGraph()->GetRootGraph()->GetOuter()->GetFName().ToString(), 
+								FString::Printf(TEXT("library_controller.remove_function_from_library('%s')"),
+												*InNode->GetFName().ToString()));
 		}
 		else
 		{
-			RigVMPythonUtils::Print(FString::Printf(TEXT("blueprint.get_controller_by_name('%s').remove_node_by_name('%s')"),
-				*GetGraph()->GetGraphName(),
-				*InNode->GetNodePath()));
+			RigVMPythonUtils::Print(GetGraph()->GetRootGraph()->GetOuter()->GetFName().ToString(), 
+								FString::Printf(TEXT("blueprint.get_controller_by_name('%s').remove_node_by_name('%s')"),
+												*GetGraph()->GetGraphName(),
+												*InNode->GetNodePath()));
 		}
 	}
 
@@ -4457,9 +4471,10 @@ bool URigVMController::RenameNode(URigVMNode* InNode, const FName& InNewName, bo
 
 	if (GetGraph()->IsA<URigVMFunctionLibrary>())
 	{
-		RigVMPythonUtils::Print(FString::Printf(TEXT("library_controller.rename_function('%s', '%s')"),
-			*OldName,
-			*InNewName.ToString()));
+		RigVMPythonUtils::Print(GetGraph()->GetRootGraph()->GetOuter()->GetFName().ToString(), 
+							FString::Printf(TEXT("library_controller.rename_function('%s', '%s')"),
+										*OldName,
+										*InNewName.ToString()));
 	}
 
 	return true;
@@ -4608,9 +4623,10 @@ bool URigVMController::SetNodeSelection(const TArray<FName>& InNodeNames, bool b
 		}
 		ArrayStr += TEXT("]");
 		
-		RigVMPythonUtils::Print(FString::Printf(TEXT("blueprint.get_controller_by_name('%s').set_node_selection(%s)"),
-				*GetGraph()->GetGraphName(),
-				*ArrayStr));
+		RigVMPythonUtils::Print(GetGraph()->GetRootGraph()->GetOuter()->GetFName().ToString(), 
+							FString::Printf(TEXT("blueprint.get_controller_by_name('%s').set_node_selection(%s)"),
+											*GetGraph()->GetGraphName(),
+											*ArrayStr));
 	}
 
 	return bSelectedSomething;
@@ -4646,10 +4662,11 @@ bool URigVMController::SetNodePosition(URigVMNode* InNode, const FVector2D& InPo
 
 	if (bPrintPythonCommand)
 	{
-		RigVMPythonUtils::Print(FString::Printf(TEXT("blueprint.get_controller_by_name('%s').set_node_position_by_name('%s', %s)"),
-				*GetGraph()->GetGraphName(),
-				*InNode->GetNodePath(),
-				*RigVMPythonUtils::Vector2DToPythonString(InPosition)));
+		RigVMPythonUtils::Print(GetGraph()->GetRootGraph()->GetOuter()->GetFName().ToString(), 
+							FString::Printf(TEXT("blueprint.get_controller_by_name('%s').set_node_position_by_name('%s', %s)"),
+											*GetGraph()->GetGraphName(),
+											*InNode->GetNodePath(),
+											*RigVMPythonUtils::Vector2DToPythonString(InPosition)));
 	}
 
 	return true;
@@ -4700,10 +4717,11 @@ bool URigVMController::SetNodeSize(URigVMNode* InNode, const FVector2D& InSize,
 
 	if (bPrintPythonCommand)
 	{
-		RigVMPythonUtils::Print(FString::Printf(TEXT("blueprint.get_controller_by_name('%s').set_node_size_by_name('%s', %s)"),
-				*GetGraph()->GetGraphName(),
-				*InNode->GetNodePath(),
-				*RigVMPythonUtils::Vector2DToPythonString(InSize)));
+		RigVMPythonUtils::Print(GetGraph()->GetRootGraph()->GetOuter()->GetFName().ToString(), 
+							FString::Printf(TEXT("blueprint.get_controller_by_name('%s').set_node_size_by_name('%s', %s)"),
+											*GetGraph()->GetGraphName(),
+											*InNode->GetNodePath(),
+											*RigVMPythonUtils::Vector2DToPythonString(InSize)));
 	}
 
 	return true;
@@ -4765,10 +4783,11 @@ bool URigVMController::SetNodeColor(URigVMNode* InNode, const FLinearColor& InCo
 
 	if (bPrintPythonCommand)
 	{
-		RigVMPythonUtils::Print(FString::Printf(TEXT("blueprint.get_controller_by_name('%s').set_node_color_by_name('%s', %s)"),
-				*GetGraph()->GetGraphName(),
-				*InNode->GetNodePath(),
-				*RigVMPythonUtils::LinearColorToPythonString(InColor)));
+		RigVMPythonUtils::Print(GetGraph()->GetRootGraph()->GetOuter()->GetFName().ToString(), 
+							FString::Printf(TEXT("blueprint.get_controller_by_name('%s').set_node_color_by_name('%s', %s)"),
+											*GetGraph()->GetGraphName(),
+											*InNode->GetNodePath(),
+											*RigVMPythonUtils::LinearColorToPythonString(InColor)));
 	}
 
 	return true;
@@ -4818,7 +4837,8 @@ bool URigVMController::SetNodeCategory(URigVMCollapseNode* InNode, const FString
 
 	if (bPrintPythonCommand)
 	{
-		RigVMPythonUtils::Print(FString::Printf(TEXT("blueprint.get_controller_by_name('%s').set_node_category_by_name('%s', '%s')"),
+		RigVMPythonUtils::Print(GetGraph()->GetRootGraph()->GetOuter()->GetFName().ToString(),
+			FString::Printf(TEXT("blueprint.get_controller_by_name('%s').set_node_category_by_name('%s', '%s')"),
 				*GetGraph()->GetGraphName(),
 				*InNode->GetNodePath(),
 				*InCategory));
@@ -4871,7 +4891,8 @@ bool URigVMController::SetNodeKeywords(URigVMCollapseNode* InNode, const FString
 
 	if (bPrintPythonCommand)
 	{
-		RigVMPythonUtils::Print(FString::Printf(TEXT("blueprint.get_controller_by_name('%s').set_node_keywords_by_name('%s', '%s')"),
+		RigVMPythonUtils::Print(GetGraph()->GetRootGraph()->GetOuter()->GetFName().ToString(),
+			FString::Printf(TEXT("blueprint.get_controller_by_name('%s').set_node_keywords_by_name('%s', '%s')"),
 				*GetGraph()->GetGraphName(),
 				*InNode->GetNodePath(),
 				*InKeywords));
@@ -4924,7 +4945,8 @@ bool URigVMController::SetNodeDescription(URigVMCollapseNode* InNode, const FStr
 
 	if (bPrintPythonCommand)
 	{
-		RigVMPythonUtils::Print(FString::Printf(TEXT("blueprint.get_controller_by_name('%s').set_node_description_by_name('%s', '%s')"),
+		RigVMPythonUtils::Print(GetGraph()->GetRootGraph()->GetOuter()->GetFName().ToString(),
+			FString::Printf(TEXT("blueprint.get_controller_by_name('%s').set_node_description_by_name('%s', '%s')"),
 				*GetGraph()->GetGraphName(),
 				*InNode->GetNodePath(),
 				*InDescription));
@@ -4979,7 +5001,8 @@ bool URigVMController::SetCommentText(URigVMNode* InNode, const FString& InComme
 
 		if (bPrintPythonCommand)
 		{
-			RigVMPythonUtils::Print(FString::Printf(TEXT("blueprint.get_controller_by_name('%s').set_comment_text_by_name('%s', '%s')"),
+			RigVMPythonUtils::Print(GetGraph()->GetRootGraph()->GetOuter()->GetFName().ToString(),
+				FString::Printf(TEXT("blueprint.get_controller_by_name('%s').set_comment_text_by_name('%s', '%s')"),
 				*GetGraph()->GetGraphName(),
 				*CommentNode->GetNodePath(),
 				*InCommentText));
@@ -5247,7 +5270,8 @@ bool URigVMController::SetPinExpansion(const FString& InPinPath, bool bIsExpande
 	const bool bSuccess = SetPinExpansion(Pin, bIsExpanded, bSetupUndoRedo);
 	if (bSuccess && bPrintPythonCommand)
 	{
-		RigVMPythonUtils::Print(FString::Printf(TEXT("blueprint.get_controller_by_name('%s').set_pin_expansion('%s', %s)"),
+		RigVMPythonUtils::Print(GetGraph()->GetRootGraph()->GetOuter()->GetFName().ToString(),
+			FString::Printf(TEXT("blueprint.get_controller_by_name('%s').set_pin_expansion('%s', %s)"),
 			*GetGraph()->GetGraphName(),
 			*InPinPath,
 			(bIsExpanded) ? TEXT("True") : TEXT("False")));
@@ -5434,7 +5458,8 @@ bool URigVMController::SetPinDefaultValue(const FString& InPinPath, const FStrin
 
 	if (bPrintPythonCommand)
 	{
-		RigVMPythonUtils::Print(FString::Printf(TEXT("blueprint.get_controller_by_name('%s').set_pin_default_value('%s', '%s', %s)"),
+		RigVMPythonUtils::Print(GetGraph()->GetRootGraph()->GetOuter()->GetFName().ToString(),
+			FString::Printf(TEXT("blueprint.get_controller_by_name('%s').set_pin_default_value('%s', '%s', %s)"),
 			*GetGraph()->GetGraphName(),
 			*InPinPath,
 			*InDefaultValue,
@@ -5578,7 +5603,8 @@ bool URigVMController::ResetPinDefaultValue(const FString& InPinPath, bool bSetu
 	const bool bSuccess = ResetPinDefaultValue(Pin, bSetupUndoRedo);
 	if (bSuccess && bPrintPythonCommand)
 	{
-		RigVMPythonUtils::Print(FString::Printf(TEXT("blueprint.get_controller_by_name('%s').reset_pin_default_value('%s')"),
+		RigVMPythonUtils::Print(GetGraph()->GetRootGraph()->GetOuter()->GetFName().ToString(),
+			FString::Printf(TEXT("blueprint.get_controller_by_name('%s').reset_pin_default_value('%s')"),
 			*GetGraph()->GetGraphName(),
 			*InPinPath));
 	}
@@ -5716,7 +5742,8 @@ FString URigVMController::InsertArrayPin(const FString& InArrayPinPath, int32 In
 	{
 		if (bPrintPythonCommand)
 		{
-			RigVMPythonUtils::Print(FString::Printf(TEXT("blueprint.get_controller_by_name('%s').insert_array_pin('%s', %d, '%s')"),
+			RigVMPythonUtils::Print(GetGraph()->GetRootGraph()->GetOuter()->GetFName().ToString(),
+				FString::Printf(TEXT("blueprint.get_controller_by_name('%s').insert_array_pin('%s', %d, '%s')"),
 				*GetGraph()->GetGraphName(),
 				*InArrayPinPath,
 				InIndex,
@@ -5870,7 +5897,8 @@ bool URigVMController::RemoveArrayPin(const FString& InArrayElementPinPath, bool
 
 	if (bPrintPythonCommand)
 	{
-		RigVMPythonUtils::Print(FString::Printf(TEXT("blueprint.get_controller_by_name('%s').remove_array_pin('%s')"),
+		RigVMPythonUtils::Print(GetGraph()->GetRootGraph()->GetOuter()->GetFName().ToString(),
+			FString::Printf(TEXT("blueprint.get_controller_by_name('%s').remove_array_pin('%s')"),
 			*GetGraph()->GetGraphName(),
 			*InArrayElementPinPath));
 	}
@@ -6026,7 +6054,8 @@ bool URigVMController::BindPinToVariable(const FString& InPinPath, const FString
 	const bool bSuccess = BindPinToVariable(Pin, InNewBoundVariablePath, bSetupUndoRedo);
 	if (bSuccess && bPrintPythonCommand)
 	{
-		RigVMPythonUtils::Print(FString::Printf(TEXT("blueprint.get_controller_by_name('%s').bind_pin_to_variable('%s', '%s')"),
+		RigVMPythonUtils::Print(GetGraph()->GetRootGraph()->GetOuter()->GetFName().ToString(),
+			FString::Printf(TEXT("blueprint.get_controller_by_name('%s').bind_pin_to_variable('%s', '%s')"),
 			*GetGraph()->GetGraphName(),
 			*InPinPath,
 			*InNewBoundVariablePath));
@@ -6245,7 +6274,8 @@ bool URigVMController::PromotePinToVariable(const FString& InPinPath, bool bCrea
 	const bool bSuccess = PromotePinToVariable(Pin, bCreateVariableNode, InNodePosition, bSetupUndoRedo);
 	if (bSuccess && bPrintPythonCommand)
 	{
-		RigVMPythonUtils::Print(FString::Printf(TEXT("blueprint.get_controller_by_name('%s').promote_pin_to_variable('%s', %s, %s)"),
+		RigVMPythonUtils::Print(GetGraph()->GetRootGraph()->GetOuter()->GetFName().ToString(), 
+			FString::Printf(TEXT("blueprint.get_controller_by_name('%s').promote_pin_to_variable('%s', %s, %s)"),
 			*GetGraph()->GetGraphName(),
 			*InPinPath,
 			(bCreateVariableNode) ? TEXT("True") : TEXT("False"),
@@ -6374,7 +6404,8 @@ bool URigVMController::AddLink(const FString& InOutputPinPath, const FString& In
 	const bool bSuccess = AddLink(OutputPin, InputPin, bSetupUndoRedo);
 	if (bSuccess && bPrintPythonCommand)
 	{		
-		RigVMPythonUtils::Print(FString::Printf(TEXT("blueprint.get_controller_by_name('%s').add_link('%s', '%s')"),
+		RigVMPythonUtils::Print(GetGraph()->GetRootGraph()->GetOuter()->GetFName().ToString(), 
+			FString::Printf(TEXT("blueprint.get_controller_by_name('%s').add_link('%s', '%s')"),
 			*GetGraph()->GetGraphName(),
 			*OutputPin->GetPinPath(),
 			*InputPin->GetPinPath()));
@@ -6514,7 +6545,8 @@ bool URigVMController::BreakLink(const FString& InOutputPinPath, const FString&
 	const bool bSuccess = BreakLink(OutputPin, InputPin, bSetupUndoRedo);
 	if (bSuccess && bPrintPythonCommand)
 	{
-		RigVMPythonUtils::Print(FString::Printf(TEXT("blueprint.get_controller_by_name('%s').break_link('%s', '%s')"),
+		RigVMPythonUtils::Print(GetGraph()->GetRootGraph()->GetOuter()->GetFName().ToString(), 
+			FString::Printf(TEXT("blueprint.get_controller_by_name('%s').break_link('%s', '%s')"),
 			*GetGraph()->GetGraphName(),
 			*OutputPin->GetPinPath(),
 			*InputPin->GetPinPath()));
@@ -6609,7 +6641,8 @@ bool URigVMController::BreakAllLinks(const FString& InPinPath, bool bAsInput, bo
 	const bool bSuccess = BreakAllLinks(Pin, bAsInput, bSetupUndoRedo);
 	if (bSuccess && bPrintPythonCommand)
 	{
-		RigVMPythonUtils::Print(FString::Printf(TEXT("blueprint.get_controller_by_name('%s').break_all_links('%s', %s)"),
+		RigVMPythonUtils::Print(GetGraph()->GetRootGraph()->GetOuter()->GetFName().ToString(), 
+			FString::Printf(TEXT("blueprint.get_controller_by_name('%s').break_all_links('%s', %s)"),
 			*GetGraph()->GetGraphName(),
 			*Pin->GetPinPath(),
 			bAsInput ? TEXT("True") : TEXT("False")));
@@ -6801,7 +6834,8 @@ FName URigVMController::AddExposedPin(const FName& InPinName, ERigVMPinDirection
 		}
 		
 		//AddExposedPin(const FName& InPinName, ERigVMPinDirection InDirection, const FString& InCPPType, const FName& InCPPTypeObjectPath, const FString& InDefaultValue, bool bSetupUndoRedo, bool bPrintPythonCommand)
-		RigVMPythonUtils::Print(FString::Printf(TEXT("blueprint.get_controller_by_name('%s').add_exposed_pin('%s', %s, '%s', '%s', '%s')"),
+		RigVMPythonUtils::Print(GetGraph()->GetRootGraph()->GetOuter()->GetFName().ToString(), 
+			FString::Printf(TEXT("blueprint.get_controller_by_name('%s').add_exposed_pin('%s', %s, '%s', '%s', '%s')"),
 				*GetGraph()->GetGraphName(),
 				*InPinName.ToString(),
 				*DirectionStr,
@@ -6887,7 +6921,8 @@ bool URigVMController::RemoveExposedPin(const FName& InPinName, bool bSetupUndoR
 
 	if (bSuccessfullyRemovedPin && bPrintPythonCommand)
 	{
-		RigVMPythonUtils::Print(FString::Printf(TEXT("blueprint.get_controller_by_name('%s').remove_exposed_pin('%s')"),
+		RigVMPythonUtils::Print(GetGraph()->GetRootGraph()->GetOuter()->GetFName().ToString(), 
+			FString::Printf(TEXT("blueprint.get_controller_by_name('%s').remove_exposed_pin('%s')"),
 				*GetGraph()->GetGraphName(),
 				*InPinName.ToString()));
 	}
@@ -7035,7 +7070,8 @@ bool URigVMController::RenameExposedPin(const FName& InOldPinName, const FName&
 
 	if (bPrintPythonCommand)
 	{
-		RigVMPythonUtils::Print(FString::Printf(TEXT("blueprint.get_controller_by_name('%s').rename_exposed_pin('%s', '%s')"),
+		RigVMPythonUtils::Print(GetGraph()->GetRootGraph()->GetOuter()->GetFName().ToString(), 
+			FString::Printf(TEXT("blueprint.get_controller_by_name('%s').rename_exposed_pin('%s', '%s')"),
 				*GetGraph()->GetGraphName(),
 				*InOldPinName.ToString(),
 				*InNewPinName.ToString()));
@@ -7146,7 +7182,8 @@ bool URigVMController::ChangeExposedPinType(const FName& InPinName, const FStrin
 
 	if (bPrintPythonCommand)
 	{
-		RigVMPythonUtils::Print(FString::Printf(TEXT("blueprint.get_controller_by_name('%s').change_exposed_pin_type('%s', '%s', '%s')"),
+		RigVMPythonUtils::Print(GetGraph()->GetRootGraph()->GetOuter()->GetFName().ToString(), 
+			FString::Printf(TEXT("blueprint.get_controller_by_name('%s').change_exposed_pin_type('%s', '%s', '%s')"),
 				*GetGraph()->GetGraphName(),
 				*InPinName.ToString(),
 				*InCPPType,
@@ -7218,7 +7255,8 @@ bool URigVMController::SetExposedPinIndex(const FName& InPinName, int32 InNewInd
 
 	if (bPrintPythonCommand)
 	{
-		RigVMPythonUtils::Print(FString::Printf(TEXT("blueprint.get_controller_by_name('%s').set_exposed_pin_index('%s', %d)"),
+		RigVMPythonUtils::Print(GetGraph()->GetRootGraph()->GetOuter()->GetFName().ToString(), 
+			FString::Printf(TEXT("blueprint.get_controller_by_name('%s').set_exposed_pin_index('%s', %d)"),
 				*GetGraph()->GetGraphName(),
 				*InPinName.ToString(),
 				InNewIndex));
@@ -7308,7 +7346,8 @@ URigVMFunctionReferenceNode* URigVMController::AddFunctionReferenceNode(URigVMLi
 	{
 		if (InFunctionDefinition->GetLibrary() == GetGraph()->GetDefaultFunctionLibrary())
 		{
-			RigVMPythonUtils::Print(FString::Printf(TEXT("blueprint.get_controller_by_name('%s').add_function_reference_node(library.find_function('%s'), %s, '%s')"),
+			RigVMPythonUtils::Print(GetGraph()->GetRootGraph()->GetOuter()->GetFName().ToString(), 
+				FString::Printf(TEXT("blueprint.get_controller_by_name('%s').add_function_reference_node(library.find_function('%s'), %s, '%s')"),
 						*GetGraph()->GetGraphName(),
 						*InFunctionDefinition->GetName(),
 						*RigVMPythonUtils::Vector2DToPythonString(InNodePosition),
@@ -7316,9 +7355,11 @@ URigVMFunctionReferenceNode* URigVMController::AddFunctionReferenceNode(URigVMLi
 		}
 		else
 		{
-			RigVMPythonUtils::Print(FString::Printf(TEXT("function_blueprint = unreal.load_object(name = '%s', outer = None)"),
+			RigVMPythonUtils::Print(GetGraph()->GetRootGraph()->GetOuter()->GetFName().ToString(), 
+				FString::Printf(TEXT("function_blueprint = unreal.load_object(name = '%s', outer = None)"),
 				*InFunctionDefinition->GetLibrary()->GetOuter()->GetPathName()));
-			RigVMPythonUtils::Print(FString::Printf(TEXT("blueprint.get_controller_by_name('%s').add_function_reference_node(function_blueprint.get_local_function_library().find_function('%s'), %s, '%s')"),
+			RigVMPythonUtils::Print(GetGraph()->GetRootGraph()->GetOuter()->GetFName().ToString(), 
+				FString::Printf(TEXT("blueprint.get_controller_by_name('%s').add_function_reference_node(function_blueprint.get_local_function_library().find_function('%s'), %s, '%s')"),
 						*GetGraph()->GetGraphName(),
 						*InFunctionDefinition->GetName(),
 						*RigVMPythonUtils::Vector2DToPythonString(InFunctionDefinition->GetPosition()), 
@@ -7485,7 +7526,8 @@ URigVMLibraryNode* URigVMController::AddFunctionToLibrary(const FName& InFunctio
 	if (bPrintPythonCommand)
 	{
 		//AddFunctionToLibrary(const FName& InFunctionName, bool bMutable, const FVector2D& InNodePosition, bool bSetupUndoRedo, bool bPrintPythonCommand)
-		RigVMPythonUtils::Print(FString::Printf(TEXT("library_controller.add_function_to_library('%s', %s, %s)"),
+		RigVMPythonUtils::Print(GetGraph()->GetRootGraph()->GetOuter()->GetFName().ToString(), 
+			FString::Printf(TEXT("library_controller.add_function_to_library('%s', %s, %s)"),
 				*InFunctionName.ToString(),
 				(bMutable) ? TEXT("True") : TEXT("False"),
 				*RigVMPythonUtils::Vector2DToPythonString(InNodePosition)));
@@ -7587,7 +7629,8 @@ FRigVMGraphVariableDescription URigVMController::AddLocalVariable(const FName& I
 
 	if (bPrintPythonCommand)
 	{
-		RigVMPythonUtils::Print(FString::Printf(TEXT("blueprint.get_controller_by_name('%s').add_local_variable_from_object_path('%s', '%s', '%s', '%s')"),
+		RigVMPythonUtils::Print(GetGraph()->GetRootGraph()->GetOuter()->GetFName().ToString(), 
+			FString::Printf(TEXT("blueprint.get_controller_by_name('%s').add_local_variable_from_object_path('%s', '%s', '%s', '%s')"),
 				*GetGraph()->GetGraphName(),
 				*NewVariable.Name.ToString(),
 				*NewVariable.CPPType,
@@ -7688,7 +7731,8 @@ bool URigVMController::RemoveLocalVariable(const FName& InVariableName, bool bSe
 
 		if (bPrintPythonCommand)
 		{
-			RigVMPythonUtils::Print(FString::Printf(TEXT("blueprint.get_controller_by_name('%s').remove_local_variable('%s')"),
+			RigVMPythonUtils::Print(GetGraph()->GetRootGraph()->GetOuter()->GetFName().ToString(), 
+				FString::Printf(TEXT("blueprint.get_controller_by_name('%s').remove_local_variable('%s')"),
 					*GetGraph()->GetGraphName(),
 					*InVariableName.ToString()));
 		}
@@ -7784,7 +7828,8 @@ bool URigVMController::RenameLocalVariable(const FName& InVariableName, const FN
 
 	if (bPrintPythonCommand)
 	{
-		RigVMPythonUtils::Print(FString::Printf(TEXT("blueprint.get_controller_by_name('%s').rename_local_variable('%s', '%s')"),
+		RigVMPythonUtils::Print(GetGraph()->GetRootGraph()->GetOuter()->GetFName().ToString(), 
+			FString::Printf(TEXT("blueprint.get_controller_by_name('%s').rename_local_variable('%s', '%s')"),
 				*GetGraph()->GetGraphName(),
 				*InVariableName.ToString(),
 				*InNewVariableName.ToString()));
@@ -7875,7 +7920,8 @@ bool URigVMController::SetLocalVariableType(const FName& InVariableName, const F
 	if (bPrintPythonCommand)
 	{
 		//bool URigVMController::SetLocalVariableType(const FName& InVariableName, const FString& InCPPType, UObject* InCPPTypeObject, bool bSetupUndoRedo, bool bPrintPythonCommand)
-		RigVMPythonUtils::Print(FString::Printf(TEXT("blueprint.get_controller_by_name('%s').set_local_variable_type_from_object_path('%s', '%s', '%s')"),
+		RigVMPythonUtils::Print(GetGraph()->GetRootGraph()->GetOuter()->GetFName().ToString(), 
+			FString::Printf(TEXT("blueprint.get_controller_by_name('%s').set_local_variable_type_from_object_path('%s', '%s', '%s')"),
 				*GetGraph()->GetGraphName(),
 				*InVariableName.ToString(),
 				*InCPPType,
@@ -8316,7 +8362,8 @@ URigVMBranchNode* URigVMController::AddBranchNode(const FVector2D& InPosition, c
 		TArray<FString> Commands = GetAddNodePythonCommands(Node);
 		for (const FString& Command : Commands)
 		{
-			RigVMPythonUtils::Print(FString::Printf(TEXT("%s"), *Command));
+			RigVMPythonUtils::Print(GetGraph()->GetRootGraph()->GetOuter()->GetFName().ToString(), 
+				FString::Printf(TEXT("%s"), *Command));
 		}
 	}
 
@@ -8421,7 +8468,8 @@ URigVMIfNode* URigVMController::AddIfNode(const FString& InCPPType, const FName&
 		TArray<FString> Commands = GetAddNodePythonCommands(Node);
 		for (const FString& Command : Commands)
 		{
-			RigVMPythonUtils::Print(FString::Printf(TEXT("%s"), *Command));
+			RigVMPythonUtils::Print(GetGraph()->GetRootGraph()->GetOuter()->GetFName().ToString(), 
+				FString::Printf(TEXT("%s"), *Command));
 		}
 	}
 
@@ -8505,7 +8553,8 @@ URigVMSelectNode* URigVMController::AddSelectNode(const FString& InCPPType, cons
 		TArray<FString> Commands = GetAddNodePythonCommands(Node);
 		for (const FString& Command : Commands)
 		{
-			RigVMPythonUtils::Print(FString::Printf(TEXT("%s"), *Command));
+			RigVMPythonUtils::Print(GetGraph()->GetRootGraph()->GetOuter()->GetFName().ToString(), 
+				FString::Printf(TEXT("%s"), *Command));
 		}
 	}
 
@@ -8581,7 +8630,8 @@ URigVMPrototypeNode* URigVMController::AddPrototypeNode(const FName& InNotation,
 		TArray<FString> Commands = GetAddNodePythonCommands(Node);
 		for (const FString& Command : Commands)
 		{
-			RigVMPythonUtils::Print(FString::Printf(TEXT("%s"), *Command));
+			RigVMPythonUtils::Print(GetGraph()->GetRootGraph()->GetOuter()->GetFName().ToString(), 
+				FString::Printf(TEXT("%s"), *Command));
 		}
 	}
 
@@ -8645,7 +8695,8 @@ URigVMEnumNode* URigVMController::AddEnumNode(const FName& InCPPTypeObjectPath,
 		TArray<FString> Commands = GetAddNodePythonCommands(Node);
 		for (const FString& Command : Commands)
 		{
-			RigVMPythonUtils::Print(FString::Printf(TEXT("%s"), *Command));
+			RigVMPythonUtils::Print(GetGraph()->GetRootGraph()->GetOuter()->GetFName().ToString(), 
+				FString::Printf(TEXT("%s"), *Command));
 		}
 	}
 
diff --git a/Engine/Source/Developer/RigVMDeveloper/Private/RigVMPythonUtils.cpp b/Engine/Source/Developer/RigVMDeveloper/Private/RigVMPythonUtils.cpp
--- a/Engine/Source/Developer/RigVMDeveloper/Private/RigVMPythonUtils.cpp
+++ b/Engine/Source/Developer/RigVMDeveloper/Private/RigVMPythonUtils.cpp
@@ -64,7 +64,7 @@ FString RigVMPythonUtils::LinearColorToPythonString(const FLinearColor& Color)
 }
 
 #if WITH_EDITOR
-void RigVMPythonUtils::Print(const FString& InMessage)
+void RigVMPythonUtils::Print(const FString& BlueprintTitle, const FString& InMessage)
 {
 	FMessageLogModule& MessageLogModule = FModuleManager::LoadModuleChecked<FMessageLogModule>("MessageLog");
 
@@ -78,10 +78,10 @@ void RigVMPythonUtils::Print(const FString& InMessage)
 		MessageLogModule.RegisterLogListing("ControlRigPythonLog", LOCTEXT("ControlRigPythonLog", "Control Rig Python Log"), InitOptions);
 	}
 	TSharedRef<IMessageLogListing> PythonLog = MessageLogModule.GetLogListing( TEXT("ControlRigPythonLog") );
+	PythonLog->SwitchToPage(FText::FromString(BlueprintTitle));
 
 	TSharedRef<FTokenizedMessage> Token = FTokenizedMessage::Create(EMessageSeverity::Info, FText::FromString(InMessage));
-	PythonLog->AddMessage(Token, false);			
-	
+	PythonLog->AddMessage(Token, false);
 }
 
 #endif
diff --git a/Engine/Source/Developer/RigVMDeveloper/Public/RigVMPythonUtils.h b/Engine/Source/Developer/RigVMDeveloper/Public/RigVMPythonUtils.h
--- a/Engine/Source/Developer/RigVMDeveloper/Public/RigVMPythonUtils.h
+++ b/Engine/Source/Developer/RigVMDeveloper/Public/RigVMPythonUtils.h
@@ -15,7 +15,7 @@ namespace RigVMPythonUtils
 	RIGVMDEVELOPER_API FString LinearColorToPythonString(const FLinearColor& Color);
 	
 #if WITH_EDITOR
-	RIGVMDEVELOPER_API void Print(const FString& InMessage);
+	RIGVMDEVELOPER_API void Print(const FString& BlueprintTitle, const FString& InMessage);
 
 #endif
 }
diff --git a/Engine/Source/Runtime/Core/Private/Logging/MessageLog.cpp b/Engine/Source/Runtime/Core/Private/Logging/MessageLog.cpp
--- a/Engine/Source/Runtime/Core/Private/Logging/MessageLog.cpp
+++ b/Engine/Source/Runtime/Core/Private/Logging/MessageLog.cpp
@@ -39,6 +39,13 @@ public:
 		FMsg::Logf(__FILE__, __LINE__, LogName, ELogVerbosity::Log, TEXT("%s"), *FText::Format(LOCTEXT("BasicMessageLog_NewPage", "New Page: {PageTitle}"), Arguments).ToString());
 	}
 
+	virtual void SwitchToPage( const FText& Title) override
+	{
+		FFormatNamedArguments Arguments;
+		Arguments.Add(TEXT("PageTitle"), Title);
+		FMsg::Logf(__FILE__, __LINE__, LogName, ELogVerbosity::Log, TEXT("%s"), *FText::Format(LOCTEXT("BasicMessageLog_SwitchToPage", "Switch To Page: {PageTitle}"), Arguments).ToString());
+	}
+
 	virtual void NotifyIfAnyMessages( const FText& Message, EMessageSeverity::Type SeverityFilter, bool bForce ) override
 	{
 		FFormatNamedArguments Arguments;
@@ -184,6 +191,12 @@ void FMessageLog::NewPage( const FText& InLabel )
 	MessageLog->NewPage(InLabel);
 }
 
+void FMessageLog::SwitchToPage(const FText& InLabel)
+{
+	Flush();
+	MessageLog->SwitchToPage(InLabel);
+}
+
 FMessageLog& FMessageLog::SuppressLoggingToOutputLog(bool bShouldSuppress)
 {
 	bSuppressLoggingToOutputLog = bShouldSuppress;
diff --git a/Engine/Source/Runtime/Core/Public/Logging/IMessageLog.h b/Engine/Source/Runtime/Core/Public/Logging/IMessageLog.h
--- a/Engine/Source/Runtime/Core/Public/Logging/IMessageLog.h
+++ b/Engine/Source/Runtime/Core/Public/Logging/IMessageLog.h
@@ -44,6 +44,12 @@ public:
 	 */
 	virtual void NewPage( const FText& Title ) = 0;
 
+	/** 
+	* Switches the current page to the one specified by the title.
+	* @param	Title	the title of the page to switch to
+	*/
+	virtual void SwitchToPage( const FText& Title ) = 0;
+
 	/** 
 	 * Notify the user if there are any messages on the current page for this log.
 	 * If there are no messages present, this call does nothing.
diff --git a/Engine/Source/Runtime/Core/Public/Logging/MessageLog.h b/Engine/Source/Runtime/Core/Public/Logging/MessageLog.h
--- a/Engine/Source/Runtime/Core/Public/Logging/MessageLog.h
+++ b/Engine/Source/Runtime/Core/Public/Logging/MessageLog.h
@@ -87,6 +87,13 @@ public:
 	 */
 	CORE_API void NewPage( const FText& InLabel );
 
+	/**
+	* Switches the current page to the one specified by the label.
+	* This call will cause a flush so that the logs state is properly reflected.
+	* @param	InLabel		The label for the page.
+	*/
+	CORE_API void SwitchToPage( const FText& InLabel );
+
 	/** Should we mirror message log messages from this instance to the output log during flush? */
 	CORE_API FMessageLog& SuppressLoggingToOutputLog(bool bShouldSuppress = true);
 
