Control Rig: Fix compilation error after changing Message log

diff --git a/Engine/Source/Developer/RigVMDeveloper/Private/RigVMModel/RigVMController.cpp b/Engine/Source/Developer/RigVMDeveloper/Private/RigVMModel/RigVMController.cpp
--- a/Engine/Source/Developer/RigVMDeveloper/Private/RigVMModel/RigVMController.cpp
+++ b/Engine/Source/Developer/RigVMDeveloper/Private/RigVMModel/RigVMController.cpp
@@ -2489,7 +2489,8 @@ TArray<FName> URigVMController::ImportNodesFromText(const FString& InText, bool
 		PythonContent = InText.Replace(TEXT("'"), TEXT("\\'"));
 		PythonContent = PythonContent.Replace(TEXT("\r\n"), TEXT("\\r\\n'\r\n'"));
 		
-		RigVMPythonUtils::Print(FString::Printf(TEXT("blueprint.get_controller_by_name('%s').import_nodes_from_text('%s')"),
+		RigVMPythonUtils::Print(GetGraph()->GetRootGraph()->GetOuter()->GetFName().ToString(),
+			FString::Printf(TEXT("blueprint.get_controller_by_name('%s').import_nodes_from_text('%s')"),
 			*GetGraph()->GetGraphName(),
 			*PythonContent));
 	}
