Control Rig: Re-enable RigVMCompiler Unit Tests for new memory storage class

diff --git a/Engine/Plugins/Animation/ControlRig/Source/ControlRig/Private/Units/Simulation/RigUnit_Kalman.cpp b/Engine/Plugins/Animation/ControlRig/Source/ControlRig/Private/Units/Simulation/RigUnit_Kalman.cpp
--- a/Engine/Plugins/Animation/ControlRig/Source/ControlRig/Private/Units/Simulation/RigUnit_Kalman.cpp
+++ b/Engine/Plugins/Animation/ControlRig/Source/ControlRig/Private/Units/Simulation/RigUnit_Kalman.cpp
@@ -17,6 +17,10 @@ FRigUnit_KalmanFloat_Execute()
 		return;
 	}
 
+#if !UE_RIGVM_UCLASS_BASED_STORAGE_DISABLED
+	Buffer.SetNum(BufferSize);
+#endif
+
 	if (Context.State == EControlRigState::Init)
 	{
 		for (float& Element : Buffer)
@@ -66,6 +70,10 @@ FRigUnit_KalmanVector_Execute()
 		return;
 	}
 
+#if !UE_RIGVM_UCLASS_BASED_STORAGE_DISABLED
+	Buffer.SetNum(BufferSize);
+#endif
+
 	if (Context.State == EControlRigState::Init)
 	{
 		for (FVector& Element : Buffer)
@@ -115,6 +123,10 @@ FRigUnit_KalmanTransform_Execute()
 		return;
 	}
 
+#if !UE_RIGVM_UCLASS_BASED_STORAGE_DISABLED
+	Buffer.SetNum(BufferSize);
+#endif
+
 	if (Context.State == EControlRigState::Init)
 	{
 		for (FTransform& Element : Buffer)
diff --git a/Engine/Plugins/Animation/ControlRig/Source/ControlRig/Private/Units/Simulation/RigUnit_TimeOffset.cpp b/Engine/Plugins/Animation/ControlRig/Source/ControlRig/Private/Units/Simulation/RigUnit_TimeOffset.cpp
--- a/Engine/Plugins/Animation/ControlRig/Source/ControlRig/Private/Units/Simulation/RigUnit_TimeOffset.cpp
+++ b/Engine/Plugins/Animation/ControlRig/Source/ControlRig/Private/Units/Simulation/RigUnit_TimeOffset.cpp
@@ -29,6 +29,9 @@ FRigUnit_TimeOffsetFloat_Execute()
 	}
 
 	int32 MaxSize = FMath::Clamp<int32>(BufferSize, 2, 512);
+#if !UE_RIGVM_UCLASS_BASED_STORAGE_DISABLED
+	Buffer.SetNum(MaxSize);
+#endif
 	MaxSize = FMath::Min<int32>(MaxSize, Buffer.Num());
 
 	if (Context.State == EControlRigState::Init)
@@ -136,6 +139,9 @@ FRigUnit_TimeOffsetVector_Execute()
 	}
 
 	int32 MaxSize = FMath::Clamp<int32>(BufferSize, 2, 512);
+#if !UE_RIGVM_UCLASS_BASED_STORAGE_DISABLED
+	Buffer.SetNum(MaxSize);
+#endif
 	MaxSize = FMath::Min<int32>(MaxSize, Buffer.Num());
 
 	if (Context.State == EControlRigState::Init)
@@ -243,6 +249,9 @@ FRigUnit_TimeOffsetTransform_Execute()
 	}
 
 	int32 MaxSize = FMath::Clamp<int32>(BufferSize, 2, 512);
+#if !UE_RIGVM_UCLASS_BASED_STORAGE_DISABLED
+	Buffer.SetNum(MaxSize);
+#endif
 	MaxSize = FMath::Min<int32>(MaxSize, Buffer.Num());
 
 	if (Context.State == EControlRigState::Init)
diff --git a/Engine/Plugins/Animation/ControlRig/Source/ControlRigDeveloper/Private/ControlRigBlueprint.cpp b/Engine/Plugins/Animation/ControlRig/Source/ControlRigDeveloper/Private/ControlRigBlueprint.cpp
--- a/Engine/Plugins/Animation/ControlRig/Source/ControlRigDeveloper/Private/ControlRigBlueprint.cpp
+++ b/Engine/Plugins/Animation/ControlRig/Source/ControlRigDeveloper/Private/ControlRigBlueprint.cpp
@@ -2839,12 +2839,28 @@ void UControlRigBlueprint::HandleModifiedEvent(ERigVMGraphNotifType InNotifType,
 						}
 						else
 						{
+#if UE_RIGVM_UCLASS_BASED_STORAGE_DISABLED
 							Compiler->CreateDebugRegister(Pin, CR->GetVM(), &PinToOperandMap, RuntimeAST);
+#else
+							if(CR->GetVM()->GetDebugMemory()->Num() == 0)
+							{
+								RequestAutoVMRecompilation();
+								MarkPackageDirty();
+							}
+							else
+							{
+								Compiler->MarkDebugWatch(true, Pin, CR->GetVM(), &PinToOperandMap, RuntimeAST);
+							}
+#endif
 						}
 					}
 					else
 					{
+#if UE_RIGVM_UCLASS_BASED_STORAGE_DISABLED
 						Compiler->RemoveDebugRegister(Pin, CR->GetVM(), &PinToOperandMap, RuntimeAST);
+#else
+						Compiler->MarkDebugWatch(true, Pin, CR->GetVM(), &PinToOperandMap, RuntimeAST);
+#endif
 					}
 				}
 				// break; fall through
diff --git a/Engine/Source/Developer/RigVMDeveloper/Private/RigVMCompiler/RigVMCompiler.cpp b/Engine/Source/Developer/RigVMDeveloper/Private/RigVMCompiler/RigVMCompiler.cpp
--- a/Engine/Source/Developer/RigVMDeveloper/Private/RigVMCompiler/RigVMCompiler.cpp
+++ b/Engine/Source/Developer/RigVMDeveloper/Private/RigVMCompiler/RigVMCompiler.cpp
@@ -105,38 +105,23 @@ FRigVMPropertyDescription FRigVMCompilerWorkData::GetProperty(const FRigVMOperan
 	return FRigVMPropertyDescription();
 }
 
-int32 FRigVMCompilerWorkData::FindPropertyPath(ERigVMMemoryType InMemoryType, const FString& InSegmentPath) const
+int32 FRigVMCompilerWorkData::FindOrAddPropertyPath(const FRigVMOperand& InOperand, const FString& InRootCPPType, const FString& InSegmentPath)
 {
-	const TArray<FRigVMPropertyPathDescription>* Descriptions = PropertyPathDescriptions.Find(InMemoryType);
-	if(Descriptions)
+	if(InSegmentPath.IsEmpty())
 	{
-		for(int32 Index = 0; Index < Descriptions->Num(); Index++)
-		{
-			const FRigVMPropertyPathDescription& Description = Descriptions->operator[](Index); 
-			if(Description.SegmentPath == InSegmentPath)
-			{
-				return Index;
-			}
-		}
+		return INDEX_NONE;
 	}
 
-	return INDEX_NONE;
-}
-
-int32 FRigVMCompilerWorkData::FindOrAddPropertyPath(const FRigVMOperand& InOperand, const FString& InSegmentPath)
-{
-	check(bSetupMemory);
-	
 	TArray<FRigVMPropertyPathDescription>& Descriptions = PropertyPathDescriptions.FindOrAdd(InOperand.GetMemoryType());
 	for(int32 Index = 0; Index < Descriptions.Num(); Index++)
 	{
 		const FRigVMPropertyPathDescription& Description = Descriptions[Index]; 
-		if(Description.SegmentPath == InSegmentPath)
+		if(Description.RootCPPType == InRootCPPType && Description.SegmentPath == InSegmentPath)
 		{
 			return Index;
 		}
 	}
-	return Descriptions.Add(FRigVMPropertyPathDescription(InOperand.GetRegisterIndex(), InSegmentPath));
+	return Descriptions.Add(FRigVMPropertyPathDescription(InOperand.GetRegisterIndex(), InRootCPPType, InSegmentPath));
 }
 
 #endif
@@ -224,8 +209,12 @@ bool URigVMCompiler::Compile(URigVMGraph* InGraph, URigVMController* InControlle
 #endif
 	
 	OutVM->ClearExternalVariables();
+	
 	for (const FRigVMExternalVariable& ExternalVariable : InExternalVariables)
 	{
+#if !UE_RIGVM_UCLASS_BASED_STORAGE_DISABLED
+		check(ExternalVariable.Property);
+#endif
 		FRigVMOperand Operand = OutVM->AddExternalVariable(ExternalVariable);
 		FString Hash = FString::Printf(TEXT("Variable::%s"), *ExternalVariable.Name.ToString());
 		OutOperands->Add(Hash, Operand);
@@ -276,7 +265,8 @@ bool URigVMCompiler::Compile(URigVMGraph* InGraph, URigVMController* InControlle
 							WorkData.VM->GetLiteralMemory().SetRegisterValueFromString(Operand, LocalVariable.CPPType, LocalVariable.CPPTypeObject, {LocalVariable.DefaultValue});
 #else
 							check(Operand.GetMemoryType() == ERigVMMemoryType::Literal);
-							WorkData.VM->GetLiteralMemory()->SetDataFromString(Operand.GetRegisterIndex(), LocalVariable.DefaultValue);
+							TArray<FRigVMPropertyDescription>& LiteralProperties = WorkData.PropertyDescriptions.FindChecked(Operand.GetMemoryType());
+							LiteralProperties[Operand.GetRegisterIndex()].DefaultValue = LocalVariable.DefaultValue;
 #endif
 							break;
 						}
@@ -286,6 +276,26 @@ bool URigVMCompiler::Compile(URigVMGraph* InGraph, URigVMController* InControlle
 		}
 	}	
 
+#if !UE_RIGVM_UCLASS_BASED_STORAGE_DISABLED
+	if(Settings.EnablePinWatches)
+	{
+		for(int32 GraphIndex=0; GraphIndex<VisitedGraphs.Num(); GraphIndex++)
+		{
+			URigVMGraph* VisitedGraph = VisitedGraphs[GraphIndex];
+			for(URigVMNode* ModelNode : VisitedGraph->GetNodes())
+			{
+				for(URigVMPin* ModelPin : ModelNode->GetPins())
+				{
+					if(ModelPin->RequiresWatch(true))
+					{
+						WorkData.WatchedPins.AddUnique(ModelPin);
+					}
+				}
+			}
+		}
+	}
+#endif
+
 	// define all parameters independent from sorted nodes
 	for (const FRigVMExprAST* Expr : AST->Expressions)
 	{
@@ -339,6 +349,32 @@ bool URigVMCompiler::Compile(URigVMGraph* InGraph, URigVMController* InControlle
 
 #if !UE_RIGVM_UCLASS_BASED_STORAGE_DISABLED
 
+	if(WorkData.WatchedPins.Num() > 0)
+	{
+		for(int32 GraphIndex=0; GraphIndex<VisitedGraphs.Num(); GraphIndex++)
+		{
+			URigVMGraph* VisitedGraph = VisitedGraphs[GraphIndex];
+			for(URigVMNode* ModelNode : VisitedGraph->GetNodes())
+			{
+				for(URigVMPin* ModelPin : ModelNode->GetPins())
+				{
+					if(ModelPin->GetDirection() == ERigVMPinDirection::Input)
+					{
+						if(ModelPin->GetSourceLinks(true).Num() == 0)
+						{
+							continue;
+						}
+					}
+					FRigVMASTProxy PinProxy = FRigVMASTProxy::MakeFromUObject(ModelPin);
+					FRigVMVarExprAST TempVarExpr(FRigVMExprAST::EType::Var, PinProxy);
+					TempVarExpr.ParserPtr = AST.Get();
+
+					FindOrAddRegister(&TempVarExpr, WorkData, true);
+				}
+			}
+		}
+	}
+
 	// now that we have determined the needed memory, let's
 	// setup properties as needed as well as property paths
 	TArray<ERigVMMemoryType> MemoryTypes;
@@ -404,6 +440,7 @@ bool URigVMCompiler::Compile(URigVMGraph* InGraph, URigVMController* InControlle
 	// setup debug registers after all other registers have been created
 	if(Settings.EnablePinWatches)
 	{
+#if UE_RIGVM_UCLASS_BASED_STORAGE_DISABLED
 		for(int32 GraphIndex=0; GraphIndex<VisitedGraphs.Num(); GraphIndex++)
 		{
 			URigVMGraph* VisitedGraph = VisitedGraphs[GraphIndex];
@@ -418,9 +455,44 @@ bool URigVMCompiler::Compile(URigVMGraph* InGraph, URigVMController* InControlle
 				}
 			}
 		}
+#else
+		for(URigVMPin* WatchedPin : WorkData.WatchedPins)
+		{
+			MarkDebugWatch(true, WatchedPin, WorkData.VM, WorkData.PinPathToOperand, AST);
+		}
+#endif
+	}
+
+#if !UE_RIGVM_UCLASS_BASED_STORAGE_DISABLED
+
+	// now that we have determined the needed memory, let's
+	// update the property paths once more
+	for(ERigVMMemoryType MemoryType : MemoryTypes)
+	{
+		const TArray<FRigVMPropertyPathDescription>* Descriptions = WorkData.PropertyPathDescriptions.Find(MemoryType);
+		if(URigVMMemoryStorage* MemoryStorageObject = WorkData.VM->GetMemoryByType(MemoryType))
+		{
+			if(URigVMMemoryStorageGeneratorClass* Class = Cast<URigVMMemoryStorageGeneratorClass>(MemoryStorageObject->GetClass()))
+			{
+				if(Descriptions)
+				{
+					Class->PropertyPathDescriptions = *Descriptions;
+				}
+				else
+				{
+					Class->PropertyPathDescriptions.Reset();
+				}
+				Class->RefreshPropertyPaths();
+			}
+		}
 	}
 
-	UE_LOG_RIGVMMEMORY(TEXT("RigVMCompiler: Finished '%s'."), *InGraph->GetPathName());
+	if(const TArray<FRigVMPropertyPathDescription>* Descriptions = WorkData.PropertyPathDescriptions.Find(ERigVMMemoryType::External))
+	{
+		WorkData.VM->ExternalPropertyPathDescriptions = *Descriptions;
+	}
+
+#endif
 
 	return true;
 }
@@ -842,6 +914,8 @@ void URigVMCompiler::TraverseAssign(const FRigVMAssignExprAST* InExpr, FRigVMCom
 
 					FString SegmentPath = Pin->GetSegmentPath();
 
+#if UE_RIGVM_UCLASS_BASED_STORAGE_DISABLED
+
 					int32 ArrayIndex = INDEX_NONE;
 					if (RootPin->IsArray())
 					{
@@ -853,14 +927,12 @@ void URigVMCompiler::TraverseAssign(const FRigVMAssignExprAST* InExpr, FRigVMCom
 							SegmentPath = SegmentPath.Mid(DotIndex + 1);
 						}
 						else
-						{
+						{		
 							SegmentPath = FString();
 						}
 						ArrayIndex = FCString::Atoi(*SegmentArrayIndex);
 					}
-
-#if UE_RIGVM_UCLASS_BASED_STORAGE_DISABLED
-					
+				
 					if (Operand.GetMemoryType() == ERigVMMemoryType::External)
 					{
 						const FRigVMExternalVariable& ExternalVariable = VM->GetExternalVariables()[Operand.GetRegisterIndex()];
@@ -877,7 +949,7 @@ void URigVMCompiler::TraverseAssign(const FRigVMAssignExprAST* InExpr, FRigVMCom
 
 #else
 
-					const int32 PropertyPathIndex = WorkData.FindOrAddPropertyPath(Operand, SegmentPath);
+					const int32 PropertyPathIndex = WorkData.FindOrAddPropertyPath(Operand, RootPin->GetCPPType(), SegmentPath);
 					Operand = FRigVMOperand(Operand.GetMemoryType(), Operand.GetRegisterIndex(), PropertyPathIndex);
 					
 #endif
@@ -1457,7 +1529,7 @@ FString URigVMCompiler::GetPinHash(const URigVMPin* InPin, const FRigVMVarExprAS
 	bool bUseFullNodePath = true;
 	if (URigVMParameterNode* ParameterNode = Cast<URigVMParameterNode>(Node))
 	{
-		if (InPin->GetName() == TEXT("Value") && !bIsLiteral)
+		if (InPin->GetName() == TEXT("Value") && !bIsLiteral && !bIsDebugValue)
 		{
 			return FString::Printf(TEXT("%sParameter::%s%s"), *Prefix, *ParameterNode->GetParameterName().ToString(), *Suffix);
 		}
@@ -1531,7 +1603,11 @@ FString URigVMCompiler::GetPinHash(const URigVMPin* InPin, const FRigVMVarExprAS
 	}
 	else
 	{
+#if UE_RIGVM_UCLASS_BASED_STORAGE_DISABLED
 		if (InVarExpr)
+#else
+		if (InVarExpr && !bIsDebugValue)
+#endif
 		{
 			const FRigVMASTProxy NodeProxy = InVarExpr->GetProxy().GetSibling(Node);
 			if (const FRigVMExprAST* NodeExpr = InVarExpr->GetParser()->GetExprForSubject(NodeProxy))
@@ -1550,6 +1626,8 @@ FString URigVMCompiler::GetPinHash(const URigVMPin* InPin, const FRigVMVarExprAS
 	return FString::Printf(TEXT("%s%s%s"), *Prefix, *InPin->GetPinPath(bUseFullNodePath), *Suffix);
 }
 
+#if UE_RIGVM_UCLASS_BASED_STORAGE_DISABLED
+
 void URigVMCompiler::CreateDebugRegister(URigVMPin* InPin, URigVM* OutVM,
 	TMap<FString, FRigVMOperand>* OutOperands, TSharedPtr<FRigVMParserAST> InRuntimeAST)
 {
@@ -1627,6 +1705,85 @@ void URigVMCompiler::RemoveDebugRegister(URigVMPin* InPin, URigVM* OutVM,
 	}
 }
 
+#else
+
+void URigVMCompiler::MarkDebugWatch(bool bRequired, URigVMPin* InPin, URigVM* OutVM,
+	TMap<FString, FRigVMOperand>* OutOperands, TSharedPtr<FRigVMParserAST> InRuntimeAST)
+{
+	check(InPin);
+	check(OutVM);
+	check(OutOperands);
+	check(InRuntimeAST.IsValid());
+	
+	URigVMPin* Pin = InPin->GetRootPin();
+	URigVMPin* SourcePin = Pin;
+	if(Settings.ASTSettings.bFoldAssignments)
+	{
+		while(SourcePin->GetSourceLinks().Num() > 0)
+		{
+			SourcePin = SourcePin->GetSourceLinks()[0]->GetSourcePin();
+		}
+	}
+		
+	TArray<const FRigVMExprAST*> Expressions = InRuntimeAST->GetExpressionsForSubject(SourcePin);
+	TArray<FRigVMOperand> VisitedKeys;
+	for(const FRigVMExprAST* Expression : Expressions)
+	{
+		check(Expression->IsA(FRigVMExprAST::EType::Var));
+		const FRigVMVarExprAST* VarExpression = Expression->To<FRigVMVarExprAST>();
+
+		if(VarExpression->GetPin() == Pin)
+		{
+			// literals don't need to be stored on the debug memory
+			if(VarExpression->IsA(FRigVMExprAST::Literal))
+			{
+				continue;
+			}
+		}
+		
+		const FString PinHash = GetPinHash(SourcePin, VarExpression, false);
+		if(const FRigVMOperand* Operand = OutOperands->Find(PinHash))
+		{
+			const FRigVMASTProxy PinProxy = FRigVMASTProxy::MakeFromUObject(Pin);
+			FRigVMVarExprAST TempVarExpr(FRigVMExprAST::EType::Var, PinProxy);
+			TempVarExpr.ParserPtr = InRuntimeAST.Get();
+
+			const FString DebugPinHash = GetPinHash(Pin, &TempVarExpr, true);
+			const FRigVMOperand* DebugOperand = OutOperands->Find(DebugPinHash);
+			if(DebugOperand)
+			{
+				if(DebugOperand->IsValid())
+				{
+					FRigVMOperand KeyOperand(Operand->GetMemoryType(), Operand->GetRegisterIndex()); // no register offset
+					if(bRequired)
+					{
+						if(!VisitedKeys.Contains(KeyOperand))
+						{
+							OutVM->OperandToDebugRegisters.FindOrAdd(KeyOperand).AddUnique(*DebugOperand);
+							VisitedKeys.Add(KeyOperand);
+						}
+					}
+					else
+					{
+						TArray<FRigVMOperand>* MappedOperands = OutVM->OperandToDebugRegisters.Find(KeyOperand);
+						if(MappedOperands)
+						{
+							MappedOperands->Remove(*DebugOperand);
+
+							if(MappedOperands->IsEmpty())
+							{
+								OutVM->OperandToDebugRegisters.Remove(KeyOperand);
+							}
+						}
+					}
+				}
+			}
+		}
+	}
+}
+
+#endif
+
 UScriptStruct* URigVMCompiler::GetScriptStructForCPPType(const FString& InCPPType)
 {
 	if (InCPPType == TEXT("FRotator"))
@@ -1748,6 +1905,14 @@ FRigVMOperand URigVMCompiler::FindOrAddRegister(const FRigVMVarExprAST* InVarExp
 	URigVMPin::FPinOverride PinOverride(InVarExpr->GetProxy(), PinOverrides);
 
 	URigVMPin* Pin = InVarExpr->GetPin();
+
+#if !UE_RIGVM_UCLASS_BASED_STORAGE_DISABLED
+	if(Pin->IsExecuteContext() && bIsDebugValue)
+	{
+		return FRigVMOperand();
+	}
+#endif
+	
 	FString CPPType = Pin->GetCPPType();
 	FString BaseCPPType = Pin->IsArray() ? Pin->GetArrayElementCppType() : CPPType;
 	FString Hash = GetPinHash(Pin, InVarExpr, bIsDebugValue);
@@ -1840,7 +2005,7 @@ FRigVMOperand URigVMCompiler::FindOrAddRegister(const FRigVMVarExprAST* InVarExp
 
 #else
 
-						const int32 PropertyPathIndex = WorkData.FindOrAddPropertyPath(Operand, SegmentPath);
+						const int32 PropertyPathIndex = WorkData.FindOrAddPropertyPath(Operand, Pin->GetCPPType(), SegmentPath);
 						Operand.RegisterOffset = (uint16)PropertyPathIndex;
 						
 #endif
@@ -1861,15 +2026,20 @@ FRigVMOperand URigVMCompiler::FindOrAddRegister(const FRigVMVarExprAST* InVarExp
 		{
 			if (Pin->GetDirection() == ERigVMPinDirection::Hidden)
 			{
+#if UE_RIGVM_UCLASS_BASED_STORAGE_DISABLED
 				ensure(WorkData.DefaultStructs.Num() > 0);
 				TSharedPtr<FStructOnScope> DefaultStruct = WorkData.DefaultStructs.Last();
 				FRigVMStruct* VMStruct = (FRigVMStruct *)DefaultStruct->GetStructMemory();
 
 				int32 DesiredArraySize = VMStruct->GetArraySize(Pin->GetFName(), WorkData.RigVMUserData);
 
+#endif
+				
 				JoinedDefaultValue = Pin->GetDefaultValue(PinOverride);
 				DefaultValues = URigVMPin::SplitDefaultValue(JoinedDefaultValue);
 
+#if UE_RIGVM_UCLASS_BASED_STORAGE_DISABLED
+
 				if (DefaultValues.Num() != DesiredArraySize)
 				{
 					FString FirstDefaultValue;
@@ -1884,6 +2054,8 @@ FRigVMOperand URigVMCompiler::FindOrAddRegister(const FRigVMVarExprAST* InVarExp
 						DefaultValues.Add(FirstDefaultValue);
 					}
 				}
+				
+#endif
 			}
 			else
 			{
@@ -2177,11 +2349,16 @@ FRigVMOperand URigVMCompiler::FindOrAddRegister(const FRigVMVarExprAST* InVarExp
 
 		}
 
+		if(bIsDebugValue)
+		{
+			// debug values are always stored as arrays
+			CPPType = FString::Printf(TEXT("TArray<%s>"), *CPPType);
+		}
 		Operand = WorkData.AddProperty(MemoryType, RegisterName, CPPType, Pin->GetCPPTypeObject(), JoinedDefaultValue);
 		
 #endif
 
-		if (bIsParameter && !bIsLiteral)
+		if (bIsParameter && !bIsLiteral && !bIsDebugValue)
 		{
 			URigVMParameterNode* ParameterNode = Cast<URigVMParameterNode>(Pin->GetNode());
 			check(ParameterNode);
diff --git a/Engine/Source/Developer/RigVMDeveloper/Private/RigVMModel/RigVMPin.cpp b/Engine/Source/Developer/RigVMDeveloper/Private/RigVMModel/RigVMPin.cpp
--- a/Engine/Source/Developer/RigVMDeveloper/Private/RigVMModel/RigVMPin.cpp
+++ b/Engine/Source/Developer/RigVMDeveloper/Private/RigVMModel/RigVMPin.cpp
@@ -466,6 +466,15 @@ int32 URigVMPin::GetArraySize() const
 
 FString URigVMPin::GetCPPType() const
 {
+	if(UScriptStruct* ScriptStruct = GetScriptStruct())
+	{
+		const FString StructCPPType = ScriptStruct->GetStructCPPName();
+		if(IsArray())
+		{
+			return FString::Printf(TEXT("TArray<%s>"), *StructCPPType);
+		}
+		return StructCPPType;
+	}
 	return CPPType;
 }
 
@@ -475,12 +484,15 @@ FString URigVMPin::GetArrayElementCppType() const
 	{
 		return FString();
 	}
-	return CPPType.Mid(7, CPPType.Len() - 8);
+
+	const FString ResolvedType = GetCPPType();
+	return ResolvedType.Mid(7, ResolvedType.Len() - 8);
 }
 
 bool URigVMPin::IsStringType() const
 {
-	return CPPType.Equals(TEXT("FString")) || CPPType.Equals(TEXT("FName"));
+	const FString ResolvedType = GetCPPType();
+	return ResolvedType.Equals(TEXT("FString")) || ResolvedType.Equals(TEXT("FName"));
 }
 
 bool URigVMPin::IsExecuteContext() const
@@ -798,33 +810,11 @@ bool URigVMPin::CanBeBoundToVariable(const FRigVMExternalVariable& InExternalVar
 #else
 	if(!InSegmentPath.IsEmpty())
 	{
-		const FProperty* Property = InExternalVariable.Property;
-		if(Property == nullptr)
-		{
-			return false;
-		}
+		check(InExternalVariable.Property);
 
-		TArray<FString> Segments;
-		SplitPinPath(InSegmentPath, Segments);
-
-		while(!Segments.IsEmpty())
-		{
-			if(const FStructProperty* StructProperty = CastField<FStructProperty>(Property))
-			{
-				const FString FirstSegment = Segments[0];
-				Segments.RemoveAt(0);
-
-				Property = StructProperty->Struct->FindPropertyByName(*FirstSegment);
-				if(Property == nullptr)
-				{
-					return false;
-				}
-			}
-			else
-			{
-				return false;
-			}
-		}
+		const FProperty* Property = InExternalVariable.Property;
+		const FRigVMPropertyPath PropertyPath(Property, InSegmentPath);
+		Property = PropertyPath.GetTargetProperty();
 
 		FRigVMExternalVariable::GetTypeFromProperty(Property, ExternalCPPType, ExternalCPPTypeObject);
 	}
diff --git a/Engine/Source/Developer/RigVMDeveloper/Public/RigVMCompiler/RigVMCompiler.h b/Engine/Source/Developer/RigVMDeveloper/Public/RigVMCompiler/RigVMCompiler.h
--- a/Engine/Source/Developer/RigVMDeveloper/Public/RigVMCompiler/RigVMCompiler.h
+++ b/Engine/Source/Developer/RigVMDeveloper/Public/RigVMCompiler/RigVMCompiler.h
@@ -93,6 +93,8 @@ public:
 	FRigVMOperand ComparisonOperand;
 
 #if !UE_RIGVM_UCLASS_BASED_STORAGE_DISABLED
+
+	TArray<URigVMPin*> WatchedPins;
 	
 	TMap<ERigVMMemoryType, TArray<FRigVMPropertyPathDescription>> PropertyPathDescriptions;
 	TMap<ERigVMMemoryType, TArray<FRigVMPropertyDescription>> PropertyDescriptions;
@@ -100,8 +102,7 @@ public:
 	FRigVMOperand AddProperty(ERigVMMemoryType InMemoryType, const FName& InName, const FString& InCPPType, UObject* InCPPTypeObject, const FString& InDefaultValue = FString());
 	FRigVMOperand FindProperty(ERigVMMemoryType InMemoryType, const FName& InName);
 	FRigVMPropertyDescription GetProperty(const FRigVMOperand& InOperand);
-	int32 FindPropertyPath(ERigVMMemoryType InMemoryType, const FString& InSegmentPath) const;
-	int32 FindOrAddPropertyPath(const FRigVMOperand& InOperand, const FString& InSegmentPath);
+	int32 FindOrAddPropertyPath(const FRigVMOperand& InOperand, const FString& InRootCPPType, const FString& InSegmentPath);
 	
 #endif
 };
@@ -128,9 +129,18 @@ public:
 
 	static UScriptStruct* GetScriptStructForCPPType(const FString& InCPPType);
 	static FString GetPinHash(const URigVMPin* InPin, const FRigVMVarExprAST* InVarExpr, bool bIsDebugValue = false);
+
+#if UE_RIGVM_UCLASS_BASED_STORAGE_DISABLED
+	
 	void CreateDebugRegister(URigVMPin* InPin, URigVM* OutVM, TMap<FString, FRigVMOperand>* OutOperands, TSharedPtr<FRigVMParserAST> InRuntimeAST);
 	void RemoveDebugRegister(URigVMPin* InPin, URigVM* OutVM, TMap<FString, FRigVMOperand>* OutOperands, TSharedPtr<FRigVMParserAST> InRuntimeAST);
 
+#else
+
+	void MarkDebugWatch(bool bRequired, URigVMPin* InPin, URigVM* OutVM, TMap<FString, FRigVMOperand>* OutOperands, TSharedPtr<FRigVMParserAST> InRuntimeAST);
+
+#endif
+
 private:
 
 	TArray<URigVMPin*> GetLinkedPins(URigVMPin* InPin, bool bInputs = true, bool bOutputs = true, bool bRecursive = true);
diff --git a/Engine/Source/Runtime/RigVM/Private/RigVMCore/RigVM.cpp b/Engine/Source/Runtime/RigVM/Private/RigVMCore/RigVM.cpp
--- a/Engine/Source/Runtime/RigVM/Private/RigVMCore/RigVM.cpp
+++ b/Engine/Source/Runtime/RigVM/Private/RigVMCore/RigVM.cpp
@@ -156,6 +156,8 @@ void URigVM::Save(FArchive& Ar)
 	{
 		LiteralMemoryStorageObject->Serialize(Ar);
 	}
+
+	Ar << ExternalPropertyPathDescriptions;
 	
 #endif
 	
@@ -205,12 +207,26 @@ void URigVM::Load(FArchive& Ar)
 	if(WorkMemoryStorageObject == nullptr)
 	{
 		UClass* StorageClass = URigVMMemoryStorageGeneratorClass::GetStorageClass(GetOutermost(), ERigVMMemoryType::Work);
-		WorkMemoryStorageObject = NewObject<URigVMMemoryStorage>(this, StorageClass, NAME_None, RF_Public);
+		if(StorageClass)
+		{
+			WorkMemoryStorageObject = NewObject<URigVMMemoryStorage>(this, StorageClass, NAME_None, RF_Public);
+		}
+		else
+		{
+			WorkMemoryStorageObject = NewObject<URigVMMemoryStorage>(this, NAME_None, RF_Public);
+		}
 	}
 	if(LiteralMemoryStorageObject == nullptr)
 	{
 		UClass* StorageClass = URigVMMemoryStorageGeneratorClass::GetStorageClass(GetOutermost(), ERigVMMemoryType::Literal);
-		LiteralMemoryStorageObject = NewObject<URigVMMemoryStorage>(this, StorageClass, NAME_None, RF_Public);
+		if(StorageClass)
+		{
+			LiteralMemoryStorageObject = NewObject<URigVMMemoryStorage>(this, StorageClass, NAME_None, RF_Public);
+		}
+		else
+		{
+			LiteralMemoryStorageObject = NewObject<URigVMMemoryStorage>(this, NAME_None, RF_Public);
+		}
 	}
 	if(DebugMemoryStorageObject == nullptr)
 	{
@@ -230,6 +246,8 @@ void URigVM::Load(FArchive& Ar)
 	{
 		LiteralMemoryStorageObject->Serialize(Ar);
 	}
+	
+	Ar << ExternalPropertyPathDescriptions;
 
 	Ar << FunctionNamesStorage;
 	Ar << ByteCodeStorage;
@@ -472,6 +490,8 @@ void URigVM::Reset()
 #endif
 	FunctionNamesStorage.Reset();
 	FunctionsStorage.Reset();
+	ExternalPropertyPathDescriptions.Reset();
+	ExternalPropertyPaths.Reset();
 	ByteCodeStorage.Reset();
 	Instructions.Reset();
 	Parameters.Reset();
@@ -501,6 +521,8 @@ void URigVM::Empty()
 #endif
 	FunctionNamesStorage.Empty();
 	FunctionsStorage.Empty();
+	ExternalPropertyPathDescriptions.Empty();
+	ExternalPropertyPaths.Empty();
 	ByteCodeStorage.Empty();
 	Instructions.Empty();
 	Parameters.Empty();
@@ -603,6 +625,9 @@ void URigVM::CopyFrom(URigVM* InVM, bool bDeferCopy, bool bReferenceLiteralMemor
 	CopyMemoryStorage(WorkMemoryStorageObject, InVM->WorkMemoryStorageObject, this);
 	CopyMemoryStorage(LiteralMemoryStorageObject, InVM->LiteralMemoryStorageObject, this);
 	CopyMemoryStorage(DebugMemoryStorageObject, InVM->DebugMemoryStorageObject, this);
+
+	ExternalPropertyPathDescriptions = InVM->ExternalPropertyPathDescriptions;
+	ExternalPropertyPaths.Reset();
 	
 #endif
 
@@ -684,7 +709,6 @@ URigVMMemoryStorage* URigVM::GetMemoryByType(ERigVMMemoryType InMemoryType) cons
 			return GetLiteralMemory();
 		}
 		case ERigVMMemoryType::Work:
-		case ERigVMMemoryType::External:
 		{
 			return GetWorkMemory();
 		}
@@ -819,6 +843,10 @@ void URigVM::InvalidateCachedMemory()
 	CachedMemory.Reset();
 	FirstHandleForInstruction.Reset();
 	CachedMemoryHandles.Reset();
+
+#if !UE_RIGVM_UCLASS_BASED_STORAGE_DISABLED
+	ExternalPropertyPaths.Reset();
+#endif
 }
 
 void URigVM::CopyDeferredVMIfRequired()
@@ -880,6 +908,28 @@ void URigVM::CacheMemoryHandlesIfRequired(TArrayView<URigVMMemoryStorage*> InMem
 		CachedMemory.Add(InMemory[Index]);
 	}
 
+#if !UE_RIGVM_UCLASS_BASED_STORAGE_DISABLED
+
+	ExternalPropertyPaths.Reset();
+
+	ExternalPropertyPaths.SetNumZeroed(ExternalPropertyPathDescriptions.Num());
+	for(int32 PropertyPathIndex = 0; PropertyPathIndex < ExternalPropertyPaths.Num(); PropertyPathIndex++)
+	{
+		ExternalPropertyPaths[PropertyPathIndex] = FRigVMPropertyPath();
+
+		const int32 PropertyIndex = ExternalPropertyPathDescriptions[PropertyPathIndex].PropertyIndex;
+		if(ExternalVariables.IsValidIndex(PropertyIndex))
+		{
+			check(ExternalVariables[PropertyIndex].Property);
+			
+			ExternalPropertyPaths[PropertyPathIndex] = FRigVMPropertyPath(
+				ExternalVariables[PropertyIndex].Property,
+				ExternalPropertyPathDescriptions[PropertyPathIndex].SegmentPath);
+		}
+	}
+	
+#endif		
+
 	FRigVMByteCode& ByteCode = GetByteCode();
 
 	uint16 InstructionIndex = 0;
@@ -2637,15 +2687,15 @@ void URigVM::CacheSingleMemoryHandle(const FRigVMOperand& InArg, bool bForExecut
 
 	URigVMMemoryStorage* Memory = GetMemoryByType(InArg.GetMemoryType());
 
-	const FRigVMPropertyPath* PropertyPath = nullptr;
-	if(InArg.GetRegisterOffset())
-	{
-		check(Memory->GetPropertyPaths().IsValidIndex(InArg.GetRegisterOffset()));
-		PropertyPath = &Memory->GetPropertyPaths()[InArg.GetRegisterOffset()];
-	}
-
 	if (InArg.GetMemoryType() == ERigVMMemoryType::External)
 	{
+		const FRigVMPropertyPath* PropertyPath = nullptr;
+		if(InArg.GetRegisterOffset() != INDEX_NONE)
+		{
+			check(ExternalPropertyPaths.IsValidIndex(InArg.GetRegisterOffset()));
+			PropertyPath = &ExternalPropertyPaths[InArg.GetRegisterOffset()];
+		}
+
 		check(ExternalVariables.IsValidIndex(InArg.GetRegisterIndex()));
 
 		FRigVMExternalVariable& ExternalVariable = ExternalVariables[InArg.GetRegisterIndex()];
@@ -2656,6 +2706,12 @@ void URigVM::CacheSingleMemoryHandle(const FRigVMOperand& InArg, bool bForExecut
 		return;
 	}
 
+	const FRigVMPropertyPath* PropertyPath = nullptr;
+	if(InArg.GetRegisterOffset() != INDEX_NONE)
+	{
+		check(Memory->GetPropertyPaths().IsValidIndex(InArg.GetRegisterOffset()));
+		PropertyPath = &Memory->GetPropertyPaths()[InArg.GetRegisterOffset()];
+	}
 	uint8* Data = Memory->GetData<uint8>(InArg.GetRegisterIndex());
 	const FProperty* Property = Memory->GetProperties()[InArg.GetRegisterIndex()];
 	const FRigVMMemoryHandle Handle(Data, Property, PropertyPath);
@@ -2797,6 +2853,16 @@ void URigVM::CopyOperandForDebuggingImpl(const FRigVMOperand& InArg, const FRigV
 
 	// add an element to the end for debug watching
 	FScriptArrayHelper ArrayHelper(TargetArrayProperty, TargetPtr);
+
+	if (Context.GetSlice().GetIndex() == 0)
+	{
+		ArrayHelper.Resize(0);
+	}
+	else if(Context.GetSlice().GetIndex() == ArrayHelper.Num() - 1)
+	{
+		return;
+	}
+
 	const int32 AddedIndex = ArrayHelper.AddValue();
 	TargetPtr = ArrayHelper.GetRawPtr(AddedIndex);
 	TargetProperty = TargetArrayProperty->Inner;
diff --git a/Engine/Source/Runtime/RigVM/Private/RigVMCore/RigVMMemoryStorage.cpp b/Engine/Source/Runtime/RigVM/Private/RigVMCore/RigVMMemoryStorage.cpp
--- a/Engine/Source/Runtime/RigVM/Private/RigVMCore/RigVMMemoryStorage.cpp
+++ b/Engine/Source/Runtime/RigVM/Private/RigVMCore/RigVMMemoryStorage.cpp
@@ -5,6 +5,14 @@
 
 ////////////////////////////////////////////////////////////////////////////////////////////////////////////////////////
 
+#if !UE_RIGVM_UCLASS_BASED_STORAGE_DISABLED
+
+const FRigVMPropertyPath FRigVMMemoryHandle::EmptyPropertyPath;
+
+#endif
+
+////////////////////////////////////////////////////////////////////////////////////////////////////////////////////////
+
 const FString FRigVMPropertyDescription::ArrayPrefix = TEXT("TArray<");
 const FString FRigVMPropertyDescription::MapPrefix = TEXT("TMap<");
 const FString FRigVMPropertyDescription::SetPrefix = TEXT("TSet<");
@@ -16,6 +24,7 @@ FRigVMPropertyDescription::FRigVMPropertyDescription(const FProperty* InProperty
 	: Name(InName)
 	, Property(InProperty)
 	, CPPType()
+	, CPPTypeObject(nullptr)
 	, Containers()
 	, DefaultValue(InDefaultValue)
 {
@@ -155,9 +164,6 @@ FString FRigVMPropertyDescription::GetBaseCPPType() const
 
 ////////////////////////////////////////////////////////////////////////////////////////////////////////////////////////
 
-FCriticalSection RigVMMemoryStorageClassCriticalSection;
-TMap<FString, URigVMMemoryStorageGeneratorClass::FClassInfo> URigVMMemoryStorageGeneratorClass::PackageToInfo;
-
 void URigVMMemoryStorageGeneratorClass::PurgeClass(bool bRecompilingOnLoad)
 {
 	Super::PurgeClass(bRecompilingOnLoad);
@@ -183,20 +189,7 @@ void URigVMMemoryStorageGeneratorClass::Link(FArchive& Ar, bool bRelinkExistingP
 		Property = CastField<FProperty>(Property->Next);
 	}
 
-	// Update the property paths based on the descriptions
-	PropertyPaths.SetNumZeroed(PropertyPathDescriptions.Num());
-	for(int32 PropertyPathIndex = 0; PropertyPathIndex < PropertyPaths.Num(); PropertyPathIndex++)
-	{
-		PropertyPaths[PropertyPathIndex] = FRigVMPropertyPath();
-
-		const int32 PropertyIndex = PropertyPathDescriptions[PropertyPathIndex].PropertyIndex;
-		if(LinkedProperties.IsValidIndex(PropertyIndex))
-		{
-			PropertyPaths[PropertyPathIndex] = FRigVMPropertyPath(
-				LinkedProperties[PropertyIndex],
-				PropertyPathDescriptions[PropertyPathIndex].SegmentPath);
-		}
-	}
+	RefreshPropertyPaths();
 }
 
 void URigVMMemoryStorageGeneratorClass::Serialize(FArchive& Ar)
@@ -208,40 +201,20 @@ void URigVMMemoryStorageGeneratorClass::Serialize(FArchive& Ar)
 		Ar << PropertyPathDescriptions;
 		Ar << MemoryType;
 	}
+}
 
-	if(Ar.IsLoading())
-	{
-		UPackage* Package = GetOutermost();
-		
-		FScopeLock Lock(&RigVMMemoryStorageClassCriticalSection);
-
-		FClassInfo& ClassInfo = PackageToInfo.FindOrAdd(Package->GetPathName());
-		URigVMMemoryStorageGeneratorClass** ClassPtr = ClassInfo.GetClassPtr(MemoryType);
-		if(ClassPtr)
-		{
-			if(*ClassPtr != this)
-			{
-				if(*ClassPtr)
-				{
-					(*ClassPtr)->RemoveFromRoot();
-					(*ClassPtr)->MarkPendingKill();
-				}
-				*ClassPtr = this;
-			}
-		}
-	}
+FString URigVMMemoryStorageGeneratorClass::GetClassName(ERigVMMemoryType InMemoryType)
+{
+	return TEXT("RigVMMemory_") + StaticEnum<ERigVMMemoryType>()->GetDisplayNameTextByValue((int64)InMemoryType).ToString();
 }
 
 URigVMMemoryStorageGeneratorClass* URigVMMemoryStorageGeneratorClass::GetStorageClass(UObject* InOuter, ERigVMMemoryType InMemoryType)
 {
 	check(InOuter);
-	UPackage* Package = InOuter->GetOutermost(); 
-
-	FScopeLock Lock(&RigVMMemoryStorageClassCriticalSection);
+	UPackage* Package = InOuter->GetOutermost();
 
-	FClassInfo& ClassInfo = PackageToInfo.FindOrAdd(Package->GetPathName());
-	URigVMMemoryStorageGeneratorClass** ClassPtr = ClassInfo.GetClassPtr(InMemoryType);
-	return *ClassPtr;
+	const FString ClassName = GetClassName(InMemoryType);
+	return FindObject<URigVMMemoryStorageGeneratorClass>(Package, *ClassName);
 }
 
 URigVMMemoryStorageGeneratorClass* URigVMMemoryStorageGeneratorClass::CreateStorageClass(
@@ -254,29 +227,29 @@ URigVMMemoryStorageGeneratorClass* URigVMMemoryStorageGeneratorClass::CreateStor
 	UPackage* Package = InOuter->GetOutermost(); 
 	UClass *SuperClass = URigVMMemoryStorage::StaticClass();
 
-	FScopeLock Lock(&RigVMMemoryStorageClassCriticalSection);
-
-	FClassInfo& ClassInfo = PackageToInfo.FindOrAdd(Package->GetPathName());
-	URigVMMemoryStorageGeneratorClass** ClassPtr = ClassInfo.GetClassPtr(InMemoryType);
-	check(ClassPtr);
+	const FString ClassName = GetClassName(InMemoryType);
 	
-	if(*ClassPtr == nullptr)
+	URigVMMemoryStorageGeneratorClass* OldClass = FindObject<URigVMMemoryStorageGeneratorClass>(Package, *ClassName);
+	if(OldClass)
 	{
-		(*ClassPtr) = NewObject<URigVMMemoryStorageGeneratorClass>(
-			InOuter->GetOutermost(),
-			NAME_None,
-			RF_Standalone | RF_Public
-			);
-
-		(*ClassPtr)->AddToRoot();
+		OldClass->RemoveFromRoot();
+		OldClass->Rename(nullptr, GetTransientPackage(), REN_ForceNoResetLoaders | REN_DoNotDirty | REN_DontCreateRedirectors | REN_NonTransactional);
+		OldClass->MarkPendingKill();
 	}
 
-	URigVMMemoryStorageGeneratorClass* Class = *ClassPtr;
+	URigVMMemoryStorageGeneratorClass* Class = NewObject<URigVMMemoryStorageGeneratorClass>(
+		Package,
+		*ClassName,
+		RF_Standalone | RF_Public
+	);
+
+	Class->AddToRoot();
 
 	Class->PurgeClass(false);
-	Class->PropertyLink = SuperClass->PropertyLink;
 	Class->SetSuperStruct(SuperClass);
+	Class->PropertyLink = SuperClass->PropertyLink;
 	Class->ClassWithin = UObject::StaticClass();
+	Class->ClassFlags |= CLASS_NotPlaceable;
 	Class->MemoryType = InMemoryType;
 
 	// Generate properties
@@ -284,7 +257,7 @@ URigVMMemoryStorageGeneratorClass* URigVMMemoryStorageGeneratorClass::CreateStor
 
 	for(const FRigVMPropertyDescription& PropertyDescription : InProperties)
 	{
-		FProperty* CachedProperty = AddProperty(Class, PropertyDescription, false, false, LinkToProperty);
+		FProperty* CachedProperty = AddProperty(Class, PropertyDescription, LinkToProperty);
 		check(CachedProperty);
 	}
 
@@ -313,25 +286,35 @@ URigVMMemoryStorageGeneratorClass* URigVMMemoryStorageGeneratorClass::CreateStor
 
 		Property->ImportText(*DefaultValue, ValuePtr, EPropertyPortFlags::PPF_None, nullptr);
 	}
+
+	/*
+	if(OldClass)
+	{
+		UEngine::FCopyPropertiesForUnrelatedObjectsParams Params;
+		Params.bAggressiveDefaultSubobjectReplacement = false;
+		Params.bDoDelta = false;
+		Params.bCopyDeprecatedProperties = true;
+		Params.bSkipCompilerGeneratedDefaults = true;
+		Params.bNotifyObjectReplacement = true;
+
+		Params.bClearReferences = true;
+		UEngine::CopyPropertiesForUnrelatedObjects(OldClass->GetDefaultObject(), CDO, Params);
+
+		// we might need to run this again on the old instances,
+		// but the VM compiler should take care of it anyway
+	}
+	*/
 	
 	return Class;
 }
 
-FProperty* URigVMMemoryStorageGeneratorClass::AddProperty(URigVMMemoryStorageGeneratorClass* InClass, const FRigVMPropertyDescription& InProperty, bool bPurge, bool bLink, FField** LinkToProperty)
+FProperty* URigVMMemoryStorageGeneratorClass::AddProperty(URigVMMemoryStorageGeneratorClass* InClass, const FRigVMPropertyDescription& InProperty, FField** LinkToProperty)
 {
 	UClass *SuperClass = URigVMMemoryStorage::StaticClass();
 	
 	check(InClass);
 	check(InClass->GetSuperClass() == SuperClass);
 
-	if(bPurge)
-	{
-		InClass->PurgeClass(false);
-		InClass->PropertyLink = SuperClass->PropertyLink;
-		InClass->SetSuperStruct(SuperClass);
-		InClass->ClassWithin = UObject::StaticClass();
-	}
-	
 	if(LinkToProperty == nullptr)
 	{
 		LinkToProperty = &InClass->ChildProperties;
@@ -465,26 +448,32 @@ FProperty* URigVMMemoryStorageGeneratorClass::AddProperty(URigVMMemoryStorageGen
 		(*LinkToProperty) = Result;
 	}
 
-	if(bLink)
+	return Result;
+}
+
+void URigVMMemoryStorageGeneratorClass::RefreshPropertyPaths()
+{
+	// Update the property paths based on the descriptions
+	PropertyPaths.SetNumZeroed(PropertyPathDescriptions.Num());
+	for(int32 PropertyPathIndex = 0; PropertyPathIndex < PropertyPaths.Num(); PropertyPathIndex++)
 	{
-		// Update the class
-		InClass->Bind();
-		InClass->StaticLink(true);
+		PropertyPaths[PropertyPathIndex] = FRigVMPropertyPath();
 
-		URigVMMemoryStorage* CDO = Cast<URigVMMemoryStorage>(InClass->GetDefaultObject(true));
-		const FString& DefaultValue = InProperty.DefaultValue;
-		if(!DefaultValue.IsEmpty())
+		const int32 PropertyIndex = PropertyPathDescriptions[PropertyPathIndex].PropertyIndex;
+		if(LinkedProperties.IsValidIndex(PropertyIndex))
 		{
-			uint8* ValuePtr = Result->ContainerPtrToValuePtr<uint8>(CDO);
-			Result->ImportText(*DefaultValue, ValuePtr, EPropertyPortFlags::PPF_None, nullptr);
-		}		
-	};
-	
-	return Result;
+			PropertyPaths[PropertyPathIndex] = FRigVMPropertyPath(
+				LinkedProperties[PropertyIndex],
+				PropertyPathDescriptions[PropertyPathIndex].SegmentPath);
+		}
+	}
 }
 
 ////////////////////////////////////////////////////////////////////////////////////////////////////////////////////////
 
+const TArray<const FProperty*> URigVMMemoryStorage::EmptyProperties;
+const TArray<FRigVMPropertyPath> URigVMMemoryStorage::EmptyPropertyPaths;
+
 FString URigVMMemoryStorage::GetDataAsString(int32 InPropertyIndex)
 {
 	check(IsValidIndex(InPropertyIndex));
@@ -582,34 +571,8 @@ bool URigVMMemoryStorage::CopyProperty(
 			return;
 		}
 
-		MemoryPtr = PropertyPath.GetData<uint8>(MemoryPtr);
-
-		const FRigVMPropertyPathSegment& LastSegment = PropertyPath[PropertyPath.Num() - 1];
-		switch(LastSegment.Type)
-		{
-		case ERigVMPropertyPathSegmentType::ArrayElement:
-			{
-				const FArrayProperty* ArrayProperty = CastFieldChecked<FArrayProperty>(LastSegment.Property);
-				Property = ArrayProperty->Inner;
-				break;
-			}
-		case ERigVMPropertyPathSegmentType::MapValue:
-			{
-				const FMapProperty* MapProperty = CastFieldChecked<FMapProperty>(LastSegment.Property);
-				Property = MapProperty->ValueProp;
-				break;
-			}
-		case ERigVMPropertyPathSegmentType::StructMember:
-			{
-				Property = LastSegment.Property;
-				break;
-			}
-		default:
-			{
-				checkNoEntry();
-				break;
-			}
-		}
+		MemoryPtr = PropertyPath.GetData<uint8>(MemoryPtr, Property);
+		Property = PropertyPath.GetTargetProperty();
 	};
 
 	uint8* SourcePtr = (uint8*)InSourcePtr;
@@ -646,23 +609,31 @@ bool URigVMMemoryStorage::CopyProperty(
 {
 	return CopyProperty(
 		TargetHandle.GetProperty(),
-		TargetHandle.GetData(true),
+		TargetHandle.GetData(false),
+		TargetHandle.GetPropertyPathRef(),
 		SourceHandle.GetProperty(),
-		SourceHandle.GetData(true));
+		SourceHandle.GetData(false),
+		SourceHandle.GetPropertyPathRef());
 }
 
 #endif
 
 const TArray<const FProperty*>& URigVMMemoryStorage::GetProperties() const
 {
-	const URigVMMemoryStorageGeneratorClass* Class = CastChecked<URigVMMemoryStorageGeneratorClass>(GetClass());
-	return Class->GetProperties();
+	if(const URigVMMemoryStorageGeneratorClass* Class = Cast<URigVMMemoryStorageGeneratorClass>(GetClass()))
+	{
+		return Class->GetProperties();
+	}
+	return EmptyProperties;
 }
 
 const TArray<FRigVMPropertyPath>& URigVMMemoryStorage::GetPropertyPaths() const
 {
-	const URigVMMemoryStorageGeneratorClass* Class = CastChecked<URigVMMemoryStorageGeneratorClass>(GetClass());
-	return Class->GetPropertyPaths();
+	if(const URigVMMemoryStorageGeneratorClass* Class = Cast<URigVMMemoryStorageGeneratorClass>(GetClass()))
+	{
+		return Class->GetPropertyPaths();
+	}
+	return EmptyPropertyPaths;
 }
 
 int32 URigVMMemoryStorage::GetPropertyIndex(const FProperty* InProperty) const
diff --git a/Engine/Source/Runtime/RigVM/Private/RigVMCore/RigVMPropertyPath.cpp b/Engine/Source/Runtime/RigVM/Private/RigVMCore/RigVMPropertyPath.cpp
--- a/Engine/Source/Runtime/RigVM/Private/RigVMCore/RigVMPropertyPath.cpp
+++ b/Engine/Source/Runtime/RigVM/Private/RigVMCore/RigVMPropertyPath.cpp
@@ -131,7 +131,7 @@ bool FRigVMPropertyPath::IsDirect() const
 			case ERigVMPropertyPathSegmentType::ArrayElement:
 			case ERigVMPropertyPathSegmentType::MapValue:
 			{
-				return true;
+				return false;
 			}
 			default:
 			{
@@ -139,10 +139,47 @@ bool FRigVMPropertyPath::IsDirect() const
 			}
 		}
 	}
-	return false;
+	return true;
 }
 
-uint8* FRigVMPropertyPath::GetData_Internal(uint8* InPtr) const
+const FProperty* FRigVMPropertyPath::GetTargetProperty() const
+{
+	if(Segments.IsEmpty())
+	{
+		return nullptr;
+	}
+	
+	const FRigVMPropertyPathSegment& LastSegment = Segments[Num() - 1];
+	switch(LastSegment.Type)
+	{
+		case ERigVMPropertyPathSegmentType::ArrayElement:
+		{
+			const FArrayProperty* ArrayProperty = CastFieldChecked<FArrayProperty>(LastSegment.Property);
+			return ArrayProperty->Inner;
+			break;
+		}
+		case ERigVMPropertyPathSegmentType::MapValue:
+		{
+			const FMapProperty* MapProperty = CastFieldChecked<FMapProperty>(LastSegment.Property);
+			return MapProperty->ValueProp;
+			break;
+		}
+		case ERigVMPropertyPathSegmentType::StructMember:
+		{
+			return LastSegment.Property;
+			break;
+		}
+		default:
+		{
+			checkNoEntry();
+			break;
+		}
+	}
+
+	return nullptr;
+}
+
+uint8* FRigVMPropertyPath::GetData_Internal(uint8* InPtr, const FProperty* InProperty) const
 {
 	for(const FRigVMPropertyPathSegment& Segment : Segments)
 	{
@@ -150,19 +187,35 @@ uint8* FRigVMPropertyPath::GetData_Internal(uint8* InPtr) const
 		{
 			case ERigVMPropertyPathSegmentType::StructMember:
 			{
-				// index represents the offset needed
-				// from container to struct member
+#if UE_BUILD_DEBUG
+				// make sure to only allow jumps within the same data structure.
+				// we check this to avoid heap corruption.
+				const FStructProperty* StructProperty = CastFieldChecked<FStructProperty>(InProperty);
+				check(StructProperty->Struct == Segment.Property->GetOwnerStruct());
+				InProperty = Segment.Property;
+#endif
+
 				InPtr += Segment.Property->GetOffset_ForInternal();
 				break;
 			}
 			case ERigVMPropertyPathSegmentType::ArrayElement:
 			{
+#if UE_BUILD_DEBUG
+				// we check this to avoid heap corruption.
+				check(InProperty->SameType(Segment.Property));
+				InProperty = CastFieldChecked<FArrayProperty>(Segment.Property)->Inner;
+#endif
 				FScriptArrayHelper ArrayHelper(CastFieldChecked<FArrayProperty>(Segment.Property), InPtr);
 				InPtr = ArrayHelper.GetRawPtr(Segment.Index);
 				break;
 			}
 			case ERigVMPropertyPathSegmentType::MapValue:
 			{
+#if UE_BUILD_DEBUG
+				// we check this to avoid heap corruption.
+				check(InProperty->SameType(Segment.Property));
+				InProperty = CastFieldChecked<FMapProperty>(Segment.Property)->ValueProp;
+#endif
 				FScriptMapHelper MapHelper(CastFieldChecked<FMapProperty>(Segment.Property), InPtr);
 				InPtr = MapHelper.FindValueFromHash(&Segment.Name);;
 				break;
diff --git a/Engine/Source/Runtime/RigVM/Public/RigVMCore/RigVM.h b/Engine/Source/Runtime/RigVM/Public/RigVMCore/RigVM.h
--- a/Engine/Source/Runtime/RigVM/Public/RigVMCore/RigVM.h
+++ b/Engine/Source/Runtime/RigVM/Public/RigVMCore/RigVM.h
@@ -250,6 +250,9 @@ public:
 	UPROPERTY()
 	URigVMMemoryStorage* DebugMemoryStorageObject;
 
+	TArray<FRigVMPropertyPathDescription> ExternalPropertyPathDescriptions;
+	TArray<FRigVMPropertyPath> ExternalPropertyPaths;
+
 	// The byte code of the VM
 	UPROPERTY()
 	FRigVMByteCode ByteCodeStorage;
@@ -344,13 +347,31 @@ public:
 		DefaultValues.Add(DefaultValue);
 		return AddParameter(InParameterType, InName, InCPPType, DefaultValues);
 	}
-	
+
 	// Retrieve the array size of the parameter
 	int32 GetParameterArraySize(const FRigVMParameter& InParameter) const
 	{
 		return (int32)GetWorkMemory()[InParameter.GetRegisterIndex()].GetTotalElementCount();
 	}
 
+#else
+
+	// Retrieve the array size of the parameter
+	FORCEINLINE int32 GetParameterArraySize(const FRigVMParameter& InParameter) const
+	{
+		const int32 PropertyIndex = InParameter.GetRegisterIndex();
+		const FProperty* Property = GetWorkMemory()->GetProperties()[PropertyIndex];
+		const FArrayProperty* ArrayProperty = CastField<FArrayProperty>(Property);
+		if(ArrayProperty)
+		{
+			FScriptArrayHelper ArrayHelper(ArrayProperty, GetWorkMemory()->GetData<uint8>(PropertyIndex));
+			return ArrayHelper.Num();
+		}
+		return 1;
+	}
+
+#endif
+
 	// Retrieve the array size of the parameter
 	int32 GetParameterArraySize(int32 InParameterIndex) const
 	{
@@ -363,8 +384,6 @@ public:
 		int32 ParameterIndex = ParametersNameMap.FindChecked(InParameterName);
 		return GetParameterArraySize(ParameterIndex);
 	}
-
-#endif
 	
 	// Retrieve the value of a parameter
 	template<class T>
diff --git a/Engine/Source/Runtime/RigVM/Public/RigVMCore/RigVMExternalVariable.h b/Engine/Source/Runtime/RigVM/Public/RigVMCore/RigVMExternalVariable.h
--- a/Engine/Source/Runtime/RigVM/Public/RigVMCore/RigVMExternalVariable.h
+++ b/Engine/Source/Runtime/RigVM/Public/RigVMCore/RigVMExternalVariable.h
@@ -32,26 +32,32 @@ struct RIGVM_API FRigVMExternalVariable
 		if (CastField<FBoolProperty>(InProperty))
 		{
 			OutTypeName = TEXT("bool");
+			OutTypeObject = nullptr;
 		}
 		else if (CastField<FIntProperty>(InProperty))
 		{
 			OutTypeName = TEXT("int32");
+			OutTypeObject = nullptr;
 		}
 		else if (CastField<FFloatProperty>(InProperty))
 		{
 			OutTypeName = TEXT("float");
+			OutTypeObject = nullptr;
 		}
 		else if (CastField<FDoubleProperty>(InProperty))
 		{
 			OutTypeName = TEXT("double");
+			OutTypeObject = nullptr;
 		}
 		else if (CastField<FStrProperty>(InProperty))
 		{
 			OutTypeName = TEXT("FString");
+			OutTypeObject = nullptr;
 		}
 		else if (CastField<FNameProperty>(InProperty))
 		{
 			OutTypeName = TEXT("FName");
+			OutTypeObject = nullptr;
 		}
 		else if (const FEnumProperty* EnumProperty = CastField<FEnumProperty>(InProperty))
 		{
@@ -78,14 +84,14 @@ struct RIGVM_API FRigVMExternalVariable
 		}
 	}
 
-	FORCEINLINE static FRigVMExternalVariable Make(FProperty* InProperty, UObject* InContainer)
+	FORCEINLINE static FRigVMExternalVariable Make(const FProperty* InProperty, void* InContainer, const FName& InOptionalName = NAME_None)
 	{
 		check(InProperty);
 
-		FProperty* Property = InProperty;
+		const FProperty* Property = InProperty;
 
 		FRigVMExternalVariable ExternalVariable;
-		ExternalVariable.Name = InProperty->GetFName();
+		ExternalVariable.Name = InOptionalName.IsNone() ? InProperty->GetFName() : InOptionalName;
 #if !UE_RIGVM_UCLASS_BASED_STORAGE_DISABLED
 		ExternalVariable.Property = Property;
 #endif
@@ -98,7 +104,7 @@ struct RIGVM_API FRigVMExternalVariable
 		}
 
 		FString TypePrefix, TypeSuffix;
-		if (FArrayProperty* ArrayProperty = CastField<FArrayProperty>(Property))
+		if (const FArrayProperty* ArrayProperty = CastField<FArrayProperty>(Property))
 		{
 			ExternalVariable.bIsArray = true;
 			TypePrefix = TEXT("TArray<");
diff --git a/Engine/Source/Runtime/RigVM/Public/RigVMCore/RigVMMemoryStorage.h b/Engine/Source/Runtime/RigVM/Public/RigVMCore/RigVMMemoryStorage.h
--- a/Engine/Source/Runtime/RigVM/Public/RigVMCore/RigVMMemoryStorage.h
+++ b/Engine/Source/Runtime/RigVM/Public/RigVMCore/RigVMMemoryStorage.h
@@ -27,19 +27,11 @@ public:
 		, PropertyPath(nullptr)
 	{}
 
-	FORCEINLINE_DEBUGGABLE FRigVMMemoryHandle(uint8* InPtr, const FProperty* InProperty,  const FRigVMPropertyPath* InPropertyPath = nullptr)
+	FORCEINLINE_DEBUGGABLE FRigVMMemoryHandle(uint8* InPtr, const FProperty* InProperty,  const FRigVMPropertyPath* InPropertyPath)
 		: Ptr(InPtr)
 		, Property(InProperty)
 		, PropertyPath(InPropertyPath)
 	{
-		if(PropertyPath)
-		{
-			if(PropertyPath->IsDirect())
-			{
-				Ptr = PropertyPath->GetData<uint8>(Ptr);
-				PropertyPath = nullptr;
-			}
-		}
 	}
 
 	FORCEINLINE uint8* GetData(bool bFollowPropertyPath = false)
@@ -62,13 +54,22 @@ public:
 		return PropertyPath;
 	}
 
+	FORCEINLINE const FRigVMPropertyPath& GetPropertyPathRef() const
+	{
+		if(PropertyPath)
+		{
+			return *PropertyPath;
+		}
+		return EmptyPropertyPath;
+	}
+
 private:
 
-	FORCEINLINE uint8* GetData_Internal(bool bFollowPropertyPath = false) const
+	FORCEINLINE_DEBUGGABLE uint8* GetData_Internal(bool bFollowPropertyPath = false) const
 	{
 		if(bFollowPropertyPath && PropertyPath != nullptr)
 		{
-			return PropertyPath->GetData<uint8>(Ptr);
+			return PropertyPath->GetData<uint8>(Ptr, Property);
 		}
 		return Ptr;
 	}
@@ -76,6 +77,7 @@ private:
 	uint8* Ptr;
 	const FProperty* Property;
 	const FRigVMPropertyPath* PropertyPath;
+	static const FRigVMPropertyPath EmptyPropertyPath;
 
 	friend class URigVM;
 };
@@ -85,6 +87,7 @@ private:
 //////////////////////////////////////////////////////////////////////////////
 /// Property Management
 //////////////////////////////////////////////////////////////////////////////
+
 struct RIGVM_API FRigVMPropertyDescription
 {
 	FName Name;
@@ -142,6 +145,7 @@ public:
 	void Serialize(FArchive& Ar) override;;
 
 	// URigVMMemoryStorageGeneratorClass specific
+	static FString GetClassName(ERigVMMemoryType InMemoryType);
 	static URigVMMemoryStorageGeneratorClass* GetStorageClass(UObject* InOuter, ERigVMMemoryType InMemoryType);
 	static URigVMMemoryStorageGeneratorClass* CreateStorageClass(
 		UObject* InOuter,
@@ -155,51 +159,17 @@ public:
 
 protected:
 
-	static FProperty* AddProperty(URigVMMemoryStorageGeneratorClass* InClass, const FRigVMPropertyDescription& InProperty, bool bPurge = false, bool bLink = true, FField** LinkToProperty = nullptr);
+	static FProperty* AddProperty(URigVMMemoryStorageGeneratorClass* InClass, const FRigVMPropertyDescription& InProperty, FField** LinkToProperty = nullptr);
 
 private:
 
+	void RefreshPropertyPaths();
+
 	ERigVMMemoryType MemoryType;
 	TArray<const FProperty*> LinkedProperties;
 	TArray<FRigVMPropertyPath> PropertyPaths;
 	TArray<FRigVMPropertyPathDescription> PropertyPathDescriptions;
 
-	struct FClassInfo
-	{
-		URigVMMemoryStorageGeneratorClass* LiteralStorageClass;
-		URigVMMemoryStorageGeneratorClass* WorkStorageClass;
-		URigVMMemoryStorageGeneratorClass* DebugStorageClass;
-		
-		FClassInfo()
-			: LiteralStorageClass(nullptr)
-			, WorkStorageClass(nullptr)
-			, DebugStorageClass(nullptr)
-		{
-		}
-
-		URigVMMemoryStorageGeneratorClass** GetClassPtr(ERigVMMemoryType InMemoryType)
-		{
-			switch(InMemoryType)
-			{
-			case ERigVMMemoryType::Literal:
-				{
-					return &LiteralStorageClass;
-				}
-			case ERigVMMemoryType::Debug:
-				{
-					return &DebugStorageClass;
-				}
-			default:
-				{
-					break;
-				}
-			}
-			return &WorkStorageClass;
-		}
-	};
-	
-	static TMap<FString, FClassInfo> PackageToInfo;
-
 	friend class URigVMMemoryStorage;
 	friend class URigVMCompiler;
 };
@@ -297,7 +267,8 @@ public:
 	template<typename T>
 	FORCEINLINE T* GetData(int32 InPropertyIndex, const FRigVMPropertyPath& InPropertyPath)
 	{
-		return InPropertyPath.GetData<T>(GetData<uint8>(InPropertyIndex));
+		const FProperty* Property = GetProperties()[InPropertyIndex];
+		return InPropertyPath.GetData<T>(GetData<uint8>(InPropertyIndex), Property);
 	}
 
 	template<typename T>
@@ -364,4 +335,9 @@ public:
 		FRigVMMemoryHandle& SourceHandle);
 	
 #endif
+
+private:
+	
+	static const TArray<const FProperty*> EmptyProperties;
+	static const TArray<FRigVMPropertyPath> EmptyPropertyPaths;
 };
diff --git a/Engine/Source/Runtime/RigVM/Public/RigVMCore/RigVMPropertyPath.h b/Engine/Source/Runtime/RigVM/Public/RigVMCore/RigVMPropertyPath.h
--- a/Engine/Source/Runtime/RigVM/Public/RigVMCore/RigVMPropertyPath.h
+++ b/Engine/Source/Runtime/RigVM/Public/RigVMCore/RigVMPropertyPath.h
@@ -8,21 +8,25 @@
 struct RIGVM_API FRigVMPropertyPathDescription
 {
 	int32 PropertyIndex;
+	FString RootCPPType;
 	FString SegmentPath;
 	
 	FRigVMPropertyPathDescription()
 		: PropertyIndex(INDEX_NONE)
+		, RootCPPType()
 		, SegmentPath()
 	{}
 
-	FRigVMPropertyPathDescription(int32 InPropertyIndex, const FString& InSegmentPath)
+	FRigVMPropertyPathDescription(int32 InPropertyIndex, const FString& InRootCPPType, const FString& InSegmentPath)
 		: PropertyIndex(InPropertyIndex)
+		, RootCPPType(InRootCPPType)
 		, SegmentPath(InSegmentPath)
 	{}
 
 	FORCEINLINE friend FArchive& operator<<(FArchive& Ar, FRigVMPropertyPathDescription& Path)
 	{
 		Ar << Path.PropertyIndex;
+		Ar << Path.RootCPPType;
 		Ar << Path.SegmentPath;
 		return Ar;
 	}
@@ -78,12 +82,10 @@ public:
 		return GetTypeHash(this) < GetTypeHash(Other);
 	}
 
-	uint8* GetData_Internal(uint8* InPtr) const;
-
 	template<typename T>
-	FORCEINLINE T* GetData(uint8* InPtr) const
+	FORCEINLINE T* GetData(uint8* InPtr, const FProperty* InProperty) const
 	{
-		return (T*)GetData_Internal(InPtr);
+		return (T*)GetData_Internal(InPtr, InProperty);
 	}
 
 	friend FORCEINLINE uint32 GetTypeHash(const FRigVMPropertyPath& InPropertyPath)
@@ -99,10 +101,14 @@ public:
 		);
 	}
 
+	const FProperty* GetTargetProperty() const;
+	
 	static FRigVMPropertyPath Empty;
 
 private:
 
+	uint8* GetData_Internal(uint8* InPtr, const FProperty* InProperty) const;
+
 	FString Path;
 	TArray<FRigVMPropertyPathSegment> Segments;
 };
diff --git a/Engine/Source/Runtime/RigVM/Public/RigVMCore/RigVMStruct.h b/Engine/Source/Runtime/RigVM/Public/RigVMCore/RigVMStruct.h
--- a/Engine/Source/Runtime/RigVM/Public/RigVMCore/RigVMStruct.h
+++ b/Engine/Source/Runtime/RigVM/Public/RigVMCore/RigVMStruct.h
@@ -3,6 +3,7 @@
 #pragma once
 
 #include "CoreMinimal.h"
+#include "RigVMDefines.h"
 #include "RigVMCore/RigVMRegistry.h"
 #include "RigVMCore/RigVMExternalVariable.h"
 #include "RigVMCore/RigVMTraits.h"
@@ -146,7 +147,9 @@ struct RIGVM_API FRigVMStruct
 
 public:
 
+#if UE_RIGVM_UCLASS_BASED_STORAGE_DISABLED
 	FORCEINLINE virtual int32 GetArraySize(const FName& InParameterName, const FRigVMUserDataArray& RigVMUserData) { return INDEX_NONE; }
+#endif
 
 	// loop related
 	FORCEINLINE virtual bool IsForLoop() const { return false; }
