Control Rig: Prevent pasting of local variable nodes in rig graph

diff --git a/Engine/Source/Developer/RigVMDeveloper/Private/RigVMModel/RigVMController.cpp b/Engine/Source/Developer/RigVMDeveloper/Private/RigVMModel/RigVMController.cpp
--- a/Engine/Source/Developer/RigVMDeveloper/Private/RigVMModel/RigVMController.cpp
+++ b/Engine/Source/Developer/RigVMDeveloper/Private/RigVMModel/RigVMController.cpp
@@ -8949,6 +8949,33 @@ bool URigVMController::CanAddNode(URigVMNode* InNode, bool bReportErrors, bool b
 			}
 		}
 	}
+	else if(URigVMVariableNode* VariableNode = Cast<URigVMVariableNode>(InNode))
+	{
+		if (URigVMPin* NamePin = VariableNode->FindPin(URigVMVariableNode::VariableName))
+		{
+			FString VarName = NamePin->GetDefaultValue();
+			if (!VarName.IsEmpty())
+			{
+				for (auto LocalVar : Graph->GetLocalVariables())
+				{
+					if (LocalVar.Name.ToString() == VarName)
+					{
+						return true;
+					}
+				}
+
+				for (auto ExternalVar : Graph->GetExternalVariables())
+				{
+					if (ExternalVar.Name.ToString() == VarName)
+					{
+						return true;
+					}
+				}
+
+				return false;
+			}
+		}
+	}
 
 	return true;
 }
