RigVMCompiler: When promoting a pin's property to an array, reformat its original defaultvalue as an array defaultvalue
#rnx

diff --git a/Engine/Source/Developer/RigVMDeveloper/Private/RigVMCompiler/RigVMCompiler.cpp b/Engine/Source/Developer/RigVMDeveloper/Private/RigVMCompiler/RigVMCompiler.cpp
--- a/Engine/Source/Developer/RigVMDeveloper/Private/RigVMCompiler/RigVMCompiler.cpp
+++ b/Engine/Source/Developer/RigVMDeveloper/Private/RigVMCompiler/RigVMCompiler.cpp
@@ -2328,6 +2328,7 @@ FRigVMOperand URigVMCompiler::FindOrAddRegister(const FRigVMVarExprAST* InVarExp
 		{
 			// debug values are always stored as arrays
 			CPPType = FString::Printf(TEXT("TArray<%s>"), *CPPType);
+			JoinedDefaultValue = URigVMPin::GetDefaultValueForArray({ JoinedDefaultValue });
 		}
 		else if(Pin->GetDirection() == ERigVMPinDirection::Hidden && Pin->GetNode()->IsA<URigVMUnitNode>())
 		{
@@ -2338,6 +2339,7 @@ FRigVMOperand URigVMCompiler::FindOrAddRegister(const FRigVMVarExprAST* InVarExp
 			if (!Property->HasMetaData(FRigVMStruct::SingletonMetaName))
 			{
 				CPPType = FString::Printf(TEXT("TArray<%s>"), *CPPType);
+				JoinedDefaultValue = URigVMPin::GetDefaultValueForArray({ JoinedDefaultValue });
 			}
 		}
 		Operand = WorkData.AddProperty(MemoryType, RegisterName, CPPType, Pin->GetCPPTypeObject(), JoinedDefaultValue);
diff --git a/Engine/Source/Developer/RigVMDeveloper/Private/RigVMModel/RigVMPin.cpp b/Engine/Source/Developer/RigVMDeveloper/Private/RigVMModel/RigVMPin.cpp
--- a/Engine/Source/Developer/RigVMDeveloper/Private/RigVMModel/RigVMPin.cpp
+++ b/Engine/Source/Developer/RigVMDeveloper/Private/RigVMModel/RigVMPin.cpp
@@ -8,9 +8,10 @@
 #include "RigVMCompiler/RigVMCompiler.h"
 #include "RigVMCore/RigVMExecuteContext.h"
 #include "UObject/Package.h"
+#include "Misc/DefaultValueHelper.h"
 #include "Misc/PackageName.h"
 #include "Misc/OutputDevice.h"
-#include "Misc/DefaultValueHelper.h"
+#include "Misc/StringBuilder.h"
 #include "Logging/LogScopedVerbosityOverride.h"
 #include "RigVMModel/Nodes/RigVMCollapseNode.h"
 #include "RigVMModel/Nodes/RigVMFunctionReferenceNode.h"
@@ -163,6 +164,22 @@ TArray<FString> URigVMPin::SplitDefaultValue(const FString& InDefaultValue)
 	return Parts;
 }
 
+FString URigVMPin::GetDefaultValueForArray(TConstArrayView<FString> DefaultValues)
+{
+	TStringBuilder<256> Builder;
+	Builder << TCHAR('(');
+	if (DefaultValues.Num())
+	{
+		Builder << DefaultValues[0];
+		for (const FString& DefaultValue : DefaultValues.Slice(1, DefaultValues.Num() - 1))
+		{
+			Builder << TCHAR(',') << DefaultValue;
+		}
+	}
+	Builder << TCHAR(')');
+	return FString(Builder);
+}
+
 URigVMPin::URigVMPin()
 	: Direction(ERigVMPinDirection::Invalid)
 	, bIsExpanded(false)
diff --git a/Engine/Source/Developer/RigVMDeveloper/Public/RigVMModel/RigVMPin.h b/Engine/Source/Developer/RigVMDeveloper/Public/RigVMModel/RigVMPin.h
--- a/Engine/Source/Developer/RigVMDeveloper/Public/RigVMModel/RigVMPin.h
+++ b/Engine/Source/Developer/RigVMDeveloper/Public/RigVMModel/RigVMPin.h
@@ -3,6 +3,7 @@
 #pragma once
 
 #include "CoreMinimal.h"
+#include "Containers/ArrayView.h"
 #include "UObject/ObjectMacros.h"
 #include "RigVMCore/RigVMRegistry.h"
 #include "RigVMCore/RigVMByteCode.h"
@@ -112,6 +113,8 @@ public:
 
 	// Splits the default value into name-value pairs
 	static TArray<FString> SplitDefaultValue(const FString& InDefaultValue);
+	// Joins a collection of element DefaultValues into a default value for an array of those elements
+	static FString GetDefaultValueForArray(TConstArrayView<FString> InDefaultValues);
 
 	// Default constructor
 	URigVMPin();
