Control rig: hidden pin should not have default values

diff --git a/Engine/Source/Developer/RigVMDeveloper/Private/RigVMCompiler/RigVMCompiler.cpp b/Engine/Source/Developer/RigVMDeveloper/Private/RigVMCompiler/RigVMCompiler.cpp
--- a/Engine/Source/Developer/RigVMDeveloper/Private/RigVMCompiler/RigVMCompiler.cpp
+++ b/Engine/Source/Developer/RigVMDeveloper/Private/RigVMCompiler/RigVMCompiler.cpp
@@ -2342,6 +2342,11 @@ FRigVMOperand URigVMCompiler::FindOrAddRegister(const FRigVMVarExprAST* InVarExp
 				JoinedDefaultValue = URigVMPin::GetDefaultValueForArray({ JoinedDefaultValue });
 			}
 		}
+
+		if(Pin->GetDirection() == ERigVMPinDirection::Hidden)
+		{
+			JoinedDefaultValue.Empty();
+		}
 		Operand = WorkData.AddProperty(MemoryType, RegisterName, CPPType, Pin->GetCPPTypeObject(), JoinedDefaultValue);
 		
 #endif
