Control Rig: Fix crash in GetExposedPinChain when cycles are present

diff --git a/Engine/Source/Developer/RigVMDeveloper/Private/RigVMModel/RigVMPin.cpp b/Engine/Source/Developer/RigVMDeveloper/Private/RigVMModel/RigVMPin.cpp
--- a/Engine/Source/Developer/RigVMDeveloper/Private/RigVMModel/RigVMPin.cpp
+++ b/Engine/Source/Developer/RigVMDeveloper/Private/RigVMModel/RigVMPin.cpp
@@ -245,6 +245,12 @@ FString URigVMPin::GetSegmentPath(bool bIncludeRootPin) const
 }
 
 void URigVMPin::GetExposedPinChain(TArray<const URigVMPin*>& OutExposedPins) const
+{
+	TArray<const URigVMPin*> VisitedPins = {this};
+	GetExposedPinChainImpl(OutExposedPins, VisitedPins);
+}
+
+void URigVMPin::GetExposedPinChainImpl(TArray<const URigVMPin*>& OutExposedPins, TArray<const URigVMPin*>& VisitedPins) const
 {
 	// Variable nodes do not share the operand with their source link
 	if ((GetNode()->IsA<URigVMVariableNode>() || GetNode()->IsA<URigVMParameterNode>()) && GetDirection() == ERigVMPinDirection::Input)
@@ -258,6 +264,13 @@ void URigVMPin::GetExposedPinChain(TArray<const URigVMPin*>& OutExposedPins) con
 	{
 		URigVMPin* SourcePin = Link->GetSourcePin();
 		check(SourcePin != nullptr);
+
+		// Stop recursion when cycles are present
+		if (VisitedPins.Contains(SourcePin))
+		{
+			return;
+		}
+		VisitedPins.Add(SourcePin);
 		
 		// If the source is on an entry node, add the pin and make a recursive call on the collapse node pin
 		if (URigVMFunctionEntryNode* EntryNode = Cast<URigVMFunctionEntryNode>(SourcePin->GetNode()))
@@ -267,7 +280,7 @@ void URigVMPin::GetExposedPinChain(TArray<const URigVMPin*>& OutExposedPins) con
 			{
 				if(URigVMPin* CollapseNodePin = CollapseNode->FindPin(SourcePin->GetName()))
 				{
-					CollapseNodePin->GetExposedPinChain(OutExposedPins);
+					CollapseNodePin->GetExposedPinChainImpl(OutExposedPins, VisitedPins);
 				}
 			}
 		}
@@ -278,7 +291,7 @@ void URigVMPin::GetExposedPinChain(TArray<const URigVMPin*>& OutExposedPins) con
 			{
 				if(URigVMPin* CollapseNodePin = CollapseNode->FindPin(SourcePin->GetName()))
 				{
-					CollapseNodePin->GetExposedPinChain(OutExposedPins);
+					CollapseNodePin->GetExposedPinChainImpl(OutExposedPins, VisitedPins);
 				}
 			}
 		}
@@ -289,7 +302,7 @@ void URigVMPin::GetExposedPinChain(TArray<const URigVMPin*>& OutExposedPins) con
 		}
 		else
 		{
-			SourcePin->GetExposedPinChain(OutExposedPins);
+			SourcePin->GetExposedPinChainImpl(OutExposedPins, VisitedPins);
 		}
 
 		return;
diff --git a/Engine/Source/Developer/RigVMDeveloper/Public/RigVMModel/RigVMPin.h b/Engine/Source/Developer/RigVMDeveloper/Public/RigVMModel/RigVMPin.h
--- a/Engine/Source/Developer/RigVMDeveloper/Public/RigVMModel/RigVMPin.h
+++ b/Engine/Source/Developer/RigVMDeveloper/Public/RigVMModel/RigVMPin.h
@@ -369,6 +369,8 @@ private:
 	void UpdateCPPTypeObjectIfRequired() const;
 	void SetNameFromIndex();
 
+	void GetExposedPinChainImpl(TArray<const URigVMPin*>& OutExposedPins, TArray<const URigVMPin*>& VisitedPins) const;
+
 	UPROPERTY()
 	FName DisplayName;
 
