Control Rig: Fix crash when reloading Control Rig asset after undoing collapse to function

diff --git a/Engine/Source/Developer/RigVMDeveloper/Private/RigVMModel/RigVMController.cpp b/Engine/Source/Developer/RigVMDeveloper/Private/RigVMModel/RigVMController.cpp
--- a/Engine/Source/Developer/RigVMDeveloper/Private/RigVMModel/RigVMController.cpp
+++ b/Engine/Source/Developer/RigVMDeveloper/Private/RigVMModel/RigVMController.cpp
@@ -2387,6 +2387,11 @@ TArray<FName> URigVMController::ImportNodesFromText(const FString& InText, bool
 		{
 			if(URigVMCollapseNode* CollapseNode = Cast<URigVMCollapseNode>(SubNodes[SubNodeIndex]))
 			{
+				{
+					FRigVMControllerGraphGuard GraphGuard(this, CollapseNode->GetContainedGraph(), bSetupUndoRedo);
+					ReattachLinksToPinObjects();
+				}
+				
 				SubNodes.Append(CollapseNode->GetContainedNodes());
 			}
 		}
@@ -3997,20 +4002,53 @@ URigVMFunctionReferenceNode* URigVMController::PromoteCollapseNodeToFunctionRefe
 	}
 
 	URigVMFunctionReferenceNode* FunctionRefNode = nullptr;
-	URigVMCollapseNode* FunctionDefinition = nullptr;
+
+	// Create Function
+	URigVMLibraryNode* FunctionDefinition = nullptr;
 	{
 		FRigVMControllerGraphGuard GraphGuard(this, FunctionLibrary, bSetupUndoRedo);
 		FString FunctionName = GetValidNodeName(InCollapseNode->GetName());
-		FunctionDefinition = DuplicateObject<URigVMCollapseNode>(InCollapseNode, FunctionLibrary, *FunctionName);
+		FunctionDefinition = AddFunctionToLibrary(*FunctionName, InCollapseNode->IsMutable());		
+	}
+
+	// Add interface pins in function
+	if (FunctionDefinition)
+	{
+		FRigVMControllerGraphGuard GraphGuard(this, FunctionDefinition->GetContainedGraph(), bSetupUndoRedo);
+		for(const URigVMPin* Pin : InCollapseNode->GetPins())
+		{
+			AddExposedPin(Pin->GetFName(), Pin->GetDirection(), Pin->GetCPPType(), (Pin->GetCPPTypeObject() ? *Pin->GetCPPTypeObject()->GetPathName() : TEXT("")), Pin->GetDefaultValue(), bSetupUndoRedo);
+		}
+	}
 
-		if (FunctionDefinition)
+	// Copy inner graph from collapsed node to function
+	if (FunctionDefinition)
+	{
+		FString TextContent;
+		{
+			FRigVMControllerGraphGuard GraphGuard(this, InCollapseNode->GetContainedGraph(), bSetupUndoRedo);
+			TArray<FName> NodeNames;
+			for (const URigVMNode* Node : InCollapseNode->GetContainedNodes())
+			{
+				NodeNames.Add(Node->GetFName());
+			}
+			TextContent = ExportNodesToText(NodeNames);
+		}
 		{
-			FunctionDefinition->NodeColor = FLinearColor(FColor::FromHex("005DFFFF"));
-			FunctionLibrary->Nodes.Add(FunctionDefinition);
-			Notify(ERigVMGraphNotifType::NodeAdded, FunctionDefinition);
+			FRigVMControllerGraphGuard GraphGuard(this, FunctionDefinition->GetContainedGraph(), bSetupUndoRedo);
+			ImportNodesFromText(TextContent, bSetupUndoRedo);
+
+			SetNodePosition(FunctionDefinition->GetContainedGraph()->GetEntryNode(), InCollapseNode->GetContainedGraph()->GetEntryNode()->GetPosition());
+			SetNodePosition(FunctionDefinition->GetContainedGraph()->GetReturnNode(), InCollapseNode->GetContainedGraph()->GetReturnNode()->GetPosition());
+
+			for (const URigVMLink* InnerLink : InCollapseNode->GetContainedGraph()->GetLinks())
+			{
+				AddLink(InnerLink->SourcePinPath, InnerLink->TargetPinPath, bSetupUndoRedo);
+			}
 		}
 	}
 
+	// Remove collapse node, add function reference, and add external links
 	if (FunctionDefinition)
 	{
 		FString NodeName = InCollapseNode->GetName();
@@ -4036,11 +4074,6 @@ URigVMFunctionReferenceNode* URigVMController::PromoteCollapseNodeToFunctionRefe
 				AddLink(LinkPath.Key, LinkPath.Value, bSetupUndoRedo);
 			}
 		}
-
-		{
-			FRigVMControllerGraphGuard GraphGuard(this, FunctionDefinition->GetContainedGraph(), false);
-			ReattachLinksToPinObjects(false, nullptr, true, false);
-		}
 	}
 
 	if (bSetupUndoRedo)
