ControlRig: Array Nodes

preflight: https://horde.devtools.epicgames.com/job/611cc8193a81b000012a0554

diff --git a/Engine/Plugins/Animation/ControlRig/Source/ControlRig/Private/ControlRig.cpp b/Engine/Plugins/Animation/ControlRig/Source/ControlRig/Private/ControlRig.cpp
--- a/Engine/Plugins/Animation/ControlRig/Source/ControlRig/Private/ControlRig.cpp
+++ b/Engine/Plugins/Animation/ControlRig/Source/ControlRig/Private/ControlRig.cpp
@@ -670,6 +670,7 @@ void UControlRig::Execute(const EControlRigState InState, const FName& InEventNa
 	if (ControlRigLog != nullptr)
 	{
 		ControlRigLog->Reset();
+		UpdateVMSettings();
 	}
 #endif
 
@@ -1110,19 +1111,38 @@ void UControlRig::SetVM(URigVM* NewVM)
 		{
 			NewVM->ExecutionReachedExit().AddUObject(this, &UControlRig::HandleExecutionReachedExit);
 		}
+	}
 
-		// setup array handling and error reporting on the VM
+	VM = NewVM;
+
+	UpdateVMSettings();
+}
+
+void UControlRig::UpdateVMSettings()
+{
+	if(VM)
+	{
 #if WITH_EDITOR
-		VMRuntimeSettings.LogFunction = [this](EMessageSeverity::Type InSeverity, int32 InInstructionIndex, const FString& Message)
+
+		// setup array handling and error reporting on the VM
+		VMRuntimeSettings.LogFunction = [this](EMessageSeverity::Type InSeverity, const FRigVMExecuteContext* InContext, const FString& Message)
 		{
-			LogOnce(InSeverity, InInstructionIndex, Message);
+			check(InContext);
+
+			if(ControlRigLog)
+			{
+				ControlRigLog->Report(InSeverity, InContext->FunctionName, InContext->InstructionIndex, Message);
+			}
+			else
+			{
+				LogOnce(InSeverity, InContext->InstructionIndex, Message);
+			}
 		};
+		
 #endif
-		NewVM->SetRuntimeSettings(VMRuntimeSettings);
-
+		
+		VM->SetRuntimeSettings(VMRuntimeSettings);
 	}
-
-	VM = NewVM;
 }
 
 URigVM* UControlRig::GetVM()
diff --git a/Engine/Plugins/Animation/ControlRig/Source/ControlRig/Private/Units/Execution/RigUnit_Collection.cpp b/Engine/Plugins/Animation/ControlRig/Source/ControlRig/Private/Units/Execution/RigUnit_Collection.cpp
--- a/Engine/Plugins/Animation/ControlRig/Source/ControlRig/Private/Units/Execution/RigUnit_Collection.cpp
+++ b/Engine/Plugins/Animation/ControlRig/Source/ControlRig/Private/Units/Execution/RigUnit_Collection.cpp
@@ -205,6 +205,14 @@ FRigUnit_CollectionItems_Execute()
 	}
 }
 
+FRigUnit_CollectionGetItems_Execute()
+{
+	DECLARE_SCOPE_HIERARCHICAL_COUNTER_RIGUNIT()
+
+	Items = Collection.GetKeys();
+}
+
+
 FRigUnit_CollectionUnion_Execute()
 {
 	DECLARE_SCOPE_HIERARCHICAL_COUNTER_RIGUNIT()
diff --git a/Engine/Plugins/Animation/ControlRig/Source/ControlRig/Private/Units/Execution/RigUnit_Collection.h b/Engine/Plugins/Animation/ControlRig/Source/ControlRig/Private/Units/Execution/RigUnit_Collection.h
--- a/Engine/Plugins/Animation/ControlRig/Source/ControlRig/Private/Units/Execution/RigUnit_Collection.h
+++ b/Engine/Plugins/Animation/ControlRig/Source/ControlRig/Private/Units/Execution/RigUnit_Collection.h
@@ -205,6 +205,28 @@ struct CONTROLRIG_API FRigUnit_CollectionItems : public FRigUnit_CollectionBase
 	FRigElementKeyCollection Collection;
 };
 
+/**
+* Returns a collection provided a specific list of items.
+*/
+USTRUCT(meta = (DisplayName = "Get Items", Keywords = "Collection,Array", Varying))
+struct CONTROLRIG_API FRigUnit_CollectionGetItems : public FRigUnit_CollectionBase
+{
+	GENERATED_BODY()
+
+	FRigUnit_CollectionGetItems()
+	{
+	}
+
+	RIGVM_METHOD()
+	virtual void Execute(const FRigUnitContext& Context) override;
+
+	UPROPERTY(meta = (Input))
+	FRigElementKeyCollection Collection;
+
+	UPROPERTY(meta = (Output))
+	TArray<FRigElementKey> Items;
+};
+
 /**
  * Returns the union of two provided collections
  * (the combination of all items from both A and B).
diff --git a/Engine/Plugins/Animation/ControlRig/Source/ControlRig/Public/ControlRig.h b/Engine/Plugins/Animation/ControlRig/Source/ControlRig/Public/ControlRig.h
--- a/Engine/Plugins/Animation/ControlRig/Source/ControlRig/Public/ControlRig.h
+++ b/Engine/Plugins/Animation/ControlRig/Source/ControlRig/Public/ControlRig.h
@@ -267,6 +267,9 @@ public:
 	/** Set the VM, and bind its commands */
 	void SetVM(URigVM* VM);
 
+	/** Update the settings such as array bound and log facilities */
+	void UpdateVMSettings();
+
 	UFUNCTION(BlueprintPure, Category = "Control Rig")
 	URigVM* GetVM();
 
diff --git a/Engine/Plugins/Animation/ControlRig/Source/ControlRig/Public/Rigs/RigHierarchyDefines.h b/Engine/Plugins/Animation/ControlRig/Source/ControlRig/Public/Rigs/RigHierarchyDefines.h
--- a/Engine/Plugins/Animation/ControlRig/Source/ControlRig/Public/Rigs/RigHierarchyDefines.h
+++ b/Engine/Plugins/Animation/ControlRig/Source/ControlRig/Public/Rigs/RigHierarchyDefines.h
@@ -1414,6 +1414,11 @@ struct CONTROLRIG_API FRigElementKeyCollection
 	{
 		return Keys[InIndex];
 	}
+
+	FORCEINLINE const TArray<FRigElementKey>& GetKeys() const
+	{
+		return Keys;
+	}
 	   
 	FORCEINLINE TArray<FRigElementKey>::RangedForIteratorType      begin() { return Keys.begin(); }
 	FORCEINLINE TArray<FRigElementKey>::RangedForConstIteratorType begin() const { return Keys.begin(); }
diff --git a/Engine/Plugins/Animation/ControlRig/Source/ControlRigDeveloper/Private/Graph/ControlRigGraphNode.cpp b/Engine/Plugins/Animation/ControlRig/Source/ControlRigDeveloper/Private/Graph/ControlRigGraphNode.cpp
--- a/Engine/Plugins/Animation/ControlRig/Source/ControlRigDeveloper/Private/Graph/ControlRigGraphNode.cpp
+++ b/Engine/Plugins/Animation/ControlRig/Source/ControlRigDeveloper/Private/Graph/ControlRigGraphNode.cpp
@@ -653,6 +653,7 @@ bool UControlRigGraphNode::ShowPaletteIconOnNode() const
 			ModelNode->IsA<URigVMFunctionReturnNode>() ||
 			ModelNode->IsA<URigVMFunctionReferenceNode>() ||
 			ModelNode->IsA<URigVMCollapseNode>() ||
+			ModelNode->IsLoopNode() ||
 			Cast<URigVMUnitNode>(ModelNode) != nullptr;
 	}
 	return false;
@@ -666,6 +667,7 @@ FSlateIcon UControlRigGraphNode::GetIconAndTint(FLinearColor& OutColor) const
 	static FSlateIcon EventIcon("EditorStyle", "GraphEditor.Event_16x");
 	static FSlateIcon EntryReturnIcon("EditorStyle", "GraphEditor.Default_16x");
 	static FSlateIcon CollapsedNodeIcon("EditorStyle", "GraphEditor.SubGraph_16x");
+	static FSlateIcon ArrayNodeIteratorIcon("EditorStyle", "GraphEditor.Macro.ForEach_16x");
 
 	if (URigVMNode* ModelNode = GetModelNode())
 	{
@@ -690,6 +692,14 @@ FSlateIcon UControlRigGraphNode::GetIconAndTint(FLinearColor& OutColor) const
 			return EntryReturnIcon;
 		}
 
+		if (URigVMArrayNode* ArrayNode = Cast<URigVMArrayNode>(ModelNode))
+		{
+			if(ArrayNode->IsLoopNode())
+			{
+				return ArrayNodeIteratorIcon;
+			}
+		}
+
 		if (URigVMUnitNode *UnitNode = Cast<URigVMUnitNode>(ModelNode))
 		{
 			ensure(UnitNode->GetScriptStruct());
diff --git a/Engine/Plugins/Animation/ControlRig/Source/ControlRigDeveloper/Private/Graph/ControlRigGraphSchema.cpp b/Engine/Plugins/Animation/ControlRig/Source/ControlRigDeveloper/Private/Graph/ControlRigGraphSchema.cpp
--- a/Engine/Plugins/Animation/ControlRig/Source/ControlRigDeveloper/Private/Graph/ControlRigGraphSchema.cpp
+++ b/Engine/Plugins/Animation/ControlRig/Source/ControlRigDeveloper/Private/Graph/ControlRigGraphSchema.cpp
@@ -622,7 +622,7 @@ bool UControlRigGraphSchema::SupportsPinType(TWeakPtr<const FEdGraphSchemaAction
 bool UControlRigGraphSchema::SupportsPinTypeContainer(TWeakPtr<const FEdGraphSchemaAction> SchemaAction,
 	const FEdGraphPinType& PinType, const EPinContainerType& ContainerType) const
 {
-	return ContainerType == EPinContainerType::None;
+	return ContainerType == EPinContainerType::None || ContainerType == EPinContainerType::Array;
 }
 
 void UControlRigGraphSchema::BreakPinLinks(UEdGraphPin& TargetPin, bool bSendsNodeNotifcation) const
diff --git a/Engine/Plugins/Animation/ControlRig/Source/ControlRigDeveloper/Public/Graph/ControlRigGraph.h b/Engine/Plugins/Animation/ControlRig/Source/ControlRigDeveloper/Public/Graph/ControlRigGraph.h
--- a/Engine/Plugins/Animation/ControlRig/Source/ControlRigDeveloper/Public/Graph/ControlRigGraph.h
+++ b/Engine/Plugins/Animation/ControlRig/Source/ControlRigDeveloper/Public/Graph/ControlRigGraph.h
@@ -160,6 +160,7 @@ private:
 	friend class UControlRigPrototypeNodeSpawner;
 	friend class UControlRigEnumNodeSpawner;
 	friend class UControlRigFunctionRefNodeSpawner;
+	friend class UControlRigArrayNodeSpawner;
 
 #endif
 	friend class UControlRigGraphNode;
diff --git a/Engine/Plugins/Animation/ControlRig/Source/ControlRigDeveloper/Public/Graph/ControlRigGraphSchema.h b/Engine/Plugins/Animation/ControlRig/Source/ControlRigDeveloper/Public/Graph/ControlRigGraphSchema.h
--- a/Engine/Plugins/Animation/ControlRig/Source/ControlRigDeveloper/Public/Graph/ControlRigGraphSchema.h
+++ b/Engine/Plugins/Animation/ControlRig/Source/ControlRigDeveloper/Public/Graph/ControlRigGraphSchema.h
@@ -212,5 +212,6 @@ private:
 	friend class UControlRigIfNodeSpawner;
 	friend class UControlRigSelectNodeSpawner;
 	friend class UControlRigUnitNodeSpawner;
+	friend class UControlRigArrayNodeSpawner;
 };
 
diff --git a/Engine/Plugins/Animation/ControlRig/Source/ControlRigEditor/Private/ControlRigEditorModule.cpp b/Engine/Plugins/Animation/ControlRig/Source/ControlRigEditor/Private/ControlRigEditorModule.cpp
--- a/Engine/Plugins/Animation/ControlRig/Source/ControlRigEditor/Private/ControlRigEditorModule.cpp
+++ b/Engine/Plugins/Animation/ControlRig/Source/ControlRigEditor/Private/ControlRigEditorModule.cpp
@@ -54,6 +54,7 @@
 #include "Graph/NodeSpawners/ControlRigPrototypeNodeSpawner.h"
 #include "Graph/NodeSpawners/ControlRigEnumNodeSpawner.h"
 #include "Graph/NodeSpawners/ControlRigFunctionRefNodeSpawner.h"
+#include "Graph/NodeSpawners/ControlRigArrayNodeSpawner.h"
 #include "Kismet2/BlueprintEditorUtils.h"
 #include "Kismet2/KismetDebugUtilities.h"
 #include "Graph/ControlRigGraphNode.h"
@@ -1142,6 +1143,21 @@ void FControlRigEditorModule::GetTypeActions(UControlRigBlueprint* CRB, FBluepri
 		LOCTEXT("SelectSpawnerTooltip", "Adds a new 'select' node to the graph"));
 	ActionRegistrar.AddBlueprintAction(ActionKey, SelectNodeSpawner);
 
+	const int32 FirstArrayOpCode = (int32)ERigVMOpCode::FirstArrayOpCode; 
+	const int32 LastArrayOpCode = (int32)ERigVMOpCode::LastArrayOpCode;
+	for(int32 OpCodeIndex = FirstArrayOpCode; OpCodeIndex <= LastArrayOpCode; OpCodeIndex++)
+	{
+		ERigVMOpCode OpCode = (ERigVMOpCode)OpCodeIndex;
+		FString OpCodeString = URigVMArrayNode::GetNodeTitle(OpCode);
+
+		UBlueprintNodeSpawner* ArrayNodeSpawner = UControlRigArrayNodeSpawner::CreateGeneric(
+			OpCode,
+			FText::FromString(OpCodeString),
+			LOCTEXT("ArraySpawnerCategory", "Array"),
+			FText::FromString(FString::Printf(TEXT("Adds a new '%s' node to the graph"), *OpCodeString)));
+		ActionRegistrar.AddBlueprintAction(ActionKey, ArrayNodeSpawner);
+	}
+
 	for (TObjectIterator<UEnum> EnumIt; EnumIt; ++EnumIt)
 	{
 		UEnum* EnumToConsider = (*EnumIt);
diff --git a/Engine/Plugins/Animation/ControlRig/Source/ControlRigEditor/Private/Editor/SControlRigStackView.cpp b/Engine/Plugins/Animation/ControlRig/Source/ControlRigEditor/Private/Editor/SControlRigStackView.cpp
--- a/Engine/Plugins/Animation/ControlRig/Source/ControlRigEditor/Private/Editor/SControlRigStackView.cpp
+++ b/Engine/Plugins/Animation/ControlRig/Source/ControlRigEditor/Private/Editor/SControlRigStackView.cpp
@@ -526,6 +526,22 @@ void SControlRigStackView::PopulateStackView(URigVM* InVM)
 					case ERigVMOpCode::BeginBlock:
 					case ERigVMOpCode::EndBlock:
 					case ERigVMOpCode::ChangeType:
+					case ERigVMOpCode::ArrayReset:
+					case ERigVMOpCode::ArrayGetNum: 
+					case ERigVMOpCode::ArraySetNum:
+					case ERigVMOpCode::ArrayGetAtIndex:  
+					case ERigVMOpCode::ArraySetAtIndex:
+					case ERigVMOpCode::ArrayAdd:
+					case ERigVMOpCode::ArrayInsert:
+					case ERigVMOpCode::ArrayRemove:
+					case ERigVMOpCode::ArrayFind:
+					case ERigVMOpCode::ArrayAppend:
+					case ERigVMOpCode::ArrayClone:
+					case ERigVMOpCode::ArrayIterator:
+					case ERigVMOpCode::ArrayUnion:
+					case ERigVMOpCode::ArrayDifference:
+					case ERigVMOpCode::ArrayIntersection:
+					case ERigVMOpCode::ArrayReverse:
 					{
 						const FText OpCodeText = StaticEnum<ERigVMOpCode>()->GetDisplayNameTextByValue((int32)Instructions[InstructionIndex].OpCode);
 						if(Node)
diff --git a/Engine/Plugins/Animation/ControlRig/Source/ControlRigEditor/Private/Graph/NodeSpawners/ControlRigArrayNodeSpawner.cpp b/Engine/Plugins/Animation/ControlRig/Source/ControlRigEditor/Private/Graph/NodeSpawners/ControlRigArrayNodeSpawner.cpp
new file mode 100644
--- /dev/null
+++ b/Engine/Plugins/Animation/ControlRig/Source/ControlRigEditor/Private/Graph/NodeSpawners/ControlRigArrayNodeSpawner.cpp
@@ -0,0 +1,252 @@
+// Copyright Epic Games, Inc. All Rights Reserved.
+
+#include "ControlRigArrayNodeSpawner.h"
+#include "ControlRigUnitNodeSpawner.h"
+#include "Graph/ControlRigGraph.h"
+#include "Graph/ControlRigGraphNode.h"
+#include "Graph/ControlRigGraphSchema.h"
+#include "BlueprintNodeTemplateCache.h"
+#include "RigVMModel/RigVMController.h"
+#include "ControlRigBlueprint.h"
+#include "ControlRigBlueprintUtils.h"
+#include "RigVMModel/Nodes/RigVMArrayNode.h"
+
+#if WITH_EDITOR
+#include "Editor.h"
+#include "Kismet2/BlueprintEditorUtils.h"
+#endif
+
+#define LOCTEXT_NAMESPACE "ControlRigArrayNodeSpawner"
+
+UControlRigArrayNodeSpawner* UControlRigArrayNodeSpawner::CreateGeneric(ERigVMOpCode InOpCode, const FText& InMenuDesc, const FText& InCategory, const FText& InTooltip)
+{
+	UControlRigArrayNodeSpawner* NodeSpawner = NewObject<UControlRigArrayNodeSpawner>(GetTransientPackage());
+	NodeSpawner->NodeClass = UControlRigGraphNode::StaticClass();
+	NodeSpawner->OpCode = InOpCode;
+
+	FBlueprintActionUiSpec& MenuSignature = NodeSpawner->DefaultMenuSignature;
+	
+	MenuSignature.MenuName = InMenuDesc;
+	MenuSignature.Tooltip  = InTooltip;
+	MenuSignature.Category = InCategory;
+
+	switch(InOpCode)
+	{
+		case ERigVMOpCode::ArrayReset:
+		{
+			MenuSignature.Keywords = FText::FromString(TEXT("Clear,Empty,RemoveAll"));
+			break;
+		}
+		case ERigVMOpCode::ArrayGetNum:
+		{
+			MenuSignature.Keywords = FText::FromString(TEXT("Size,Length,Count"));
+			break;
+		} 
+		case ERigVMOpCode::ArraySetNum:
+		{
+			MenuSignature.Keywords = FText::FromString(TEXT("Size,Length,Count"));
+			break;
+		}
+		case ERigVMOpCode::ArrayGetAtIndex:
+		{
+			MenuSignature.Keywords = FText::FromString(TEXT("Get,Element,At,Entry,[]"));
+			break;
+		}  
+		case ERigVMOpCode::ArraySetAtIndex:
+		{
+			MenuSignature.Keywords = FText::FromString(TEXT("Set,Element,At,Entry,[]"));
+			break;
+		}
+		case ERigVMOpCode::ArrayAdd:
+		{
+			MenuSignature.Keywords = FText::FromString(TEXT("Push"));
+			break;
+		}
+		case ERigVMOpCode::ArrayRemove:
+		{
+			MenuSignature.Keywords = FText::FromString(TEXT("Pop"));
+			break;
+		}
+		case ERigVMOpCode::ArrayFind:
+		{
+			MenuSignature.Keywords = FText::FromString(TEXT("Search,Contains"));
+			break;
+		}
+		case ERigVMOpCode::ArrayAppend:
+		{
+			MenuSignature.Keywords = FText::FromString(TEXT("Concatenate,Join,Merge"));
+			break;
+		}
+		case ERigVMOpCode::ArrayClone:
+		{
+			MenuSignature.Keywords = FText::FromString(TEXT("Make,Copy,Duplicate"));
+			break;
+		}
+		case ERigVMOpCode::ArrayIterator:
+		{
+			MenuSignature.Keywords = FText::FromString(TEXT("ForEach,ForLoop,Iterate"));
+			break;
+		}
+		case ERigVMOpCode::ArrayUnion:
+		{
+			MenuSignature.Keywords = FText::FromString(TEXT("Join,Merge,Concatenate"));
+			break;
+		}
+		case ERigVMOpCode::ArrayDifference:
+		case ERigVMOpCode::ArrayIntersection:
+		case ERigVMOpCode::ArrayReverse:
+		case ERigVMOpCode::ArrayInsert:
+		default:
+		{
+			MenuSignature.Keywords = FText::FromString(TEXT(" "));
+			break;
+		}
+	}
+	MenuSignature.Icon = FSlateIcon(TEXT("ControlRigEditorStyle"), TEXT("ControlRig.RigUnit"));
+
+	return NodeSpawner;
+}
+
+FBlueprintNodeSignature UControlRigArrayNodeSpawner::GetSpawnerSignature() const
+{
+	return FBlueprintNodeSignature(NodeClass);
+}
+
+FBlueprintActionUiSpec UControlRigArrayNodeSpawner::GetUiSpec(FBlueprintActionContext const& Context, FBindingSet const& Bindings) const
+{
+	UEdGraph* TargetGraph = (Context.Graphs.Num() > 0) ? Context.Graphs[0] : nullptr;
+	FBlueprintActionUiSpec MenuSignature = PrimeDefaultUiSpec(TargetGraph);
+
+	DynamicUiSignatureGetter.ExecuteIfBound(Context, Bindings, &MenuSignature);
+	return MenuSignature;
+}
+
+UEdGraphNode* UControlRigArrayNodeSpawner::Invoke(UEdGraph* ParentGraph, FBindingSet const& Bindings, FVector2D const Location) const
+{
+	UControlRigGraphNode* NewNode = nullptr;
+
+	bool const bIsTemplateNode = FBlueprintNodeTemplateCache::IsTemplateOuter(ParentGraph);
+
+	if (bIsTemplateNode)
+	{
+		NewNode = NewObject<UControlRigGraphNode>(ParentGraph, TEXT("ArrayNode"));
+		ParentGraph->AddNode(NewNode, false);
+
+		NewNode->CreateNewGuid();
+		NewNode->PostPlacedNewNode();
+
+		UEdGraphPin* InputValuePin = UEdGraphPin::CreatePin(NewNode);
+		UEdGraphPin* OutputValuePin = UEdGraphPin::CreatePin(NewNode);
+		NewNode->Pins.Add(InputValuePin);
+		NewNode->Pins.Add(OutputValuePin);
+
+		InputValuePin->PinType.PinCategory = TEXT("POLYMORPH");
+		OutputValuePin->PinType.PinCategory = TEXT("POLYMORPH");
+		InputValuePin->Direction = EGPD_Input;
+		OutputValuePin->Direction = EGPD_Output;
+		NewNode->SetFlags(RF_Transactional);
+
+		return NewNode;
+	}
+
+
+	bool const bUndo = !bIsTemplateNode;
+
+	// First create a backing member for our node
+	UControlRigGraph* RigGraph = Cast<UControlRigGraph>(ParentGraph);
+	if(RigGraph == nullptr) return nullptr;
+	UControlRigBlueprint* RigBlueprint = Cast<UControlRigBlueprint>(FBlueprintEditorUtils::FindBlueprintForGraph(ParentGraph));
+	check(RigBlueprint);
+
+	FString CPPType = "int32";
+	FString CPPTypeObjectPath;
+
+	if (UControlRigGraphSchema* RigSchema = Cast<UControlRigGraphSchema>((UEdGraphSchema*)ParentGraph->GetSchema()))
+	{
+		if (const UEdGraphPin* LastPin = RigSchema->LastPinForCompatibleCheck)
+		{
+			if (URigVMPin* ModelPin = RigBlueprint->GetModel(ParentGraph)->FindPin(LastPin->GetName()))
+			{
+				CPPType = ModelPin->GetCPPType();
+				if(ModelPin->IsArray())
+				{
+					CPPType = ModelPin->GetArrayElementCppType();
+				}
+				if (ModelPin->GetCPPTypeObject())
+				{
+					CPPTypeObjectPath = ModelPin->GetCPPTypeObject()->GetPathName();
+				}
+			}
+		}
+	}
+
+	FName MemberName = NAME_None;
+
+#if WITH_EDITOR
+	if (GEditor && !bIsTemplateNode)
+	{
+		GEditor->CancelTransaction(0);
+	}
+#endif
+
+	URigVMController* Controller = bIsTemplateNode ? RigGraph->GetTemplateController() : RigBlueprint->GetController(ParentGraph);
+
+	const FString OpCodeString = StaticEnum<ERigVMOpCode>()->GetNameStringByValue((int64)OpCode);
+	const FName Name = *OpCodeString;
+
+	if (!bIsTemplateNode)
+	{
+		Controller->OpenUndoBracket(FString::Printf(TEXT("Add '%s' Node"), *Name.ToString()));
+	}
+
+	if (URigVMArrayNode* ModelNode = Controller->AddArrayNodeFromObjectPath(OpCode, CPPType, CPPTypeObjectPath, Location, Name.ToString(), bUndo, !bIsTemplateNode))
+	{
+		NewNode = Cast<UControlRigGraphNode>(RigGraph->FindNodeForModelNodeName(ModelNode->GetFName()));
+
+		if (NewNode && bUndo)
+		{
+			Controller->ClearNodeSelection(true);
+			Controller->SelectNode(ModelNode, true, true);
+		}
+
+		if (bUndo)
+		{
+			Controller->CloseUndoBracket();
+		}
+	}
+	else
+	{
+		if (bUndo)
+		{
+			Controller->CancelUndoBracket();
+		}
+	}
+
+
+	return NewNode;
+}
+
+bool UControlRigArrayNodeSpawner::IsTemplateNodeFilteredOut(FBlueprintActionFilter const& Filter) const
+{
+	for (UBlueprint* Blueprint : Filter.Context.Blueprints)
+	{
+		if (UControlRigBlueprint* RigBlueprint = Cast<UControlRigBlueprint>(Blueprint))
+		{
+			if (Filter.Context.Pins.Num() == 0)
+			{
+				return true;
+			}
+
+			FString PinPath = Filter.Context.Pins[0]->GetName();
+			UEdGraph* EdGraph = Filter.Context.Pins[0]->GetOwningNode()->GetGraph();
+			if (URigVMPin* ModelPin = RigBlueprint->GetModel(EdGraph)->FindPin(PinPath))
+			{
+				return ModelPin->IsExecuteContext();
+			}
+		}
+	}
+
+	return false;
+}
+
+#undef LOCTEXT_NAMESPACE
diff --git a/Engine/Plugins/Animation/ControlRig/Source/ControlRigEditor/Private/Graph/NodeSpawners/ControlRigArrayNodeSpawner.h b/Engine/Plugins/Animation/ControlRig/Source/ControlRigEditor/Private/Graph/NodeSpawners/ControlRigArrayNodeSpawner.h
new file mode 100644
--- /dev/null
+++ b/Engine/Plugins/Animation/ControlRig/Source/ControlRigEditor/Private/Graph/NodeSpawners/ControlRigArrayNodeSpawner.h
@@ -0,0 +1,46 @@
+// Copyright Epic Games, Inc. All Rights Reserved.
+
+#pragma once
+
+#include "CoreMinimal.h"
+#include "UObject/ObjectMacros.h"
+#include "UObject/Object.h"
+#include "Templates/SubclassOf.h"
+#include "UObject/UnrealType.h"
+#include "EdGraph/EdGraphPin.h"
+#include "Engine/Blueprint.h"
+#include "BlueprintActionFilter.h"
+#include "BlueprintNodeSignature.h"
+#include "BlueprintFieldNodeSpawner.h"
+#include "RigVMCore/RigVMByteCode.h"
+#include "ControlRigArrayNodeSpawner.generated.h"
+
+class UControlRigGraphNode;
+
+UCLASS(Transient)
+class CONTROLRIGEDITOR_API UControlRigArrayNodeSpawner : public UBlueprintNodeSpawner
+{
+	GENERATED_BODY()
+
+public:
+
+	/**
+	 * Creates a new UControlRigArrayNodeSpawner, charged with spawning 
+	 * a new member-variable node
+	 * 
+	 * @return A newly allocated instance of this class.
+	 */
+	static UControlRigArrayNodeSpawner* CreateGeneric(ERigVMOpCode InOpCode, const FText& InMenuDesc, const FText& InCategory, const FText& InTooltip);
+
+	// UBlueprintNodeSpawner interface
+	virtual FBlueprintNodeSignature GetSpawnerSignature() const override;
+	virtual FBlueprintActionUiSpec GetUiSpec(FBlueprintActionContext const& Context, FBindingSet const& Bindings) const override;
+	virtual UEdGraphNode* Invoke(UEdGraph* ParentGraph, FBindingSet const& Bindings, FVector2D const Location) const override;
+	virtual bool IsTemplateNodeFilteredOut(FBlueprintActionFilter const& Filter) const override;
+	// End UBlueprintNodeSpawner interface
+
+private:
+
+	ERigVMOpCode OpCode;
+	friend class UEngineTestControlRig;
+};
diff --git a/Engine/Plugins/Animation/ControlRig/Source/ControlRigEditor/Private/Graph/NodeSpawners/ControlRigRerouteNodeSpawner.cpp b/Engine/Plugins/Animation/ControlRig/Source/ControlRigEditor/Private/Graph/NodeSpawners/ControlRigRerouteNodeSpawner.cpp
--- a/Engine/Plugins/Animation/ControlRig/Source/ControlRigEditor/Private/Graph/NodeSpawners/ControlRigRerouteNodeSpawner.cpp
+++ b/Engine/Plugins/Animation/ControlRig/Source/ControlRigEditor/Private/Graph/NodeSpawners/ControlRigRerouteNodeSpawner.cpp
@@ -163,7 +163,7 @@ bool UControlRigRerouteNodeSpawner::IsTemplateNodeFilteredOut(FBlueprintActionFi
 			}
 			for (UEdGraphPin* Pin : Filter.Context.Pins)
 			{
-				if (Pin->PinType.ContainerType != EPinContainerType::None)
+				if (Pin->PinType.ContainerType != EPinContainerType::None && Pin->PinType.ContainerType != EPinContainerType::Array)
 				{
 					return true;
 				}
diff --git a/Engine/Plugins/Animation/ControlRig/Source/ControlRigEditor/Private/Graph/SControlRigGraphNode.cpp b/Engine/Plugins/Animation/ControlRig/Source/ControlRigEditor/Private/Graph/SControlRigGraphNode.cpp
--- a/Engine/Plugins/Animation/ControlRig/Source/ControlRigEditor/Private/Graph/SControlRigGraphNode.cpp
+++ b/Engine/Plugins/Animation/ControlRig/Source/ControlRigEditor/Private/Graph/SControlRigGraphNode.cpp
@@ -335,33 +335,37 @@ void SControlRigGraphNode::Construct( const FArguments& InArgs )
 
 				if(PinInfo.bIsContainer)
 				{
-					// make sure to minimize the width of the label
-					FirstSlot->SetAutoWidth();
-					
-					// add array plus button
-					SlotLayout->AddSlot()
-					.AutoWidth()
-					.HAlign(HAlign_Left)
-                    .Padding(PinWidgetSidePadding, TopPadding, MyEmptySidePadding, 0.f)
-					[
-						SNew(SButton)
-						.ContentPadding(0.0f)
-						.ButtonStyle(FEditorStyle::Get(), "NoBorder")
-						.OnClicked(this, &SControlRigGraphNode::HandleAddArrayElement, PinInfo.ModelPinPath)
-						.IsEnabled(this, &SGraphNode::IsNodeEditable)
-						.Cursor(EMouseCursor::Default)
-						.ToolTipText(LOCTEXT("AddArrayElement", "Add Array Element"))
+					URigVMPin* ModelPin = ModelNode->GetGraph()->FindPin(PinInfo.ModelPinPath);
+					if(ModelPin && ModelPin->GetSourceLinks().Num() == 0)
+					{
+						// make sure to minimize the width of the label
+						FirstSlot->SetAutoWidth();
+						
+						// add array plus button
+						SlotLayout->AddSlot()
+						.AutoWidth()
+						.HAlign(HAlign_Left)
+						.Padding(PinWidgetSidePadding, TopPadding, MyEmptySidePadding, 0.f)
 						[
-							SNew(SHorizontalBox)
-							+SHorizontalBox::Slot()
-							.AutoWidth()
-							.VAlign(VAlign_Center)
+							SNew(SButton)
+							.ContentPadding(0.0f)
+							.ButtonStyle(FEditorStyle::Get(), "NoBorder")
+							.OnClicked(this, &SControlRigGraphNode::HandleAddArrayElement, PinInfo.ModelPinPath)
+							.IsEnabled(this, &SGraphNode::IsNodeEditable)
+							.Cursor(EMouseCursor::Default)
+							.ToolTipText(LOCTEXT("AddArrayElement", "Add Array Element"))
 							[
-								SNew(SImage)
-								.Image(FEditorStyle::GetBrush(TEXT("Icons.PlusCircle")))
+								SNew(SHorizontalBox)
+								+SHorizontalBox::Slot()
+								.AutoWidth()
+								.VAlign(VAlign_Center)
+								[
+									SNew(SImage)
+									.Image(FEditorStyle::GetBrush(TEXT("Icons.PlusCircle")))
+								]
 							]
-						]
-					];
+						];
+					}
 				}
 			}
 			// io pins
diff --git a/Engine/Source/Developer/RigVMDeveloper/Private/RigVMCompiler/RigVMAST.cpp b/Engine/Source/Developer/RigVMDeveloper/Private/RigVMCompiler/RigVMAST.cpp
--- a/Engine/Source/Developer/RigVMDeveloper/Private/RigVMCompiler/RigVMAST.cpp
+++ b/Engine/Source/Developer/RigVMDeveloper/Private/RigVMCompiler/RigVMAST.cpp
@@ -13,6 +13,7 @@
 #include "RigVMModel/Nodes/RigVMEnumNode.h"
 #include "RigVMModel/Nodes/RigVMFunctionReturnNode.h"
 #include "RigVMModel/Nodes/RigVMFunctionEntryNode.h"
+#include "RigVMModel/Nodes/RigVMArrayNode.h"
 #include "RigVMModel/RigVMGraph.h"
 #include "RigVMModel/RigVMController.h"
 #include "RigVMCore/RigVMExecuteContext.h"
@@ -86,6 +87,10 @@ FName FRigVMExprAST::GetTypeName() const
 		{
 			return TEXT("[Select.]");
 		}
+		case EType::Array:
+		{
+			return TEXT("[Array..]");
+		}
 		case EType::Invalid:
 		{
 			return TEXT("[Invalid]");
@@ -413,6 +418,12 @@ FString FRigVMExprAST::DumpDot(TArray<bool>& OutExpressionDefined, const FString
 				Label = TEXT("NoOp");
 				break;
 			}
+			case EType::Array:
+			{
+				ERigVMOpCode OpCode = CastChecked<URigVMArrayNode>(To<FRigVMArrayExprAST>()->GetNode())->GetOpCode();
+				Label = StaticEnum<ERigVMOpCode>()->GetDisplayNameTextByValue((int64)OpCode).ToString();
+				break;
+			}
 			case EType::Exit:
 			{
 				Label = TEXT("Exit");
@@ -575,6 +586,29 @@ bool FRigVMNodeExprAST::IsConstant() const
 	return FRigVMExprAST::IsConstant();
 }
 
+const FRigVMVarExprAST* FRigVMNodeExprAST::FindVarWithPinName(const FName& InPinName) const
+{
+	if(URigVMNode* CurrentNode = GetNode())
+	{
+		for(URigVMPin* Pin : CurrentNode->GetPins())
+		{
+			if(Pin->GetFName() == InPinName)
+			{
+				const int32 PinIndex = Pin->GetPinIndex();
+				if(PinIndex <= NumChildren())
+				{
+					const FRigVMExprAST* Child = ChildAt(PinIndex);
+					if (Child->IsA(FRigVMExprAST::Var))
+					{
+						return Child->To<FRigVMVarExprAST>();
+					}
+				}
+			}
+		}
+	}
+	return nullptr;
+}
+
 FRigVMNodeExprAST::FRigVMNodeExprAST(EType InType, const FRigVMASTProxy& InNodeProxy)
 	: FRigVMBlockExprAST(InType)
 	, Proxy(InNodeProxy)
@@ -697,6 +731,17 @@ bool FRigVMVarExprAST::SupportsSoftLinks() const
 			}
 		}
 	}
+	else if(URigVMArrayNode* ArrayNode = Cast<URigVMArrayNode>(GetPin()->GetNode()))
+	{
+		if(ArrayNode->IsLoopNode())
+		{
+			if (GetPin()->GetFName() != FRigVMStruct::ExecuteContextName &&
+				GetPin()->GetName() != URigVMArrayNode::CompletedName)
+			{
+				return true;
+			}
+		}
+	}
 	return false;
 }
 
@@ -809,29 +854,6 @@ int32 FRigVMSelectExprAST::NumValues() const
 	return GetNode()->FindPin(URigVMSelectNode::ValueName)->GetArraySize();
 }
 
-const FRigVMVarExprAST* FRigVMCallExternExprAST::FindVarWithPinName(const FName& InPinName) const
-{
-	if(URigVMUnitNode* UnitNode = Cast<URigVMUnitNode>(GetNode()))
-	{
-		for(URigVMPin* Pin : UnitNode->GetPins())
-		{
-			if(Pin->GetFName() == InPinName)
-			{
-				const int32 PinIndex = Pin->GetPinIndex();
-				if(PinIndex <= NumChildren())
-				{
-					const FRigVMExprAST* Child = ChildAt(PinIndex);
-					if (Child->IsA(FRigVMExprAST::Var))
-					{
-						return Child->To<FRigVMVarExprAST>();
-					}
-				}
-			}
-		}
-	}
-	return nullptr;
-}
-
 void FRigVMParserASTSettings::Report(EMessageSeverity::Type InSeverity, UObject* InSubject, const FString& InMessage) const
 {
 	if (ReportDelegate.IsBound())
@@ -1014,14 +1036,8 @@ FRigVMExprAST* FRigVMParserAST::TraverseMutableNode(const FRigVMASTProxy& InNode
 				{
 					FRigVMASTProxy SourcePinProxy = InNodeProxy.GetSibling(SourcePin);
 
-					bool bIsForLoop = false;
-					if (URigVMUnitNode* UnitNode = Cast<URigVMUnitNode>(Node))
-					{
-						bIsForLoop = UnitNode->IsLoopNode();
-					}
-
 					FRigVMExprAST* ParentExpr = InParentExpr;
-					if (NodeExpr->IsA(FRigVMExprAST::Branch) || bIsForLoop)
+					if (NodeExpr->IsA(FRigVMExprAST::Branch) || Node->IsLoopNode())
 					{
 						if (FRigVMExprAST** PinExpr = SubjectToExpression.Find(SourcePinProxy))
 						{
@@ -1116,6 +1132,10 @@ FRigVMExprAST* FRigVMParserAST::CreateExpressionForNode(const FRigVMASTProxy& In
 		{
 			NodeExpr = MakeExpr<FRigVMSelectExprAST>(InNodeProxy);
 		}
+		else if (InNodeProxy.IsA<URigVMArrayNode>())
+		{
+			NodeExpr = MakeExpr<FRigVMArrayExprAST>(InNodeProxy);
+		}
 		else
 		{
 			NodeExpr = MakeExpr<FRigVMCallExternExprAST>(InNodeProxy);
@@ -1377,7 +1397,8 @@ FRigVMExprAST* FRigVMParserAST::TraverseLink(const FRigVMPinProxyPair& InLink, F
 		// if this is a copy expression - we should require the copy to use a ref instead
 		if (NodeExpr->IsA(FRigVMExprAST::EType::CallExtern) ||
 			NodeExpr->IsA(FRigVMExprAST::EType::If) ||
-			NodeExpr->IsA(FRigVMExprAST::EType::Select))
+			NodeExpr->IsA(FRigVMExprAST::EType::Select) ||
+			NodeExpr->IsA(FRigVMExprAST::EType::Array))
 		{
 			for (FRigVMExprAST* ChildExpr : *NodeExpr)
 			{
@@ -1701,6 +1722,15 @@ void FRigVMParserAST::FoldAssignments()
 				continue;
 			}
 		}
+		
+		// if this node is an array iterator node - let's skip the folding
+		if (URigVMArrayNode* ArrayNode = Cast<URigVMArrayNode>(TargetPin->GetNode()))
+		{
+			if (ArrayNode->IsLoopNode())
+			{
+				continue;
+			}
+		}
 
 		FRigVMExprAST* Parent = AssignExpr->Parents[0];
 		if (!Parent->IsA(FRigVMExprAST::EType::Var))
diff --git a/Engine/Source/Developer/RigVMDeveloper/Private/RigVMCompiler/RigVMCompiler.cpp b/Engine/Source/Developer/RigVMDeveloper/Private/RigVMCompiler/RigVMCompiler.cpp
--- a/Engine/Source/Developer/RigVMDeveloper/Private/RigVMCompiler/RigVMCompiler.cpp
+++ b/Engine/Source/Developer/RigVMDeveloper/Private/RigVMCompiler/RigVMCompiler.cpp
@@ -384,13 +384,15 @@ bool URigVMCompiler::Compile(URigVMGraph* InGraph, URigVMController* InControlle
 
 	for(ERigVMMemoryType MemoryType : MemoryTypes)
 	{
+		UPackage* Package = InGraph->GetOutermost();
+
 		TArray<FRigVMPropertyDescription>* Properties = WorkData.PropertyDescriptions.Find(MemoryType);
 		if(Properties == nullptr)
 		{
+			URigVMMemoryStorageGeneratorClass::RemoveStorageClass(Package, MemoryType);
 			continue;
 		}
 
-		UPackage* Package = InGraph->GetOutermost();
 		URigVMMemoryStorageGeneratorClass::CreateStorageClass(Package, MemoryType, *Properties);
 	}
 
@@ -569,6 +571,11 @@ void URigVMCompiler::TraverseExpression(const FRigVMExprAST* InExpr, FRigVMCompi
 			TraverseSelect(InExpr->To<FRigVMSelectExprAST>(), WorkData);
 			break;
 		}
+		case FRigVMExprAST::EType::Array:
+		{
+			TraverseArray(InExpr->To<FRigVMArrayExprAST>(), WorkData);
+			break;
+		}
 		default:
 		{
 			ensure(false);
@@ -613,7 +620,7 @@ void URigVMCompiler::TraverseEntry(const FRigVMEntryExprAST* InExpr, FRigVMCompi
 		{
 			if (ChildExpr->IsA(FRigVMExprAST::EType::Var))
 			{
-				Operands.Add(WorkData.ExprToOperand.FindChecked(ChildExpr->To<FRigVMVarExprAST>()));
+				Operands.Add(WorkData.ExprToOperand.FindChecked(GetSourceVarExpr(ChildExpr)));
 			}
 			else
 			{
@@ -666,11 +673,11 @@ int32 URigVMCompiler::TraverseCallExtern(const FRigVMCallExternExprAST* InExpr,
 		{
 			if (ChildExpr->GetType() == FRigVMExprAST::EType::CachedValue)
 			{
-				Operands.Add(WorkData.ExprToOperand.FindChecked(ChildExpr->To<FRigVMCachedValueExprAST>()->GetVarExpr()));
+				Operands.Add(WorkData.ExprToOperand.FindChecked(GetSourceVarExpr(ChildExpr->To<FRigVMCachedValueExprAST>()->GetVarExpr())));
 			}
 			else if (ChildExpr->IsA(FRigVMExprAST::EType::Var))
 			{
-				Operands.Add(WorkData.ExprToOperand.FindChecked(ChildExpr->To<FRigVMVarExprAST>()));
+				Operands.Add(WorkData.ExprToOperand.FindChecked(GetSourceVarExpr(ChildExpr->To<FRigVMVarExprAST>())));
 			}
 			else
 			{
@@ -853,12 +860,18 @@ void URigVMCompiler::TraverseAssign(const FRigVMAssignExprAST* InExpr, FRigVMCom
 		checkNoEntry();
 	}
 
-	FRigVMOperand Source = WorkData.ExprToOperand.FindChecked(SourceExpr);
+	FRigVMOperand Source = WorkData.ExprToOperand.FindChecked(GetSourceVarExpr(SourceExpr));
 
 	if (!WorkData.bSetupMemory)
 	{
 		const FRigVMVarExprAST* TargetExpr = InExpr->GetFirstParentOfType(FRigVMVarExprAST::EType::Var)->To<FRigVMVarExprAST>();
+		TargetExpr = GetSourceVarExpr(TargetExpr);
+		
 		FRigVMOperand Target = WorkData.ExprToOperand.FindChecked(TargetExpr);
+		if(Target == Source)
+		{
+			return;
+		}
 
 		// if this is a copy - we should check if operands need offsets
 		if (InExpr->GetType() == FRigVMExprAST::EType::Copy)
@@ -935,30 +948,34 @@ void URigVMCompiler::TraverseAssign(const FRigVMAssignExprAST* InExpr, FRigVMCom
 			Local::SetupRegisterOffset(WorkData.VM, InExpr->GetTargetPin(), Target, TargetExpr, false, WorkData);
 		}
 
-		WorkData.VM->GetByteCode().AddCopyOp(WorkData.VM->GetCopyOpForOperands(Source, Target));
-		int32 InstructionIndex = WorkData.VM->GetByteCode().GetNumInstructions() - 1;
-
-		if (Settings.SetupNodeInstructionIndex)
+		FRigVMCopyOp CopyOp = WorkData.VM->GetCopyOpForOperands(Source, Target);
+		if(CopyOp.IsValid())
 		{
-			if (Source.GetMemoryType() == ERigVMMemoryType::External)
+			WorkData.VM->GetByteCode().AddCopyOp(CopyOp);
+
+			int32 InstructionIndex = WorkData.VM->GetByteCode().GetNumInstructions() - 1;
+			if (Settings.SetupNodeInstructionIndex)
 			{
-				if (URigVMPin* SourcePin = InExpr->GetSourcePin())
+				if (Source.GetMemoryType() == ERigVMMemoryType::External)
 				{
-					if (URigVMVariableNode* VariableNode = Cast<URigVMVariableNode>(SourcePin->GetNode()))
+					if (URigVMPin* SourcePin = InExpr->GetSourcePin())
 					{
-						const FRigVMCallstack Callstack = SourceExpr->GetProxy().GetSibling(VariableNode).GetCallstack();
-						WorkData.VM->GetByteCode().SetSubject(InstructionIndex, Callstack.GetCallPath(), Callstack.GetStack());
+						if (URigVMVariableNode* VariableNode = Cast<URigVMVariableNode>(SourcePin->GetNode()))
+						{
+							const FRigVMCallstack Callstack = SourceExpr->GetProxy().GetSibling(VariableNode).GetCallstack();
+							WorkData.VM->GetByteCode().SetSubject(InstructionIndex, Callstack.GetCallPath(), Callstack.GetStack());
+						}
 					}
 				}
-			}
-			if (Target.GetMemoryType() == ERigVMMemoryType::External)
-			{
-				if (URigVMPin* TargetPin = InExpr->GetTargetPin())
+				if (Target.GetMemoryType() == ERigVMMemoryType::External)
 				{
-					if (URigVMVariableNode* VariableNode = Cast<URigVMVariableNode>(TargetPin->GetNode()))
+					if (URigVMPin* TargetPin = InExpr->GetTargetPin())
 					{
-						const FRigVMCallstack Callstack = TargetExpr->GetProxy().GetSibling(VariableNode).GetCallstack();
-						WorkData.VM->GetByteCode().SetSubject(InstructionIndex, Callstack.GetCallPath(), Callstack.GetStack());
+						if (URigVMVariableNode* VariableNode = Cast<URigVMVariableNode>(TargetPin->GetNode()))
+						{
+							const FRigVMCallstack Callstack = TargetExpr->GetProxy().GetSibling(VariableNode).GetCallstack();
+							WorkData.VM->GetByteCode().SetSubject(InstructionIndex, Callstack.GetCallPath(), Callstack.GetStack());
+						}
 					}
 				}
 			}
@@ -1265,7 +1282,7 @@ void URigVMCompiler::TraverseSelect(const FRigVMSelectExprAST* InExpr, FRigVMCom
 		TraverseExpression(CaseExpressions[CaseIndex], WorkData);
 
 		// add copy op to copy the result
-		FRigVMOperand& CaseOperand = WorkData.ExprToOperand.FindChecked(CaseExpressions[CaseIndex]);
+		FRigVMOperand& CaseOperand = WorkData.ExprToOperand.FindChecked(GetSourceVarExpr(CaseExpressions[CaseIndex]));
 		WorkData.VM->GetByteCode().AddCopyOp(WorkData.VM->GetCopyOpForOperands(CaseOperand, ResultOperand));
 		if (Settings.SetupNodeInstructionIndex)
 		{
@@ -1293,6 +1310,325 @@ void URigVMCompiler::TraverseSelect(const FRigVMSelectExprAST* InExpr, FRigVMCom
 	}
 }
 
+void URigVMCompiler::TraverseArray(const FRigVMArrayExprAST* InExpr, FRigVMCompilerWorkData& WorkData)
+{
+	URigVMArrayNode* ArrayNode = Cast<URigVMArrayNode>(InExpr->GetNode());
+	check(ArrayNode);
+	
+	if (WorkData.bSetupMemory)
+	{
+		TraverseChildren(InExpr, WorkData);
+	}
+	else
+	{
+		const FRigVMCallstack Callstack = InExpr->GetProxy().GetCallstack();
+
+		static const FName ExecuteName = FRigVMStruct::ExecuteName;
+		static const FName ArrayName = *URigVMArrayNode::ArrayName;
+		static const FName NumName = *URigVMArrayNode::NumName;
+		static const FName IndexName = *URigVMArrayNode::IndexName;
+		static const FName ElementName = *URigVMArrayNode::ElementName;
+		static const FName SuccessName = *URigVMArrayNode::SuccessName;
+		static const FName OtherName = *URigVMArrayNode::OtherName;
+		static const FName CloneName = *URigVMArrayNode::CloneName;
+		static const FName CountName = *URigVMArrayNode::CountName;
+		static const FName RatioName = *URigVMArrayNode::RatioName;
+		static const FName ResultName = *URigVMArrayNode::ResultName;
+		static const FName ContinueName = *URigVMArrayNode::ContinueName;
+		static const FName CompletedName = *URigVMArrayNode::CompletedName;
+
+		const ERigVMOpCode OpCode = ArrayNode->GetOpCode();
+		switch(OpCode)
+		{
+			case ERigVMOpCode::ArrayReset:
+			{
+				TraverseChildren(InExpr, WorkData);
+				const FRigVMOperand& ArrayOperand = WorkData.ExprToOperand.FindChecked(GetSourceVarExpr(InExpr->ChildAt(1)));
+				WorkData.VM->GetByteCode().AddArrayResetOp(ArrayOperand);
+				if (Settings.SetupNodeInstructionIndex)
+				{
+					WorkData.VM->GetByteCode().SetSubject(WorkData.VM->GetByteCode().GetNumInstructions() - 1, Callstack.GetCallPath(), Callstack.GetStack());
+				}
+				break;
+			}
+			case ERigVMOpCode::ArrayGetNum:
+			{
+				TraverseChildren(InExpr, WorkData);
+				const FRigVMOperand& ArrayOperand = WorkData.ExprToOperand.FindChecked(GetSourceVarExpr(InExpr->ChildAt(0)));
+				const FRigVMOperand& NumOperand = WorkData.ExprToOperand.FindChecked(GetSourceVarExpr(InExpr->ChildAt(1)));
+				WorkData.VM->GetByteCode().AddArrayGetNumOp(ArrayOperand, NumOperand);
+				if (Settings.SetupNodeInstructionIndex)
+				{
+					WorkData.VM->GetByteCode().SetSubject(WorkData.VM->GetByteCode().GetNumInstructions() - 1, Callstack.GetCallPath(), Callstack.GetStack());
+				}
+				break;
+			}
+			case ERigVMOpCode::ArraySetNum:
+			{
+				TraverseChildren(InExpr, WorkData);
+				const FRigVMOperand& ArrayOperand = WorkData.ExprToOperand.FindChecked(GetSourceVarExpr(InExpr->ChildAt(1)));
+				const FRigVMOperand& NumOperand = WorkData.ExprToOperand.FindChecked(GetSourceVarExpr(InExpr->ChildAt(2)));
+				WorkData.VM->GetByteCode().AddArraySetNumOp(ArrayOperand, NumOperand);
+				if (Settings.SetupNodeInstructionIndex)
+				{
+					WorkData.VM->GetByteCode().SetSubject(WorkData.VM->GetByteCode().GetNumInstructions() - 1, Callstack.GetCallPath(), Callstack.GetStack());
+				}
+				break;
+			}
+			case ERigVMOpCode::ArrayGetAtIndex:
+			{
+				TraverseChildren(InExpr, WorkData);
+				const FRigVMOperand& ArrayOperand = WorkData.ExprToOperand.FindChecked(GetSourceVarExpr(InExpr->ChildAt(0)));
+				const FRigVMOperand& IndexOperand = WorkData.ExprToOperand.FindChecked(GetSourceVarExpr(InExpr->ChildAt(1)));
+				const FRigVMOperand& ElementOperand = WorkData.ExprToOperand.FindChecked(GetSourceVarExpr(InExpr->ChildAt(2)));
+				WorkData.VM->GetByteCode().AddArrayGetAtIndexOp(ArrayOperand, IndexOperand, ElementOperand);
+				if (Settings.SetupNodeInstructionIndex)
+				{
+					WorkData.VM->GetByteCode().SetSubject(WorkData.VM->GetByteCode().GetNumInstructions() - 1, Callstack.GetCallPath(), Callstack.GetStack());
+				}
+				break;
+			}
+			case ERigVMOpCode::ArraySetAtIndex:
+			{
+				TraverseChildren(InExpr, WorkData);
+				const FRigVMOperand& ArrayOperand = WorkData.ExprToOperand.FindChecked(GetSourceVarExpr(InExpr->ChildAt(1)));
+				const FRigVMOperand& IndexOperand = WorkData.ExprToOperand.FindChecked(GetSourceVarExpr(InExpr->ChildAt(2)));
+				const FRigVMOperand& ElementOperand = WorkData.ExprToOperand.FindChecked(GetSourceVarExpr(InExpr->ChildAt(3)));
+				WorkData.VM->GetByteCode().AddArraySetAtIndexOp(ArrayOperand, IndexOperand, ElementOperand);
+				if (Settings.SetupNodeInstructionIndex)
+				{
+					WorkData.VM->GetByteCode().SetSubject(WorkData.VM->GetByteCode().GetNumInstructions() - 1, Callstack.GetCallPath(), Callstack.GetStack());
+				}
+				break;
+			}
+			case ERigVMOpCode::ArrayAdd:
+			{
+				TraverseChildren(InExpr, WorkData);
+				const FRigVMOperand& ArrayOperand = WorkData.ExprToOperand.FindChecked(GetSourceVarExpr(InExpr->ChildAt(1)));
+				const FRigVMOperand& ElementOperand = WorkData.ExprToOperand.FindChecked(GetSourceVarExpr(InExpr->ChildAt(2)));
+				const FRigVMOperand& IndexOperand = WorkData.ExprToOperand.FindChecked(GetSourceVarExpr(InExpr->ChildAt(3)));
+				WorkData.VM->GetByteCode().AddArrayAddOp(ArrayOperand, ElementOperand, IndexOperand);
+				if (Settings.SetupNodeInstructionIndex)
+				{
+					WorkData.VM->GetByteCode().SetSubject(WorkData.VM->GetByteCode().GetNumInstructions() - 1, Callstack.GetCallPath(), Callstack.GetStack());
+				}
+				break;
+			}
+			case ERigVMOpCode::ArrayInsert:
+			{
+				TraverseChildren(InExpr, WorkData);
+				const FRigVMOperand& ArrayOperand = WorkData.ExprToOperand.FindChecked(GetSourceVarExpr(InExpr->ChildAt(1)));
+				const FRigVMOperand& IndexOperand = WorkData.ExprToOperand.FindChecked(GetSourceVarExpr(InExpr->ChildAt(2)));
+				const FRigVMOperand& ElementOperand = WorkData.ExprToOperand.FindChecked(GetSourceVarExpr(InExpr->ChildAt(3)));
+				WorkData.VM->GetByteCode().AddArrayInsertOp(ArrayOperand, IndexOperand, ElementOperand);
+				if (Settings.SetupNodeInstructionIndex)
+				{
+					WorkData.VM->GetByteCode().SetSubject(WorkData.VM->GetByteCode().GetNumInstructions() - 1, Callstack.GetCallPath(), Callstack.GetStack());
+				}
+				break;
+			}
+			case ERigVMOpCode::ArrayRemove:
+			{
+				TraverseChildren(InExpr, WorkData);
+				const FRigVMOperand& ArrayOperand = WorkData.ExprToOperand.FindChecked(GetSourceVarExpr(InExpr->ChildAt(1)));
+				const FRigVMOperand& IndexOperand = WorkData.ExprToOperand.FindChecked(GetSourceVarExpr(InExpr->ChildAt(2)));
+				WorkData.VM->GetByteCode().AddArrayRemoveOp(ArrayOperand, IndexOperand);
+				if (Settings.SetupNodeInstructionIndex)
+				{
+					WorkData.VM->GetByteCode().SetSubject(WorkData.VM->GetByteCode().GetNumInstructions() - 1, Callstack.GetCallPath(), Callstack.GetStack());
+				}
+				break;
+			}
+			case ERigVMOpCode::ArrayFind:
+			{
+				TraverseChildren(InExpr, WorkData);
+				const FRigVMOperand& ArrayOperand = WorkData.ExprToOperand.FindChecked(GetSourceVarExpr(InExpr->ChildAt(0)));
+				const FRigVMOperand& ElementOperand = WorkData.ExprToOperand.FindChecked(GetSourceVarExpr(InExpr->ChildAt(1)));
+				const FRigVMOperand& IndexOperand = WorkData.ExprToOperand.FindChecked(GetSourceVarExpr(InExpr->ChildAt(2)));
+				const FRigVMOperand& SuccessOperand = WorkData.ExprToOperand.FindChecked(GetSourceVarExpr(InExpr->ChildAt(3)));
+				WorkData.VM->GetByteCode().AddArrayFindOp(ArrayOperand, ElementOperand, IndexOperand, SuccessOperand);
+				if (Settings.SetupNodeInstructionIndex)
+				{
+					WorkData.VM->GetByteCode().SetSubject(WorkData.VM->GetByteCode().GetNumInstructions() - 1, Callstack.GetCallPath(), Callstack.GetStack());
+				}
+				break;
+			}
+			case ERigVMOpCode::ArrayAppend:
+			{
+				TraverseChildren(InExpr, WorkData);
+				const FRigVMOperand& ArrayOperand = WorkData.ExprToOperand.FindChecked(GetSourceVarExpr(InExpr->ChildAt(1)));
+				const FRigVMOperand& OtherOperand = WorkData.ExprToOperand.FindChecked(GetSourceVarExpr(InExpr->ChildAt(2)));
+				WorkData.VM->GetByteCode().AddArrayAppendOp(ArrayOperand, OtherOperand);
+				if (Settings.SetupNodeInstructionIndex)
+				{
+					WorkData.VM->GetByteCode().SetSubject(WorkData.VM->GetByteCode().GetNumInstructions() - 1, Callstack.GetCallPath(), Callstack.GetStack());
+				}
+				break;
+			}
+			case ERigVMOpCode::ArrayClone:
+			{
+				TraverseChildren(InExpr, WorkData);
+				const FRigVMOperand& ArrayOperand = WorkData.ExprToOperand.FindChecked(GetSourceVarExpr(InExpr->ChildAt(0)));
+				const FRigVMOperand& CloneOperand = WorkData.ExprToOperand.FindChecked(GetSourceVarExpr(InExpr->ChildAt(1)));
+				WorkData.VM->GetByteCode().AddArrayCloneOp(ArrayOperand, CloneOperand);
+				if (Settings.SetupNodeInstructionIndex)
+				{
+					WorkData.VM->GetByteCode().SetSubject(WorkData.VM->GetByteCode().GetNumInstructions() - 1, Callstack.GetCallPath(), Callstack.GetStack());
+				}
+				break;
+			}
+			case ERigVMOpCode::ArrayIterator:
+			{
+				const FRigVMExprAST* ExecuteExpr = InExpr->ChildAt(0);
+				const FRigVMExprAST* ArrayExpr = InExpr->ChildAt(1);
+				const FRigVMExprAST* ElementExpr = InExpr->ChildAt(2);
+				const FRigVMExprAST* IndexExpr = InExpr->ChildAt(3);
+				const FRigVMExprAST* CountExpr = InExpr->ChildAt(4);
+				const FRigVMExprAST* RatioExpr = InExpr->ChildAt(5);
+				const FRigVMExprAST* ContinueExpr = InExpr->ChildAt(6);
+				const FRigVMExprAST* CompletedExpr = InExpr->ChildAt(7);
+				const FRigVMOperand& ExecuteOperand = WorkData.ExprToOperand.FindChecked(GetSourceVarExpr(ExecuteExpr));
+				const FRigVMOperand& ArrayOperand = WorkData.ExprToOperand.FindChecked(GetSourceVarExpr(ArrayExpr));
+				const FRigVMOperand& ElementOperand = WorkData.ExprToOperand.FindChecked(GetSourceVarExpr(ElementExpr));
+				const FRigVMOperand& IndexOperand = WorkData.ExprToOperand.FindChecked(GetSourceVarExpr(IndexExpr));
+				const FRigVMOperand& CountOperand = WorkData.ExprToOperand.FindChecked(GetSourceVarExpr(CountExpr));
+				const FRigVMOperand& RatioOperand = WorkData.ExprToOperand.FindChecked(GetSourceVarExpr(RatioExpr));
+				const FRigVMOperand& ContinueOperand = WorkData.ExprToOperand.FindChecked(GetSourceVarExpr(ContinueExpr));
+				const FRigVMOperand& CompletedOperand = WorkData.ExprToOperand.FindChecked(GetSourceVarExpr(CompletedExpr));
+
+				WorkData.ExprToSkip.AddUnique(ExecuteExpr);
+				WorkData.ExprToSkip.AddUnique(CompletedExpr);
+
+				// traverse the input array
+				TraverseExpression(ArrayExpr, WorkData);
+
+				// zero the index
+				WorkData.VM->GetByteCode().AddZeroOp(IndexOperand);
+				if (Settings.SetupNodeInstructionIndex)
+				{
+					WorkData.VM->GetByteCode().SetSubject(WorkData.VM->GetByteCode().GetNumInstructions() - 1, Callstack.GetCallPath(), Callstack.GetStack());
+				}
+
+				// add the iterator
+				WorkData.VM->GetByteCode().AddArrayIteratorOp(ArrayOperand, ElementOperand, IndexOperand, CountOperand, RatioOperand, ContinueOperand);
+				const int32 IteratorInstruction = WorkData.VM->GetByteCode().GetNumInstructions() - 1;
+				if (Settings.SetupNodeInstructionIndex)
+				{
+					WorkData.VM->GetByteCode().SetSubject(WorkData.VM->GetByteCode().GetNumInstructions() - 1, Callstack.GetCallPath(), Callstack.GetStack());
+				}
+
+				// jump to the end of the loop
+				const uint64 JumpToEndByte = WorkData.VM->GetByteCode().AddJumpIfOp(ERigVMOpCode::JumpForwardIf, 0, ContinueOperand, false);
+				const int32 JumpToEndInstruction = WorkData.VM->GetByteCode().GetNumInstructions() - 1;
+				if (Settings.SetupNodeInstructionIndex)
+				{
+					WorkData.VM->GetByteCode().SetSubject(WorkData.VM->GetByteCode().GetNumInstructions() - 1, Callstack.GetCallPath(), Callstack.GetStack());
+				}
+
+				// begin the block
+				WorkData.VM->GetByteCode().AddBeginBlockOp(CountOperand, IndexOperand);
+				if (Settings.SetupNodeInstructionIndex)
+				{
+					WorkData.VM->GetByteCode().SetSubject(WorkData.VM->GetByteCode().GetNumInstructions() - 1, Callstack.GetCallPath(), Callstack.GetStack());
+				}
+
+				// traverse the per iteration instructions
+				WorkData.ExprToSkip.Remove(ExecuteExpr);
+				TraverseExpression(ExecuteExpr, WorkData);
+
+				// end the block
+				WorkData.VM->GetByteCode().AddEndBlockOp();
+				if (Settings.SetupNodeInstructionIndex)
+				{
+					WorkData.VM->GetByteCode().SetSubject(WorkData.VM->GetByteCode().GetNumInstructions() - 1, Callstack.GetCallPath(), Callstack.GetStack());
+				}
+
+				// increment index per loop iteration
+				WorkData.VM->GetByteCode().AddIncrementOp(IndexOperand);
+				if (Settings.SetupNodeInstructionIndex)
+				{
+					WorkData.VM->GetByteCode().SetSubject(WorkData.VM->GetByteCode().GetNumInstructions() - 1, Callstack.GetCallPath(), Callstack.GetStack());
+
+				}
+
+				// jump backwards instruction (to the beginning of the iterator)
+				const int32 JumpToStartInstruction = WorkData.VM->GetByteCode().GetNumInstructions();
+				WorkData.VM->GetByteCode().AddJumpOp(ERigVMOpCode::JumpBackward, JumpToStartInstruction - IteratorInstruction);
+				if (Settings.SetupNodeInstructionIndex)
+				{
+					WorkData.VM->GetByteCode().SetSubject(WorkData.VM->GetByteCode().GetNumInstructions() - 1, Callstack.GetCallPath(), Callstack.GetStack());
+				}
+
+				// fix up the first jump instruction
+				const int32 InstructionsToEnd = WorkData.VM->GetByteCode().GetNumInstructions() - JumpToEndInstruction;
+				WorkData.VM->GetByteCode().GetOpAt<FRigVMJumpIfOp>(JumpToEndByte).InstructionIndex = InstructionsToEnd;
+				if (Settings.SetupNodeInstructionIndex)
+				{
+					WorkData.VM->GetByteCode().SetSubject(WorkData.VM->GetByteCode().GetNumInstructions() - 1, Callstack.GetCallPath(), Callstack.GetStack());
+				}
+					
+				WorkData.ExprToSkip.Remove(CompletedExpr);
+				TraverseExpression(CompletedExpr, WorkData);
+				break;
+			}
+			case ERigVMOpCode::ArrayUnion:
+			{
+				TraverseChildren(InExpr, WorkData);
+				const FRigVMOperand& ArrayOperand = WorkData.ExprToOperand.FindChecked(GetSourceVarExpr(InExpr->ChildAt(1)));
+				const FRigVMOperand& OtherOperand = WorkData.ExprToOperand.FindChecked(GetSourceVarExpr(InExpr->ChildAt(2)));
+				WorkData.VM->GetByteCode().AddArrayUnionOp(ArrayOperand, OtherOperand);
+				if (Settings.SetupNodeInstructionIndex)
+				{
+					WorkData.VM->GetByteCode().SetSubject(WorkData.VM->GetByteCode().GetNumInstructions() - 1, Callstack.GetCallPath(), Callstack.GetStack());
+				}
+				break;
+			}
+			case ERigVMOpCode::ArrayDifference:
+			{
+				TraverseChildren(InExpr, WorkData);
+				const FRigVMOperand& ArrayOperand = WorkData.ExprToOperand.FindChecked(GetSourceVarExpr(InExpr->ChildAt(0)));
+				const FRigVMOperand& OtherOperand = WorkData.ExprToOperand.FindChecked(GetSourceVarExpr(InExpr->ChildAt(1)));
+				const FRigVMOperand& ResultOperand = WorkData.ExprToOperand.FindChecked(GetSourceVarExpr(InExpr->ChildAt(2)));
+				WorkData.VM->GetByteCode().AddArrayDifferenceOp(ArrayOperand, OtherOperand, ResultOperand);
+				if (Settings.SetupNodeInstructionIndex)
+				{
+					WorkData.VM->GetByteCode().SetSubject(WorkData.VM->GetByteCode().GetNumInstructions() - 1, Callstack.GetCallPath(), Callstack.GetStack());
+				}
+				break;
+			}
+			case ERigVMOpCode::ArrayIntersection:
+			{
+				TraverseChildren(InExpr, WorkData);
+				const FRigVMOperand& ArrayOperand = WorkData.ExprToOperand.FindChecked(GetSourceVarExpr(InExpr->ChildAt(0)));
+				const FRigVMOperand& OtherOperand = WorkData.ExprToOperand.FindChecked(GetSourceVarExpr(InExpr->ChildAt(1)));
+				const FRigVMOperand& ResultOperand = WorkData.ExprToOperand.FindChecked(GetSourceVarExpr(InExpr->ChildAt(2)));
+				WorkData.VM->GetByteCode().AddArrayIntersectionOp(ArrayOperand, OtherOperand, ResultOperand);
+				if (Settings.SetupNodeInstructionIndex)
+				{
+					WorkData.VM->GetByteCode().SetSubject(WorkData.VM->GetByteCode().GetNumInstructions() - 1, Callstack.GetCallPath(), Callstack.GetStack());
+				}
+				break;
+			}
+			case ERigVMOpCode::ArrayReverse:
+			{
+				TraverseChildren(InExpr, WorkData);
+				const FRigVMOperand& ArrayOperand = WorkData.ExprToOperand.FindChecked(GetSourceVarExpr(InExpr->ChildAt(1)));
+				WorkData.VM->GetByteCode().AddArrayReverseOp(ArrayOperand);
+				if (Settings.SetupNodeInstructionIndex)
+				{
+					WorkData.VM->GetByteCode().SetSubject(WorkData.VM->GetByteCode().GetNumInstructions() - 1, Callstack.GetCallPath(), Callstack.GetStack());
+				}
+				break;
+			}
+			default:
+			{
+				checkNoEntry();
+				break;
+			}
+		}
+	}
+}
+
 void URigVMCompiler::InitializeLocalVariables(const FRigVMExprAST* InExpr, FRigVMCompilerWorkData& WorkData)
 {
 	// Initialize local variables if we are entering a new graph
@@ -1342,6 +1678,11 @@ void URigVMCompiler::InitializeLocalVariables(const FRigVMExprAST* InExpr, FRigV
 				Callstack = InExpr->To<FRigVMSelectExprAST>()->GetProxy().GetCallstack();
 				break;
 			}
+			case FRigVMExprAST::EType::Array:
+			{
+				Callstack = InExpr->To<FRigVMArrayExprAST>()->GetProxy().GetCallstack();
+				break;
+			}
 		}
 
 		if (Callstack.Num() > 0)
@@ -1475,13 +1816,22 @@ FString URigVMCompiler::GetPinHash(const URigVMPin* InPin, const FRigVMVarExprAS
 
 	if (InVarExpr != nullptr && !bIsDebugValue)
 	{
-		if (InVarExpr->IsA(FRigVMExprAST::ExternalVar) && !bIsDebugValue)
+		if (InVarExpr->IsA(FRigVMExprAST::ExternalVar))
 		{
 			URigVMPin::FPinOverride PinOverride(InVarExpr->GetProxy(), InVarExpr->GetParser()->GetPinOverrides());
 			FString VariablePath = InPin->GetBoundVariablePath(PinOverride);
 			return FString::Printf(TEXT("%sVariable::%s%s"), *Prefix, *VariablePath, *Suffix);
 		}
 
+		// for IO array pins we'll walk left and use that pin hash instead
+		if(const FRigVMVarExprAST* SourceVarExpr = GetSourceVarExpr(InVarExpr))
+		{
+			if(SourceVarExpr != InVarExpr)
+			{
+				return GetPinHash(SourceVarExpr->GetPin(), SourceVarExpr, bIsDebugValue);
+			}
+		}
+
 		bIsExecutePin = InPin->IsExecuteContext();
 		bIsLiteral = InVarExpr->GetType() == FRigVMExprAST::EType::Literal;
 
@@ -1601,6 +1951,56 @@ FString URigVMCompiler::GetPinHash(const URigVMPin* InPin, const FRigVMVarExprAS
 	return FString::Printf(TEXT("%s%s%s"), *Prefix, *InPin->GetPinPath(bUseFullNodePath), *Suffix);
 }
 
+const FRigVMVarExprAST* URigVMCompiler::GetSourceVarExpr(const FRigVMExprAST* InExpr)
+{
+	if(InExpr)
+	{
+		if(InExpr->IsA(FRigVMExprAST::EType::CachedValue))
+		{
+			return GetSourceVarExpr(InExpr->To<FRigVMCachedValueExprAST>()->GetVarExpr());
+		}
+
+		if(InExpr->IsA(FRigVMExprAST::EType::Var))
+		{
+			const FRigVMVarExprAST* VarExpr = InExpr->To<FRigVMVarExprAST>();
+			
+			if(VarExpr->GetPin()->IsReferenceCountedContainer() &&
+				((VarExpr->GetPin()->GetDirection() == ERigVMPinDirection::Input) || (VarExpr->GetPin()->GetDirection() == ERigVMPinDirection::IO)))
+			{
+				// if this is a variable setter we cannot follow the source var
+				if(VarExpr->GetPin()->GetDirection() == ERigVMPinDirection::Input)
+				{
+					if(VarExpr->GetPin()->GetNode()->IsA<URigVMVariableNode>())
+					{
+						return VarExpr;
+					}
+				}
+				
+				if(const FRigVMExprAST* AssignExpr = VarExpr->GetFirstChildOfType(FRigVMExprAST::EType::Assign))
+				{
+					// don't follow a copy assignment
+					if(AssignExpr->IsA(FRigVMExprAST::EType::Copy))
+					{
+						return VarExpr;
+					}
+					
+					if(const FRigVMExprAST* CachedValueExpr = VarExpr->GetFirstChildOfType(FRigVMExprAST::EType::CachedValue))
+					{
+						return GetSourceVarExpr(CachedValueExpr->To<FRigVMCachedValueExprAST>()->GetVarExpr());
+					}
+					else if(const FRigVMExprAST* ChildExpr = VarExpr->GetFirstChildOfType(FRigVMExprAST::EType::Var))
+					{
+						return GetSourceVarExpr(ChildExpr->To<FRigVMVarExprAST>());
+					}
+				}
+			}
+			return VarExpr;
+		}
+	}
+
+	return nullptr;
+}
+
 #if UE_RIGVM_UCLASS_BASED_STORAGE_DISABLED
 
 void URigVMCompiler::CreateDebugRegister(URigVMPin* InPin, URigVM* OutVM,
@@ -1867,6 +2267,11 @@ uint16 URigVMCompiler::GetElementSizeFromCPPType(const FString& InCPPType, UScri
 
 FRigVMOperand URigVMCompiler::FindOrAddRegister(const FRigVMVarExprAST* InVarExpr, FRigVMCompilerWorkData& WorkData, bool bIsDebugValue)
 {
+	if(!bIsDebugValue)
+	{
+		InVarExpr = GetSourceVarExpr(InVarExpr);
+	}
+	
 	if (!bIsDebugValue)
 	{
 		FRigVMOperand const* ExistingOperand = WorkData.ExprToOperand.Find(InVarExpr);
@@ -2347,6 +2752,7 @@ FRigVMOperand URigVMCompiler::FindOrAddRegister(const FRigVMVarExprAST* InVarExp
 		{
 			JoinedDefaultValue.Empty();
 		}
+	
 		Operand = WorkData.AddProperty(MemoryType, RegisterName, CPPType, Pin->GetCPPTypeObject(), JoinedDefaultValue);
 		
 #endif
diff --git a/Engine/Source/Developer/RigVMDeveloper/Private/RigVMModel/Nodes/RigVMArrayNode.cpp b/Engine/Source/Developer/RigVMDeveloper/Private/RigVMModel/Nodes/RigVMArrayNode.cpp
new file mode 100644
--- /dev/null
+++ b/Engine/Source/Developer/RigVMDeveloper/Private/RigVMModel/Nodes/RigVMArrayNode.cpp
@@ -0,0 +1,412 @@
+// Copyright Epic Games, Inc. All Rights Reserved.
+
+#include "RigVMModel/Nodes/RigVMArrayNode.h"
+#include "RigVMModel/RigVMGraph.h"
+
+const FString URigVMArrayNode::ArrayName = TEXT("Array");
+const FString URigVMArrayNode::NumName = TEXT("Num");
+const FString URigVMArrayNode::IndexName = TEXT("Index");
+const FString URigVMArrayNode::ElementName = TEXT("Element");
+const FString URigVMArrayNode::SuccessName = TEXT("Success");
+const FString URigVMArrayNode::OtherName = TEXT("Other");
+const FString URigVMArrayNode::CloneName = TEXT("Clone");
+const FString URigVMArrayNode::CountName = TEXT("Count");
+const FString URigVMArrayNode::RatioName = TEXT("Ratio");
+const FString URigVMArrayNode::ResultName = TEXT("Result");
+const FString URigVMArrayNode::ContinueName = TEXT("Continue");
+const FString URigVMArrayNode::CompletedName = TEXT("Completed");
+
+URigVMArrayNode::URigVMArrayNode()
+: OpCode(ERigVMOpCode::Invalid)
+{
+}
+
+FString URigVMArrayNode::GetNodeTitle() const
+{
+	return GetNodeTitle(OpCode);
+}
+
+FString URigVMArrayNode::GetNodeTitle(ERigVMOpCode InOpCode)
+{
+	if(InOpCode == ERigVMOpCode::ArrayGetAtIndex)
+	{
+		static const FString ArrayGetAtIndexNodeTitle = TEXT("At"); 
+		return ArrayGetAtIndexNodeTitle;
+	}
+	else if(InOpCode == ERigVMOpCode::ArraySetAtIndex)
+	{
+		static const FString ArraySetAtIndexNodeTitle = TEXT("Set At"); 
+		return ArraySetAtIndexNodeTitle;
+	}
+	else if(InOpCode == ERigVMOpCode::ArrayGetNum)
+	{
+		static const FString ArrayGetNumNodeTitle = TEXT("Num"); 
+		return ArrayGetNumNodeTitle;
+	}
+	else if(InOpCode == ERigVMOpCode::ArrayIterator)
+	{
+		static const FString ArrayIteratorNodeTitle = TEXT("For Each"); 
+		return ArrayIteratorNodeTitle;
+	}
+
+	FString OpCodeString = StaticEnum<ERigVMOpCode>()->GetDisplayNameTextByValue((int64)InOpCode).ToString();
+	OpCodeString.RemoveFromStart(URigVMArrayNode::ArrayName);
+	return OpCodeString;
+}
+
+FText URigVMArrayNode::GetToolTipText() const
+{
+	switch(OpCode)
+	{
+		case ERigVMOpCode::ArrayReset:
+		{
+			return FText::FromString(TEXT("Removes all elements from an array.\nModifies the input array."));
+		}
+		case ERigVMOpCode::ArrayGetNum:
+		{
+			return FText::FromString(TEXT("Returns the number of elements of an array"));
+		} 
+		case ERigVMOpCode::ArraySetNum:
+		{
+			return FText::FromString(TEXT("Sets the numbers of elements of an array.\nModifies the input array."));
+		}
+		case ERigVMOpCode::ArrayGetAtIndex:
+		{
+			return FText::FromString(TEXT("Returns an element of an array by index."));
+		}  
+		case ERigVMOpCode::ArraySetAtIndex:
+		{
+			return FText::FromString(TEXT("Sets an element of an array by index.\nModifies the input array."));
+		}
+		case ERigVMOpCode::ArrayAdd:
+		{
+			return FText::FromString(TEXT("Adds an element to an array and returns the new element's index.\nModifies the input array."));
+		}
+		case ERigVMOpCode::ArrayInsert:
+		{
+			return FText::FromString(TEXT("Inserts an element into an array at a given index.\nModifies the input array."));
+		}
+		case ERigVMOpCode::ArrayRemove:
+		{
+			return FText::FromString(TEXT("Removes an element from an array by index.\nModifies the input array."));
+		}
+		case ERigVMOpCode::ArrayFind:
+		{
+			return FText::FromString(TEXT("Searchs a potential element in an array and returns its index."));
+		}
+		case ERigVMOpCode::ArrayAppend:
+		{
+			return FText::FromString(TEXT("Appends the another array to the main one.\nModifies the input array."));
+		}
+		case ERigVMOpCode::ArrayClone:
+		{
+			return FText::FromString(TEXT("Clones an array and returns a duplicate."));
+		}
+		case ERigVMOpCode::ArrayIterator:
+		{
+			return FText::FromString(TEXT("Loops over the elements in an array."));
+		}
+		case ERigVMOpCode::ArrayUnion:
+		{
+			return FText::FromString(TEXT("Merges two arrays while ensuring unique elements.\nModifies the input array."));
+		}
+		case ERigVMOpCode::ArrayDifference:
+		{
+			return FText::FromString(TEXT("Creates a new array containing the difference between the two input arrays.\nDifference here means elements that are only contained in either A or B."));
+		}
+		case ERigVMOpCode::ArrayIntersection:
+		{
+			return FText::FromString(TEXT("Creates a new array containing the intersection between the two input arrays.\nDifference here means elements that are contained in both A and B."));
+		}
+		case ERigVMOpCode::ArrayReverse:
+		{
+			return FText::FromString(TEXT("Reverses the order of the elements of an array.\nModifies the input array."));
+		}
+		default:
+		{
+			ensure(false);
+			break;
+		}
+	}
+	return Super::GetToolTipText();
+}
+
+bool URigVMArrayNode::IsLoopNode() const
+{
+	return GetOpCode() == ERigVMOpCode::ArrayIterator;
+}
+
+FText URigVMArrayNode::GetToolTipTextForPin(const URigVMPin* InPin) const
+{
+	switch(OpCode)
+	{
+		case ERigVMOpCode::ArrayReset:
+		{
+			if(InPin->GetName() == URigVMArrayNode::ArrayName)
+			{
+				return FText::FromString(TEXT("The array to be cleared."));
+			}
+			break;
+		}
+		case ERigVMOpCode::ArrayGetNum:
+		{
+			if(InPin->GetName() == URigVMArrayNode::ArrayName)
+			{
+				return FText::FromString(TEXT("The array to retrieve the size for."));
+			}
+			else if(InPin->GetName() == URigVMArrayNode::NumName)
+			{
+				return FText::FromString(TEXT("The size of the input array."));
+			}
+			break;
+		} 
+		case ERigVMOpCode::ArraySetNum:
+		{
+			if(InPin->GetName() == URigVMArrayNode::ArrayName)
+			{
+				return FText::FromString(TEXT("The array to set the size for."));
+			}
+			else if(InPin->GetName() == URigVMArrayNode::NumName)
+			{
+				return FText::FromString(TEXT("The new size of the array."));
+			}
+			break;
+		}
+		case ERigVMOpCode::ArrayGetAtIndex:
+		{
+			if(InPin->GetName() == URigVMArrayNode::ArrayName)
+			{
+				return FText::FromString(TEXT("The array to retrieve an element from."));
+			}
+			else if(InPin->GetName() == URigVMArrayNode::IndexName)
+			{
+				return FText::FromString(TEXT("The index of the element to retrieve."));
+			}
+			else if(InPin->GetName() == URigVMArrayNode::ElementName)
+			{
+				return FText::FromString(TEXT("The element at the given index."));
+			}
+			break;
+		}  
+		case ERigVMOpCode::ArraySetAtIndex:
+		{
+			if(InPin->GetName() == URigVMArrayNode::ArrayName)
+			{
+				return FText::FromString(TEXT("The array to set an element for."));
+			}
+			else if(InPin->GetName() == URigVMArrayNode::IndexName)
+			{
+				return FText::FromString(TEXT("The index of the element to set."));
+			}
+			else if(InPin->GetName() == URigVMArrayNode::ElementName)
+			{
+				return FText::FromString(TEXT("The new value for element to set."));
+			}
+			break;
+		}
+		case ERigVMOpCode::ArrayAdd:
+		{
+			if(InPin->GetName() == URigVMArrayNode::ArrayName)
+			{
+				return FText::FromString(TEXT("The array to add an element to."));
+			}
+			else if(InPin->GetName() == URigVMArrayNode::ElementName)
+			{
+				return FText::FromString(TEXT("The element to add to the array."));
+			}
+			else if(InPin->GetName() == URigVMArrayNode::IndexName)
+			{
+				return FText::FromString(TEXT("The index of the newly added element."));
+			}
+			break;
+		}
+		case ERigVMOpCode::ArrayInsert:
+		{
+			if(InPin->GetName() == URigVMArrayNode::ArrayName)
+			{
+				return FText::FromString(TEXT("The array to insert an element into."));
+			}
+			else if(InPin->GetName() == URigVMArrayNode::IndexName)
+			{
+				return FText::FromString(TEXT("The index at which to insert the element."));
+			}
+			else if(InPin->GetName() == URigVMArrayNode::ElementName)
+			{
+				return FText::FromString(TEXT("The element to insert into the array."));
+			}
+			break;
+		}
+		case ERigVMOpCode::ArrayRemove:
+		{
+			if(InPin->GetName() == URigVMArrayNode::ArrayName)
+			{
+				return FText::FromString(TEXT("The array to remove an element from."));
+			}
+			else if(InPin->GetName() == URigVMArrayNode::IndexName)
+			{
+				return FText::FromString(TEXT("The index at which to remove the element."));
+			}
+			break;
+		}
+		case ERigVMOpCode::ArrayFind:
+		{
+			if(InPin->GetName() == URigVMArrayNode::ArrayName)
+			{
+				return FText::FromString(TEXT("The array to search within."));
+			}
+			else if(InPin->GetName() == URigVMArrayNode::ElementName)
+			{
+				return FText::FromString(TEXT("The element to look for."));
+			}
+			else if(InPin->GetName() == URigVMArrayNode::IndexName)
+			{
+				return FText::FromString(TEXT("The index of the found element (or -1)."));
+			}
+			else if(InPin->GetName() == URigVMArrayNode::SuccessName)
+			{
+				return FText::FromString(TEXT("True if the element has been found."));
+			}
+			break;
+		}
+		case ERigVMOpCode::ArrayAppend:
+		{
+			if(InPin->GetName() == URigVMArrayNode::ArrayName)
+			{
+				return FText::FromString(TEXT("The array to append the other array to."));
+			}
+			else if(InPin->GetName() == URigVMArrayNode::OtherName)
+			{
+				return FText::FromString(TEXT("The second array to append to the first one."));
+			}
+			break;
+		}
+		case ERigVMOpCode::ArrayClone:
+		{
+			if(InPin->GetName() == URigVMArrayNode::ArrayName)
+			{
+				return FText::FromString(TEXT("The array to clone"));
+			}
+			else if(InPin->GetName() == URigVMArrayNode::CloneName)
+			{
+				return FText::FromString(TEXT("The duplicate of the input array."));
+			}
+			break;
+		}
+		case ERigVMOpCode::ArrayIterator:
+		{
+			if(InPin->GetName() == URigVMArrayNode::ArrayName)
+			{
+				return FText::FromString(TEXT("The array to iterate over"));
+			}
+			else if(InPin->GetName() == URigVMArrayNode::IndexName)
+			{
+				return FText::FromString(TEXT("The index of the current element."));
+			}
+			else if(InPin->GetName() == URigVMArrayNode::CountName)
+			{
+				return FText::FromString(TEXT("The count of all elements in the array."));
+			}
+			else if(InPin->GetName() == URigVMArrayNode::ElementName)
+			{
+				return FText::FromString(TEXT("The current element."));
+			}
+			else if(InPin->GetName() == URigVMArrayNode::RatioName)
+			{
+				return FText::FromString(TEXT("A float ratio from 0.0 (first element) to 1.0 (last element)."));
+			}
+			else if(InPin->GetName() == URigVMArrayNode::CompletedName)
+			{
+				return FText::FromString(TEXT("The execute block to run once the loop has completed."));
+			}
+			break;
+		}
+		case ERigVMOpCode::ArrayUnion:
+		{
+			if(InPin->GetName() == URigVMArrayNode::ArrayName)
+			{
+				return FText::FromString(TEXT("The array to merge the other array with."));
+			}
+			else if(InPin->GetName() == URigVMArrayNode::OtherName)
+			{
+				return FText::FromString(TEXT("The second array to merge to the first one."));
+			}
+			break;
+		}
+		case ERigVMOpCode::ArrayDifference:
+		{
+			if(InPin->GetName() == URigVMArrayNode::ArrayName)
+			{
+				return FText::FromString(TEXT("The first array to compare the other array with."));
+			}
+			else if(InPin->GetName() == URigVMArrayNode::OtherName)
+			{
+				return FText::FromString(TEXT("The second array to compare the other array with."));
+			}
+			else if(InPin->GetName() == URigVMArrayNode::ResultName)
+			{
+				return FText::FromString(TEXT("The resulting difference array."));
+			}
+			break;
+		}
+		case ERigVMOpCode::ArrayIntersection:
+		{
+			if(InPin->GetName() == URigVMArrayNode::ArrayName)
+			{
+				return FText::FromString(TEXT("The first array to compare the other array with."));
+			}
+			else if(InPin->GetName() == URigVMArrayNode::OtherName)
+			{
+				return FText::FromString(TEXT("The second array to compare the other array with."));
+			}
+			else if(InPin->GetName() == URigVMArrayNode::ResultName)
+			{
+				return FText::FromString(TEXT("The resulting intersecting array."));
+			}
+			break;
+		}
+		case ERigVMOpCode::ArrayReverse:
+		{
+			if(InPin->GetName() == URigVMArrayNode::ArrayName)
+			{
+				return FText::FromString(TEXT("The array reverse."));
+			}
+			break;
+		}
+		default:
+		{
+			ensure(false);
+			break;
+		}
+	}
+	return Super::GetToolTipTextForPin(InPin);
+}
+
+FLinearColor URigVMArrayNode::GetNodeColor() const
+{
+	return FLinearColor::White;
+}
+
+ERigVMOpCode URigVMArrayNode::GetOpCode() const
+{
+	return OpCode;
+}
+
+FString URigVMArrayNode::GetCPPType() const
+{
+	URigVMPin* ArrayPin = FindPin(ArrayName);
+	if (ArrayPin == nullptr)
+	{
+		return FString();
+	}
+	return ArrayPin->GetCPPType();
+}
+
+UObject* URigVMArrayNode::GetCPPTypeObject() const
+{
+	URigVMPin* ArrayPin = FindPin(ArrayName);
+	if (ArrayPin == nullptr)
+	{
+		return nullptr;
+	}
+	return ArrayPin->GetCPPTypeObject();
+}
diff --git a/Engine/Source/Developer/RigVMDeveloper/Private/RigVMModel/RigVMController.cpp b/Engine/Source/Developer/RigVMDeveloper/Private/RigVMModel/RigVMController.cpp
--- a/Engine/Source/Developer/RigVMDeveloper/Private/RigVMModel/RigVMController.cpp
+++ b/Engine/Source/Developer/RigVMDeveloper/Private/RigVMModel/RigVMController.cpp
@@ -7,6 +7,7 @@
 #include "RigVMModel/Nodes/RigVMFunctionReferenceNode.h"
 #include "RigVMCore/RigVMRegistry.h"
 #include "RigVMCore/RigVMExecuteContext.h"
+#include "RigVMCore/RigVMByteCode.h"
 #include "RigVMCompiler/RigVMCompiler.h"
 #include "RigVMDeveloperModule.h"
 #include "UObject/PropertyPortFlags.h"
@@ -581,6 +582,33 @@ TArray<FString> URigVMController::GetAddNodePythonCommands(URigVMNode* Node) con
 	{
 		
 		
+	}
+	else if (URigVMArrayNode* ArrayNode = Cast<URigVMArrayNode>(Node))
+	{
+		FString OpCodeString = StaticEnum<ERigVMOpCode>()->GetNameStringByValue((int64)ArrayNode->GetOpCode());
+		OpCodeString = RigVMPythonUtils::NameToPep8(OpCodeString);
+		OpCodeString.ToUpperInline();
+		
+		// add_array_node(opcode, cpp_type, cpp_type_object, position=[0.0, 0.0], node_name='', undo=True)
+		if (ArrayNode->GetCPPTypeObject())
+		{
+			Commands.Add(FString::Printf(TEXT("blueprint.get_controller_by_name('%s').add_array_node_from_object_path(unreal.RigVMOpCode.%s, '%s', '%s', %s, '%s')"),
+					*GetGraph()->GetGraphName(),
+					*OpCodeString,
+					*ArrayNode->GetCPPType(),
+					*ArrayNode->GetCPPTypeObject()->GetPathName(),
+					*RigVMPythonUtils::Vector2DToPythonString(ArrayNode->GetPosition()),
+					*ArrayNode->GetName()));	
+		}
+		else
+		{
+			Commands.Add(FString::Printf(TEXT("blueprint.get_controller_by_name('%s').add_array_node(unreal.RigVMOpCode.%s, '%s', None, %s, '%s')"),
+					*GetGraph()->GetGraphName(),
+					*OpCodeString,
+					*ArrayNode->GetCPPType(),
+					*RigVMPythonUtils::Vector2DToPythonString(ArrayNode->GetPosition()),
+					*ArrayNode->GetName()));	
+		}
 	}
 	else
 	{
@@ -591,7 +619,7 @@ TArray<FString> URigVMController::GetAddNodePythonCommands(URigVMNode* Node) con
 	{
 		for (const URigVMPin* Pin : Node->GetPins())
 		{
-			if (Pin->GetDirection() == ERigVMPinDirection::Output)
+			if (Pin->GetDirection() == ERigVMPinDirection::Output || Pin->GetDirection() == ERigVMPinDirection::Hidden)
 			{
 				continue;
 			}
@@ -853,17 +881,33 @@ URigVMVariableNode* URigVMController::AddVariableNode(const FName& InVariableNam
 	Node->Pins.Add(VariablePin);
 
 	URigVMPin* ValuePin = NewObject<URigVMPin>(Node, *URigVMVariableNode::ValueName);
-	ValuePin->CPPType = CPPType;
 
-	if (UScriptStruct* ScriptStruct = Cast<UScriptStruct>(InCPPTypeObject))
+	FRigVMExternalVariable ExternalVariable = GetVariableByName(InVariableName);
+	if(ExternalVariable.IsValid(true))
 	{
-		ValuePin->CPPTypeObject = ScriptStruct;
-		ValuePin->CPPTypeObjectPath = *ValuePin->CPPTypeObject->GetPathName();
+		ValuePin->CPPType = ExternalVariable.TypeName.ToString();
+		ValuePin->CPPTypeObject = ExternalVariable.TypeObject;
+		ValuePin->bIsDynamicArray = ExternalVariable.bIsArray;
+
+		if(ValuePin->bIsDynamicArray && !ValuePin->CPPType.StartsWith(TEXT("TArray<")))
+		{
+			ValuePin->CPPType = FString::Printf(TEXT("TArray<%s>"), *ValuePin->CPPType);
+		}
 	}
-	else if (UEnum* Enum = Cast<UEnum>(InCPPTypeObject))
+	else
 	{
-		ValuePin->CPPTypeObject = Enum;
-		ValuePin->CPPTypeObjectPath = *ValuePin->CPPTypeObject->GetPathName();
+		ValuePin->CPPType = CPPType;
+
+		if (UScriptStruct* ScriptStruct = Cast<UScriptStruct>(InCPPTypeObject))
+		{
+			ValuePin->CPPTypeObject = ScriptStruct;
+			ValuePin->CPPTypeObjectPath = *ValuePin->CPPTypeObject->GetPathName();
+		}
+		else if (UEnum* Enum = Cast<UEnum>(InCPPTypeObject))
+		{
+			ValuePin->CPPTypeObject = Enum;
+			ValuePin->CPPTypeObjectPath = *ValuePin->CPPTypeObject->GetPathName();
+		}
 	}
 
 	ValuePin->Direction = bIsGetter ? ERigVMPinDirection::Output : ERigVMPinDirection::Input;
@@ -8744,6 +8788,329 @@ URigVMEnumNode* URigVMController::AddEnumNode(const FName& InCPPTypeObjectPath,
 	return Node;
 }
 
+URigVMArrayNode* URigVMController::AddArrayNode(ERigVMOpCode InOpCode, const FString& InCPPType,
+	UObject* InCPPTypeObject, const FVector2D& InPosition, const FString& InNodeName, bool bSetupUndoRedo,
+	bool bPrintPythonCommand)
+{
+	if (!IsValidGraph())
+	{
+		return nullptr;
+	}
+
+#if UE_RIGVM_UCLASS_BASED_STORAGE_DISABLED
+
+	ReportError(TEXT("This build of Control Rig doesn't support Array Nodes."));
+	return nullptr;
+	
+#endif
+
+	// validate the op code
+	bool bIsMutable = false;
+	switch(InOpCode)
+	{
+		case ERigVMOpCode::ArrayReset:
+		case ERigVMOpCode::ArrayGetNum: 
+		case ERigVMOpCode::ArraySetNum:
+		case ERigVMOpCode::ArrayGetAtIndex:  
+		case ERigVMOpCode::ArraySetAtIndex:
+		case ERigVMOpCode::ArrayAdd:
+		case ERigVMOpCode::ArrayInsert:
+		case ERigVMOpCode::ArrayRemove:
+		case ERigVMOpCode::ArrayFind:
+		case ERigVMOpCode::ArrayAppend:
+		case ERigVMOpCode::ArrayClone:
+		case ERigVMOpCode::ArrayIterator:
+		case ERigVMOpCode::ArrayUnion:
+		case ERigVMOpCode::ArrayDifference:
+		case ERigVMOpCode::ArrayIntersection:
+		case ERigVMOpCode::ArrayReverse:
+		{
+			break;
+		}
+		default:
+		{
+			ReportErrorf(TEXT("OpCode '%s' is not valid for Array Node."), *StaticEnum<ERigVMOpCode>()->GetNameStringByValue((int64)InOpCode));
+			return nullptr;
+		}
+	}
+
+	URigVMGraph* Graph = GetGraph();
+	check(Graph);
+
+	if (Graph->IsA<URigVMFunctionLibrary>())
+	{
+		ReportError(TEXT("Cannot add array nodes to function library graphs."));
+		return nullptr;
+	}
+
+	if (InCPPTypeObject == nullptr)
+	{
+		InCPPTypeObject = URigVMCompiler::GetScriptStructForCPPType(InCPPType);
+	}
+	if (InCPPTypeObject == nullptr)
+	{
+		InCPPTypeObject = URigVMPin::FindObjectFromCPPTypeObjectPath<UObject>(InCPPType);
+	}
+
+	FString CPPType = InCPPType;
+	if (UScriptStruct* ScriptStruct = Cast<UScriptStruct>(InCPPTypeObject))
+	{
+		CPPType = ScriptStruct->GetStructCPPName();
+	}
+	
+	const FString Name = GetValidNodeName(InNodeName.IsEmpty() ? FString(TEXT("ArrayNode")) : InNodeName);
+	URigVMArrayNode* Node = NewObject<URigVMArrayNode>(Graph, *Name);
+	Node->Position = InPosition;
+	Node->OpCode = InOpCode;
+
+	struct Local
+	{
+		static URigVMPin* AddPin(URigVMController* InController, URigVMNode* InNode, const FName& InName, ERigVMPinDirection InDirection, bool bIsArray, const FString& InCPPType, UObject* InCPPTypeObject)
+		{
+			URigVMPin* Pin = NewObject<URigVMPin>(InNode, InName);
+			Pin->CPPType = InCPPType;
+			Pin->CPPTypeObject = InCPPTypeObject;
+			if(Pin->CPPTypeObject)
+			{
+				Pin->CPPTypeObjectPath = *Pin->CPPTypeObject->GetPathName();
+			}
+			if(UScriptStruct* Struct = Cast<UScriptStruct>(Pin->CPPTypeObject))
+			{
+				Pin->CPPType = Struct->GetStructCPPName();
+			}
+			if(bIsArray && !Pin->CPPType.StartsWith(TEXT("TArray<")))
+			{
+				Pin->CPPType = FString::Printf(TEXT("TArray<%s>"), *Pin->CPPType);
+			}
+			Pin->Direction = InDirection;
+			Pin->bIsDynamicArray = bIsArray;
+			InNode->Pins.Add(Pin);
+
+			if(Pin->Direction != ERigVMPinDirection::Hidden && !bIsArray && !Pin->IsExecuteContext())
+			{
+				if(UScriptStruct* Struct = Cast<UScriptStruct>(Pin->CPPTypeObject))
+				{
+					FString DefaultValue;
+					InController->CreateDefaultValueForStructIfRequired(Pin->GetScriptStruct(), DefaultValue);
+					InController->AddPinsForStruct(Struct, InNode, Pin, InDirection, DefaultValue, true, false);
+				}
+			}
+
+			return Pin;
+		}
+		
+		static URigVMPin* AddExecutePin(URigVMController* InController, URigVMNode* InNode, ERigVMPinDirection InDirection = ERigVMPinDirection::IO, const FName& InName = NAME_None)
+		{
+			const FName PinName = InName.IsNone() ? FRigVMStruct::ExecuteContextName : InName;
+			URigVMPin* Pin = AddPin(InController, InNode, PinName, InDirection, false, FString::Printf(TEXT("F%s"), *InController->ExecuteContextStruct->GetName()), InController->ExecuteContextStruct);
+			if(PinName == FRigVMStruct::ExecuteContextName)
+			{
+				Pin->DisplayName = FRigVMStruct::ExecuteName;
+			}
+			return Pin;
+		}
+
+		static URigVMPin* AddArrayPin(URigVMController* InController, URigVMNode* InNode, ERigVMPinDirection InDirection, const FString& InCPPType, UObject* InCPPTypeObject, const FName& InName = NAME_None)
+		{
+			const FName PinName = InName.IsNone() ? *URigVMArrayNode::ArrayName : InName;
+			return AddPin(InController, InNode, PinName, InDirection, true, InCPPType, InCPPTypeObject);  
+		}
+
+		static URigVMPin* AddElementPin(URigVMController* InController, URigVMNode* InNode, ERigVMPinDirection InDirection, const FString& InCPPType, UObject* InCPPTypeObject)
+		{
+			return AddPin(InController, InNode, *URigVMArrayNode::ElementName, InDirection, false, InCPPType, InCPPTypeObject);  
+		}
+
+		static URigVMPin* AddIndexPin(URigVMController* InController, URigVMNode* InNode, ERigVMPinDirection InDirection)
+		{
+			return AddPin(InController, InNode, *URigVMArrayNode::IndexName, InDirection, false, TEXT("int32"), nullptr);  
+		}
+
+		static URigVMPin* AddNumPin(URigVMController* InController, URigVMNode* InNode, ERigVMPinDirection InDirection)
+		{
+			return AddPin(InController, InNode, *URigVMArrayNode::NumName, InDirection, false, TEXT("int32"), nullptr);  
+		}
+
+		static URigVMPin* AddCountPin(URigVMController* InController, URigVMNode* InNode, ERigVMPinDirection InDirection)
+		{
+			return AddPin(InController, InNode, *URigVMArrayNode::CountName, InDirection, false, TEXT("int32"), nullptr);  
+		}
+
+		static URigVMPin* AddRatioPin(URigVMController* InController, URigVMNode* InNode)
+		{
+			return AddPin(InController, InNode, *URigVMArrayNode::RatioName, ERigVMPinDirection::Output, false, TEXT("float"), nullptr);  
+		}
+
+		static URigVMPin* AddContinuePin(URigVMController* InController, URigVMNode* InNode)
+		{
+			return AddPin(InController, InNode, *URigVMArrayNode::ContinueName, ERigVMPinDirection::Hidden, false, TEXT("bool"), nullptr);
+		}
+
+		static URigVMPin* AddSuccessPin(URigVMController* InController, URigVMNode* InNode)
+		{
+			return AddPin(InController, InNode, *URigVMArrayNode::SuccessName, ERigVMPinDirection::Output, false, TEXT("bool"), nullptr);  
+		}
+	};
+
+	switch(InOpCode)
+	{
+		case ERigVMOpCode::ArrayReset:
+		case ERigVMOpCode::ArrayReverse:
+		{
+			Local::AddExecutePin(this, Node);
+			Local::AddArrayPin(this, Node, ERigVMPinDirection::IO, CPPType, InCPPTypeObject);
+			break;
+		}
+		case ERigVMOpCode::ArrayGetNum:
+		{
+			Local::AddArrayPin(this, Node, ERigVMPinDirection::Input, CPPType, InCPPTypeObject);
+			Local::AddNumPin(this, Node, ERigVMPinDirection::Output);
+			break;
+		} 
+		case ERigVMOpCode::ArraySetNum:
+		{
+			Local::AddExecutePin(this, Node);
+			Local::AddArrayPin(this, Node, ERigVMPinDirection::IO, CPPType, InCPPTypeObject);
+			Local::AddNumPin(this, Node, ERigVMPinDirection::Input);
+			break;
+		}
+		case ERigVMOpCode::ArrayGetAtIndex:
+		{
+			Local::AddArrayPin(this, Node, ERigVMPinDirection::Input, CPPType, InCPPTypeObject);
+			Local::AddIndexPin(this, Node, ERigVMPinDirection::Input);
+			Local::AddElementPin(this, Node, ERigVMPinDirection::Output, CPPType, InCPPTypeObject);
+			break;
+		}
+		case ERigVMOpCode::ArraySetAtIndex:
+		case ERigVMOpCode::ArrayInsert:
+		{
+			Local::AddExecutePin(this, Node);
+			Local::AddArrayPin(this, Node, ERigVMPinDirection::IO, CPPType, InCPPTypeObject);
+			Local::AddIndexPin(this, Node, ERigVMPinDirection::Input);
+			Local::AddElementPin(this, Node, ERigVMPinDirection::Input, CPPType, InCPPTypeObject);
+			break;
+		}
+		case ERigVMOpCode::ArrayAdd:
+		{
+			Local::AddExecutePin(this, Node);
+			Local::AddArrayPin(this, Node, ERigVMPinDirection::IO, CPPType, InCPPTypeObject);
+			Local::AddElementPin(this, Node, ERigVMPinDirection::Input, CPPType, InCPPTypeObject);
+			Local::AddIndexPin(this, Node, ERigVMPinDirection::Output);
+			break;
+		}
+		case ERigVMOpCode::ArrayFind:
+		{
+			Local::AddArrayPin(this, Node, ERigVMPinDirection::Input, CPPType, InCPPTypeObject);
+			Local::AddElementPin(this, Node, ERigVMPinDirection::Input, CPPType, InCPPTypeObject);
+			Local::AddIndexPin(this, Node, ERigVMPinDirection::Output);
+			Local::AddSuccessPin(this, Node);
+			break;
+		}
+		case ERigVMOpCode::ArrayRemove:
+		{
+			Local::AddExecutePin(this, Node);
+			Local::AddArrayPin(this, Node, ERigVMPinDirection::IO, CPPType, InCPPTypeObject);
+			Local::AddIndexPin(this, Node, ERigVMPinDirection::Input);
+			break;
+		}
+		case ERigVMOpCode::ArrayAppend:
+		case ERigVMOpCode::ArrayUnion:
+		{
+			Local::AddExecutePin(this, Node);
+			Local::AddArrayPin(this, Node, ERigVMPinDirection::IO, CPPType, InCPPTypeObject);
+			Local::AddArrayPin(this, Node, ERigVMPinDirection::Input, CPPType, InCPPTypeObject, *URigVMArrayNode::OtherName);
+			break;
+		}
+		case ERigVMOpCode::ArrayClone:
+		{
+			Local::AddArrayPin(this, Node, ERigVMPinDirection::Input, CPPType, InCPPTypeObject);
+			Local::AddArrayPin(this, Node, ERigVMPinDirection::Output, CPPType, InCPPTypeObject, *URigVMArrayNode::CloneName);
+			break;
+		}
+		case ERigVMOpCode::ArrayIterator:
+		{
+			Local::AddExecutePin(this, Node);
+			Local::AddArrayPin(this, Node, ERigVMPinDirection::Input, CPPType, InCPPTypeObject);
+			Local::AddElementPin(this, Node, ERigVMPinDirection::Output, CPPType, InCPPTypeObject);
+			Local::AddIndexPin(this, Node, ERigVMPinDirection::Output);
+			Local::AddCountPin(this, Node, ERigVMPinDirection::Output);
+			Local::AddRatioPin(this, Node);
+			Local::AddContinuePin(this, Node);
+			Local::AddExecutePin(this, Node, ERigVMPinDirection::Output, *URigVMArrayNode::CompletedName);
+			break;
+		}
+		case ERigVMOpCode::ArrayDifference:
+		case ERigVMOpCode::ArrayIntersection:
+		{
+			Local::AddArrayPin(this, Node, ERigVMPinDirection::Input, CPPType, InCPPTypeObject);
+			Local::AddArrayPin(this, Node, ERigVMPinDirection::Input, CPPType, InCPPTypeObject, *URigVMArrayNode::OtherName);
+			Local::AddArrayPin(this, Node, ERigVMPinDirection::Output, CPPType, InCPPTypeObject, *URigVMArrayNode::ResultName);
+			break;
+		}
+		default:
+		{
+			checkNoEntry();
+		}
+	}
+
+	Graph->Nodes.Add(Node);
+
+	if (!bSuspendNotifications)
+	{
+		Graph->MarkPackageDirty();
+	}
+
+	FRigVMAddArrayNodeAction Action;
+	if (bSetupUndoRedo)
+	{
+		Action = FRigVMAddArrayNodeAction(Node);
+		Action.Title = FString::Printf(TEXT("Add %s Array Node"), *Node->GetNodeTitle());
+		ActionStack->BeginAction(Action);
+	}
+
+	Notify(ERigVMGraphNotifType::NodeAdded, Node);
+
+	if (bSetupUndoRedo)
+	{
+		ActionStack->EndAction(Action);
+	}
+
+	if (bPrintPythonCommand)
+	{
+		TArray<FString> Commands = GetAddNodePythonCommands(Node);
+		for (const FString& Command : Commands)
+		{
+			RigVMPythonUtils::Print(GetGraph()->GetRootGraph()->GetOuter()->GetFName().ToString(), 
+				FString::Printf(TEXT("%s"), *Command));
+		}
+	}
+
+	return Node;
+}
+
+URigVMArrayNode* URigVMController::AddArrayNodeFromObjectPath(ERigVMOpCode InOpCode, const FString& InCPPType,
+	const FString& InCPPTypeObjectPath, const FVector2D& InPosition, const FString& InNodeName, bool bSetupUndoRedo,
+	bool bPrintPythonCommand)
+{
+	if (!IsValidGraph())
+	{
+		return nullptr;
+	}
+
+	UObject* CPPTypeObject = nullptr;
+	if (!InCPPTypeObjectPath.IsEmpty())
+	{
+		CPPTypeObject = URigVMPin::FindObjectFromCPPTypeObjectPath<UObject>(InCPPTypeObjectPath);
+		if (CPPTypeObject == nullptr)
+		{
+			ReportErrorf(TEXT("Cannot find cpp type object for path '%s'."), *InCPPTypeObjectPath);
+			return nullptr;
+		}
+	}
+
+	return AddArrayNode(InOpCode, InCPPType, CPPTypeObject, InPosition, InNodeName, bSetupUndoRedo, bPrintPythonCommand);
+}
+
 void URigVMController::ForEveryPinRecursively(URigVMPin* InPin, TFunction<void(URigVMPin*)> OnEachPinFunction)
 {
 	OnEachPinFunction(InPin);
diff --git a/Engine/Source/Developer/RigVMDeveloper/Private/RigVMModel/RigVMControllerActions.cpp b/Engine/Source/Developer/RigVMDeveloper/Private/RigVMModel/RigVMControllerActions.cpp
--- a/Engine/Source/Developer/RigVMDeveloper/Private/RigVMModel/RigVMControllerActions.cpp
+++ b/Engine/Source/Developer/RigVMDeveloper/Private/RigVMModel/RigVMControllerActions.cpp
@@ -716,6 +716,18 @@ FRigVMRemoveNodeAction::FRigVMRemoveNodeAction(URigVMNode* InNode, URigVMControl
 	{
 		InverseAction.AddAction(FRigVMAddEnumNodeAction(EnumNode));
 	}
+	else if (URigVMArrayNode* ArrayNode = Cast<URigVMArrayNode>(InNode))
+	{
+		InverseAction.AddAction(FRigVMAddArrayNodeAction(ArrayNode));
+		for (URigVMPin* Pin : ArrayNode->GetPins())
+		{
+			if (Pin->GetDirection() == ERigVMPinDirection::Input ||
+				Pin->GetDirection() == ERigVMPinDirection::Visible)
+			{
+				InverseAction.AddAction(FRigVMSetPinDefaultValueAction(Pin, Pin->GetDefaultValue()));
+			}
+		}
+	}
 	else if (URigVMLibraryNode* LibraryNode = Cast<URigVMLibraryNode>(InNode))
 	{
 		InverseAction.AddAction(FRigVMImportNodeFromTextAction(LibraryNode, InController));
@@ -1878,3 +1890,45 @@ bool FRigVMChangeLocalVariableTypeAction::Redo(URigVMController* InController)
 	}
 	return FRigVMBaseAction::Redo(InController);
 }
+
+FRigVMAddArrayNodeAction::FRigVMAddArrayNodeAction()
+	: OpCode(ERigVMOpCode::Invalid)
+	, CPPType()
+	, CPPTypeObjectPath()
+	, Position(FVector2D::ZeroVector)
+	, NodePath()
+{
+}
+
+FRigVMAddArrayNodeAction::FRigVMAddArrayNodeAction(URigVMArrayNode* InNode)
+	: OpCode(InNode->GetOpCode())
+	, CPPType(InNode->GetCPPType())
+	, CPPTypeObjectPath()
+	, Position(InNode->GetPosition())
+	, NodePath(InNode->GetNodePath())
+{
+	if (InNode->GetCPPTypeObject())
+	{
+		CPPTypeObjectPath = InNode->GetCPPTypeObject()->GetPathName();
+	}
+}
+
+bool FRigVMAddArrayNodeAction::Undo(URigVMController* InController)
+{
+	if (!FRigVMBaseAction::Undo(InController))
+	{
+		return false;
+	}
+	return InController->RemoveNodeByName(*NodePath, false);
+}
+
+bool FRigVMAddArrayNodeAction::Redo(URigVMController* InController)
+{
+#if WITH_EDITOR
+	if (URigVMArrayNode* Node = InController->AddArrayNodeFromObjectPath(OpCode, CPPType, CPPTypeObjectPath, Position, NodePath, false))
+	{
+		return FRigVMBaseAction::Redo(InController);
+	}
+#endif
+	return false;
+}
diff --git a/Engine/Source/Developer/RigVMDeveloper/Public/RigVMCompiler/RigVMAST.h b/Engine/Source/Developer/RigVMDeveloper/Public/RigVMCompiler/RigVMAST.h
--- a/Engine/Source/Developer/RigVMDeveloper/Public/RigVMCompiler/RigVMAST.h
+++ b/Engine/Source/Developer/RigVMDeveloper/Public/RigVMCompiler/RigVMAST.h
@@ -28,6 +28,7 @@ class FRigVMExitExprAST;
 class FRigVMBranchExprAST;
 class FRigVMIfExprAST;
 class FRigVMSelectExprAST;
+class FRigVMArrayExprAST;
 
 class URigVMPin;
 class URigVMLink;
@@ -65,6 +66,7 @@ public:
 		Branch,
 		If,
 		Select,
+		Array,
 		Invalid
 	};
 
@@ -422,6 +424,17 @@ FORCEINLINE const FRigVMSelectExprAST* FRigVMExprAST::To() const
 	return (const FRigVMSelectExprAST*)this;
 }
 
+// specialized cast for type checking
+// for a Array / FRigVMArrayExprAST expression
+// will raise if types are not compatible
+// @return this expression cast to FRigVMArrayExprAST 
+template<>
+FORCEINLINE const FRigVMArrayExprAST* FRigVMExprAST::To() const
+{
+	ensure(IsA(EType::Array));
+	return (const FRigVMArrayExprAST*)this;
+}
+
 /*
  * An abstract syntax tree block expression represents a sequence
  * of child expressions to be executed in order.
@@ -497,6 +510,7 @@ public:
 	URigVMNode* GetNode() const { return GetProxy().GetSubjectChecked<URigVMNode>(); }
 
 	virtual bool IsConstant() const override;
+	const FRigVMVarExprAST* FindVarWithPinName(const FName& InPinName) const;
 
 protected:
 
@@ -579,8 +593,6 @@ public:
 		return InType == EType::CallExtern;
 	};
 
-	const FRigVMVarExprAST* FindVarWithPinName(const FName& InPinName) const;
-
 protected:
 
 	// default constructor (protected so that only parser can access it)
@@ -1053,7 +1065,7 @@ private:
 /*
  * An abstract syntax tree select expression represents a branch point for
  * selecting between multiple values.
- * In C++ the branch is is the definition: switch case
+ * In C++ the select is the definition: switch case
  */
 class RIGVMDEVELOPER_API FRigVMSelectExprAST : public FRigVMNodeExprAST
 {
@@ -1086,12 +1098,46 @@ protected:
 		: FRigVMNodeExprAST(EType::Select, InProxy)
 	{}
 
+private:
+
+	friend class FRigVMParserAST;
+};
+
+/*
+* An abstract syntax tree array expression that can be used to represent
+* any array operation, such as Add, GetNum etc
+* In C++ the array op is the definition: Value[Index] etc
+*/
+class RIGVMDEVELOPER_API FRigVMArrayExprAST : public FRigVMNodeExprAST
+{
+public:
+
+	// virtual destructor
+	virtual ~FRigVMArrayExprAST() {}
+
+	// disable copy constructor
+	FRigVMArrayExprAST(const FRigVMEntryExprAST&) = delete;
+
+	// overload of the type checking mechanism
+	virtual bool IsA(EType InType) const override
+	{
+		return InType == EType::Array;
+	};
+
+protected:
+
+	// default constructor (protected so that only parser can access it)
+	FRigVMArrayExprAST(const FRigVMASTProxy& InProxy)
+		: FRigVMNodeExprAST(EType::Array, InProxy)
+	{
+	}
 
 private:
 
 	friend class FRigVMParserAST;
 };
 
+
 /*
  * The settings to apply during the parse of the abstract syntax tree.
  * The folding settings can affect the performance of the parse dramatically.
diff --git a/Engine/Source/Developer/RigVMDeveloper/Public/RigVMCompiler/RigVMCompiler.h b/Engine/Source/Developer/RigVMDeveloper/Public/RigVMCompiler/RigVMCompiler.h
--- a/Engine/Source/Developer/RigVMDeveloper/Public/RigVMCompiler/RigVMCompiler.h
+++ b/Engine/Source/Developer/RigVMDeveloper/Public/RigVMCompiler/RigVMCompiler.h
@@ -130,6 +130,11 @@ public:
 	static UScriptStruct* GetScriptStructForCPPType(const FString& InCPPType);
 	static FString GetPinHash(const URigVMPin* InPin, const FRigVMVarExprAST* InVarExpr, bool bIsDebugValue = false);
 
+	// follows assignment expressions to find the source ref counted containers
+	// since ref counted containers are not copied for assignments.
+	// this is currently only used for arrays.
+	static const FRigVMVarExprAST* GetSourceVarExpr(const FRigVMExprAST* InExpr);
+
 #if UE_RIGVM_UCLASS_BASED_STORAGE_DISABLED
 	
 	void CreateDebugRegister(URigVMPin* InPin, URigVM* OutVM, TMap<FString, FRigVMOperand>* OutOperands, TSharedPtr<FRigVMParserAST> InRuntimeAST);
@@ -163,6 +168,7 @@ private:
 	void TraverseBranch(const FRigVMBranchExprAST* InExpr, FRigVMCompilerWorkData& WorkData);
 	void TraverseIf(const FRigVMIfExprAST* InExpr, FRigVMCompilerWorkData& WorkData);
 	void TraverseSelect(const FRigVMSelectExprAST* InExpr, FRigVMCompilerWorkData& WorkData);
+	void TraverseArray(const FRigVMArrayExprAST* InExpr, FRigVMCompilerWorkData& WorkData);
 
 	void InitializeLocalVariables(const FRigVMExprAST* InExpr, FRigVMCompilerWorkData& WorkData);
 
diff --git a/Engine/Source/Developer/RigVMDeveloper/Public/RigVMModel/Nodes/RigVMArrayNode.h b/Engine/Source/Developer/RigVMDeveloper/Public/RigVMModel/Nodes/RigVMArrayNode.h
new file mode 100644
--- /dev/null
+++ b/Engine/Source/Developer/RigVMDeveloper/Public/RigVMModel/Nodes/RigVMArrayNode.h
@@ -0,0 +1,80 @@
+// Copyright Epic Games, Inc. All Rights Reserved.
+
+#pragma once
+
+#include "RigVMCore/RigVMExternalVariable.h"
+#include "RigVMModel/RigVMNode.h"
+#include "RigVMArrayNode.generated.h"
+
+/**
+ * The Array Node represents one of a series available
+ * array operations such as SetNum, GetAtIndex etc.
+ */
+UCLASS(BlueprintType)
+class RIGVMDEVELOPER_API URigVMArrayNode : public URigVMNode
+{
+	GENERATED_BODY()
+
+public:
+
+	// Default constructor
+	URigVMArrayNode();
+
+	// Override of node title
+	virtual FString GetNodeTitle() const override;
+
+	// static version
+	static FString GetNodeTitle(ERigVMOpCode InOpCode);
+
+	// Override of the node tooltip
+	virtual FText GetToolTipText() const override;
+
+	// Override of the node color
+	virtual FLinearColor GetNodeColor() const override;
+
+	// Returns the op code of this node
+	UFUNCTION(BlueprintCallable, Category = RigVMArrayNode)
+	ERigVMOpCode GetOpCode() const;
+
+	// Returns the C++ data type of the array
+	UFUNCTION(BlueprintCallable, Category = RigVMArrayNode)
+	FString GetCPPType() const;
+
+	// Returns the C++ data type struct of the array (or nullptr)
+	UFUNCTION(BlueprintCallable, Category = RigVMArrayNode)
+	UObject* GetCPPTypeObject() const;
+
+	// Override of node color
+	virtual bool IsDefinedAsVarying() const override { return true; }
+
+	// Override loop node
+	virtual bool IsLoopNode() const override;
+
+protected:
+
+	virtual FText GetToolTipTextForPin(const URigVMPin* InPin) const override;
+	
+private:
+
+	UPROPERTY()
+	ERigVMOpCode OpCode;
+
+	static const FString ArrayName;
+	static const FString NumName;
+	static const FString IndexName;
+	static const FString ElementName;
+	static const FString SuccessName;
+	static const FString OtherName;
+	static const FString CloneName;
+	static const FString CountName;
+	static const FString RatioName;
+	static const FString ResultName;
+	static const FString ContinueName;
+	static const FString CompletedName;
+
+	friend class URigVMController;
+	friend class URigVMCompiler;
+	friend class FRigVMVarExprAST;
+	friend class FRigVMParserAST;
+};
+
diff --git a/Engine/Source/Developer/RigVMDeveloper/Public/RigVMModel/Nodes/RigVMUnitNode.h b/Engine/Source/Developer/RigVMDeveloper/Public/RigVMModel/Nodes/RigVMUnitNode.h
--- a/Engine/Source/Developer/RigVMDeveloper/Public/RigVMModel/Nodes/RigVMUnitNode.h
+++ b/Engine/Source/Developer/RigVMDeveloper/Public/RigVMModel/Nodes/RigVMUnitNode.h
@@ -25,6 +25,7 @@ public:
 	virtual bool IsDefinedAsConstant() const override;
 	virtual bool IsDefinedAsVarying() const override;
 	virtual FName GetEventName() const override;
+	virtual bool IsLoopNode() const override;
 
 	bool IsDeprecated() const;
 	FString GetDeprecatedMetadata() const;
@@ -34,9 +35,6 @@ public:
 	UFUNCTION(BlueprintCallable, Category = RigVMUnitNode)
 	UScriptStruct* GetScriptStruct() const;
 
-	// return true if this node is a loop node
-	bool IsLoopNode() const;
-
 	// Returns the name of the declared RIGVM_METHOD
 	UFUNCTION(BlueprintCallable, Category = RigVMUnitNode)
 	FName GetMethodName() const;
diff --git a/Engine/Source/Developer/RigVMDeveloper/Public/RigVMModel/RigVMController.h b/Engine/Source/Developer/RigVMDeveloper/Public/RigVMModel/RigVMController.h
--- a/Engine/Source/Developer/RigVMDeveloper/Public/RigVMModel/RigVMController.h
+++ b/Engine/Source/Developer/RigVMDeveloper/Public/RigVMModel/RigVMController.h
@@ -16,6 +16,7 @@
 #include "RigVMModel/Nodes/RigVMEnumNode.h"
 #include "RigVMModel/Nodes/RigVMCollapseNode.h"
 #include "RigVMModel/Nodes/RigVMFunctionReferenceNode.h"
+#include "RigVMModel/Nodes/RigVMArrayNode.h"
 #include "RigVMController.generated.h"
 
 class URigVMActionStack;
@@ -281,6 +282,16 @@ public:
 	UFUNCTION(BlueprintCallable, Category = RigVMController)
 	URigVMEnumNode* AddEnumNode(const FName& InCPPTypeObjectPath, const FVector2D& InPosition = FVector2D::ZeroVector, const FString& InNodeName = TEXT(""), bool bSetupUndoRedo = true, bool bPrintPythonCommand = false);
 
+	// Adds a Array Node to the edited Graph.
+	// This causes a NodeAdded modified event.
+	UFUNCTION(BlueprintCallable, Category = RigVMController)
+	URigVMArrayNode* AddArrayNode(ERigVMOpCode InOpCode, const FString& InCPPType, UObject* InCPPTypeObject, const FVector2D& InPosition = FVector2D::ZeroVector, const FString& InNodeName = TEXT(""), bool bSetupUndoRedo = true, bool bPrintPythonCommand = false);
+
+	// Adds a Array Node to the edited Graph given a struct object path name.
+	// This causes a NodeAdded modified event.
+	UFUNCTION(BlueprintCallable, Category = RigVMController)
+	URigVMArrayNode* AddArrayNodeFromObjectPath(ERigVMOpCode InOpCode, const FString& InCPPType, const FString& InCPPTypeObjectPath, const FVector2D& InPosition = FVector2D::ZeroVector, const FString& InNodeName = TEXT(""), bool bSetupUndoRedo = true, bool bPrintPythonCommand = false);
+
 	// Un-does the last action on the stack.
 	// Note: This should really only be used for unit tests,
 	// use the GEditor's main Undo method instead.
@@ -771,9 +782,13 @@ private:
 	bool BindPinToVariable(URigVMPin* InPin, const FString& InNewBoundVariablePath, bool bSetupUndoRedo);
 	bool MakeBindingsFromVariableNode(URigVMVariableNode* InNode, bool bSetupUndoRedo);
 	bool PromotePinToVariable(URigVMPin* InPin, bool bCreateVariableNode, const FVector2D& InNodePosition, bool bSetupUndoRedo);
-	bool AddLink(URigVMPin* OutputPin, URigVMPin* InputPin, bool bSetupUndoRedo);
-	bool BreakLink(URigVMPin* OutputPin, URigVMPin* InputPin, bool bSetupUndoRedo);
-	bool BreakAllLinks(URigVMPin* Pin, bool bAsInput, bool bSetupUndoRedo);
+
+public:
+	bool AddLink(URigVMPin* OutputPin, URigVMPin* InputPin, bool bSetupUndoRedo = true);
+	bool BreakLink(URigVMPin* OutputPin, URigVMPin* InputPin, bool bSetupUndoRedo = true);
+	bool BreakAllLinks(URigVMPin* Pin, bool bAsInput, bool bSetupUndoRedo = true);
+
+private:
 	void BreakAllLinksRecursive(URigVMPin* Pin, bool bAsInput, bool bTowardsParent, bool bSetupUndoRedo);
 	void UpdateRerouteNodeAfterChangingLinks(URigVMPin* PinChanged, bool bSetupUndoRedo = true);
 	bool SetPinExpansion(URigVMPin* InPin, bool bIsExpanded, bool bSetupUndoRedo = true);
diff --git a/Engine/Source/Developer/RigVMDeveloper/Public/RigVMModel/RigVMControllerActions.h b/Engine/Source/Developer/RigVMDeveloper/Public/RigVMModel/RigVMControllerActions.h
--- a/Engine/Source/Developer/RigVMDeveloper/Public/RigVMModel/RigVMControllerActions.h
+++ b/Engine/Source/Developer/RigVMDeveloper/Public/RigVMModel/RigVMControllerActions.h
@@ -1584,4 +1584,36 @@ public:
 
 	UPROPERTY()
 	TObjectPtr<UObject> CPPTypeObject;
-};
\ No newline at end of file
+};
+
+/**
+* An action adding a array node to the graph.
+*/
+USTRUCT()
+struct FRigVMAddArrayNodeAction : public FRigVMBaseAction
+{
+	GENERATED_BODY()
+
+public:
+
+	FRigVMAddArrayNodeAction();
+	FRigVMAddArrayNodeAction(URigVMArrayNode* InNode);
+	virtual ~FRigVMAddArrayNodeAction() {};
+	virtual bool Undo(URigVMController* InController) override;
+	virtual bool Redo(URigVMController* InController) override;
+
+	UPROPERTY()
+	ERigVMOpCode OpCode;
+	
+	UPROPERTY()
+	FString CPPType;
+
+	UPROPERTY()
+	FString CPPTypeObjectPath;
+
+	UPROPERTY()
+	FVector2D Position;
+
+	UPROPERTY()
+	FString NodePath;
+};
diff --git a/Engine/Source/Developer/RigVMDeveloper/Public/RigVMModel/RigVMNode.h b/Engine/Source/Developer/RigVMDeveloper/Public/RigVMModel/RigVMNode.h
--- a/Engine/Source/Developer/RigVMDeveloper/Public/RigVMModel/RigVMNode.h
+++ b/Engine/Source/Developer/RigVMDeveloper/Public/RigVMModel/RigVMNode.h
@@ -194,6 +194,9 @@ public:
 	// Returns the number of visited / run instructions for this node
 	virtual int32 GetInstructionVisitedCount(URigVM* InVM, const FRigVMASTProxy& InProxy = FRigVMASTProxy(), bool bConsolidatePerNode = false) const; 
 
+	// return true if this node is a loop node
+	virtual bool IsLoopNode() const { return false; }
+
 	UFUNCTION(BlueprintCallable, Category = RigVMNode)
 	bool HasBreakpoint() const { return bHasBreakpoint; }
 
diff --git a/Engine/Source/Developer/RigVMDeveloper/Public/RigVMModel/RigVMPin.h b/Engine/Source/Developer/RigVMDeveloper/Public/RigVMModel/RigVMPin.h
--- a/Engine/Source/Developer/RigVMDeveloper/Public/RigVMModel/RigVMPin.h
+++ b/Engine/Source/Developer/RigVMDeveloper/Public/RigVMModel/RigVMPin.h
@@ -178,6 +178,10 @@ public:
 	UFUNCTION(BlueprintCallable, Category = RigVMPin)
 	bool IsDynamicArray() const;
 
+	// Returns true if this data type is referenced counted
+	UFUNCTION(BlueprintCallable, Category = RigVMPin)
+	bool IsReferenceCountedContainer() const { return IsDynamicArray(); }
+
 	// Returns the index of the Pin within the node / parent Pin
 	UFUNCTION(BlueprintCallable, Category = RigVMPin)
 	int32 GetPinIndex() const;
diff --git a/Engine/Source/Runtime/RigVM/Private/RigVMCore/RigVM.cpp b/Engine/Source/Runtime/RigVM/Private/RigVMCore/RigVM.cpp
--- a/Engine/Source/Runtime/RigVM/Private/RigVMCore/RigVM.cpp
+++ b/Engine/Source/Runtime/RigVM/Private/RigVMCore/RigVM.cpp
@@ -431,6 +431,7 @@ bool URigVM::ValidateAllOperandsDuringLoad()
 			case ERigVMOpCode::Increment:
 			case ERigVMOpCode::Decrement:
 			case ERigVMOpCode::ArrayReset:
+			case ERigVMOpCode::ArrayReverse:
 			{
 				const FRigVMUnaryOp& Op = ByteCodeStorage.GetOpAt<FRigVMUnaryOp>(ByteCodeInstruction);
 				CheckOperandValidity(Op.Arg);
@@ -466,6 +467,7 @@ bool URigVM::ValidateAllOperandsDuringLoad()
 			case ERigVMOpCode::ArrayAppend:
 			case ERigVMOpCode::ArrayClone:
 			case ERigVMOpCode::ArrayRemove:
+			case ERigVMOpCode::ArrayUnion:
 			{
 				const FRigVMBinaryOp& Op = ByteCodeStorage.GetOpAt<FRigVMBinaryOp>(ByteCodeInstruction);
 				CheckOperandValidity(Op.ArgA);
@@ -476,6 +478,8 @@ bool URigVM::ValidateAllOperandsDuringLoad()
 			case ERigVMOpCode::ArrayGetAtIndex:
 			case ERigVMOpCode::ArraySetAtIndex:
 			case ERigVMOpCode::ArrayInsert:
+			case ERigVMOpCode::ArrayDifference:
+			case ERigVMOpCode::ArrayIntersection:
 			{
 				const FRigVMTernaryOp& Op = ByteCodeStorage.GetOpAt<FRigVMTernaryOp>(ByteCodeInstruction);
 				CheckOperandValidity(Op.ArgA);
@@ -1131,6 +1135,7 @@ void URigVM::CacheMemoryHandlesIfRequired(TArrayView<URigVMMemoryStorage*> InMem
 			case ERigVMOpCode::Increment:
 			case ERigVMOpCode::Decrement:
 			case ERigVMOpCode::ArrayReset:
+			case ERigVMOpCode::ArrayReverse:
 			{
 				const FRigVMUnaryOp& Op = ByteCode.GetOpAt<FRigVMUnaryOp>(Instructions[InstructionIndex]);
 				CacheSingleMemoryHandle(Op.Arg);
@@ -1202,6 +1207,7 @@ void URigVM::CacheMemoryHandlesIfRequired(TArrayView<URigVMMemoryStorage*> InMem
 			case ERigVMOpCode::ArrayAppend:
 			case ERigVMOpCode::ArrayClone:
 			case ERigVMOpCode::ArrayRemove:
+			case ERigVMOpCode::ArrayUnion:
 			{
 				const FRigVMBinaryOp& Op = ByteCode.GetOpAt<FRigVMBinaryOp>(Instructions[InstructionIndex]);
 				CacheSingleMemoryHandle(Op.ArgA);
@@ -1213,6 +1219,8 @@ void URigVM::CacheMemoryHandlesIfRequired(TArrayView<URigVMMemoryStorage*> InMem
 			case ERigVMOpCode::ArrayGetAtIndex:
 			case ERigVMOpCode::ArraySetAtIndex:
 			case ERigVMOpCode::ArrayInsert:
+			case ERigVMOpCode::ArrayDifference:
+			case ERigVMOpCode::ArrayIntersection:
 			{
 				const FRigVMTernaryOp& Op = ByteCode.GetOpAt<FRigVMTernaryOp>(Instructions[InstructionIndex]);
 				CacheSingleMemoryHandle(Op.ArgA);
@@ -1783,6 +1791,11 @@ bool URigVM::Initialize(TArrayView<URigVMMemoryStorage*> Memory, TArrayView<void
 			case ERigVMOpCode::ArrayAdd:
 			case ERigVMOpCode::ArrayFind:
 			case ERigVMOpCode::ArrayIterator:
+			case ERigVMOpCode::ArrayUnion:
+			case ERigVMOpCode::ArrayDifference:
+			case ERigVMOpCode::ArrayIntersection:
+			case ERigVMOpCode::ArrayReverse:
+
 			{
 				break;
 			}
@@ -2377,7 +2390,14 @@ bool URigVM::Execute(TArrayView<URigVMMemoryStorage*> Memory, TArrayView<void*>
 				FRigVMMemoryHandle& ArrayHandle = CachedMemoryHandles[FirstHandleForInstruction[Context.InstructionIndex]]; 					
 				FScriptArrayHelper ArrayHelper(CastFieldChecked<FArrayProperty>(ArrayHandle.GetProperty()), ArrayHandle.GetData());
 				ArrayHelper.Resize(0);
-#endif	
+
+				if(DebugMemoryStorageObject->Num() > 0)
+				{
+					const FRigVMUnaryOp& Op = ByteCode.GetOpAt<FRigVMUnaryOp>(Instruction);
+					CopyOperandForDebuggingIfNeeded(Op.Arg, ArrayHandle);
+				}
+#endif
+
 				Context.InstructionIndex++;
 				break;
 			}
@@ -2388,6 +2408,13 @@ bool URigVM::Execute(TArrayView<URigVMMemoryStorage*> Memory, TArrayView<void*>
 				FScriptArrayHelper ArrayHelper(CastFieldChecked<FArrayProperty>(ArrayHandle.GetProperty()), ArrayHandle.GetData());
 				int32& Count = (*((int32*)CachedMemoryHandles[FirstHandleForInstruction[Context.InstructionIndex] + 1].GetData()));
 				Count = ArrayHelper.Num();
+
+				if(DebugMemoryStorageObject->Num() > 0)
+				{
+					const FRigVMBinaryOp& Op = ByteCode.GetOpAt<FRigVMBinaryOp>(Instruction);
+					CopyOperandForDebuggingIfNeeded(Op.ArgA, ArrayHandle);
+					CopyOperandForDebuggingIfNeeded(Op.ArgB, CachedMemoryHandles[FirstHandleForInstruction[Context.InstructionIndex] + 1]);
+				}
 #endif	
 				Context.InstructionIndex++;
 				break;
@@ -2402,6 +2429,13 @@ bool URigVM::Execute(TArrayView<URigVMMemoryStorage*> Memory, TArrayView<void*>
 				{
 					ArrayHelper.Resize(Count);
 				}
+
+				if(DebugMemoryStorageObject->Num() > 0)
+				{
+					const FRigVMBinaryOp& Op = ByteCode.GetOpAt<FRigVMBinaryOp>(Instruction);
+					CopyOperandForDebuggingIfNeeded(Op.ArgA, ArrayHandle);
+					CopyOperandForDebuggingIfNeeded(Op.ArgB, CachedMemoryHandles[FirstHandleForInstruction[Context.InstructionIndex] + 1]);
+				}
 #endif
 				Context.InstructionIndex++;
 				break;
@@ -2433,6 +2467,13 @@ bool URigVM::Execute(TArrayView<URigVMMemoryStorage*> Memory, TArrayView<void*>
 						}
 					}
 				}
+
+				if(DebugMemoryStorageObject->Num() > 0)
+				{
+					const FRigVMBinaryOp& Op = ByteCode.GetOpAt<FRigVMBinaryOp>(Instruction);
+					CopyOperandForDebuggingIfNeeded(Op.ArgA, ArrayHandle);
+					CopyOperandForDebuggingIfNeeded(Op.ArgB, CachedMemoryHandles[FirstHandleForInstruction[Context.InstructionIndex] + 1]);
+				}
 #endif
 				Context.InstructionIndex++;
 				break;
@@ -2447,17 +2488,13 @@ bool URigVM::Execute(TArrayView<URigVMMemoryStorage*> Memory, TArrayView<void*>
 				FScriptArrayHelper ArrayHelper(ArrayProperty, ArrayHandle.GetData());
 				FScriptArrayHelper ClonedArrayHelper(ClonedArrayProperty, ClonedArrayHandle.GetData());
 
-				ClonedArrayHelper.Resize(ArrayHelper.Num());
-				if(ArrayHelper.Num() > 0)
+				CopyArray(ClonedArrayHelper, ClonedArrayHandle, ArrayHelper, ArrayHandle);
+					
+				if(DebugMemoryStorageObject->Num() > 0)
 				{
-					const FProperty* TargetProperty = ClonedArrayProperty->Inner;
-					const FProperty* SourceProperty = ArrayProperty->Inner;
-					for(int32 ElementIndex = 0; ElementIndex < ArrayHelper.Num(); ElementIndex++)
-					{
-						uint8* TargetMemory = ClonedArrayHelper.GetRawPtr(ElementIndex);
-						const uint8* SourceMemory = ArrayHelper.GetRawPtr(ElementIndex);
-						URigVMMemoryStorage::CopyProperty(TargetProperty, TargetMemory, SourceProperty, SourceMemory);
-					}
+					const FRigVMBinaryOp& Op = ByteCode.GetOpAt<FRigVMBinaryOp>(Instruction);
+					CopyOperandForDebuggingIfNeeded(Op.ArgA, ArrayHandle);
+					CopyOperandForDebuggingIfNeeded(Op.ArgB, ClonedArrayHandle);
 				}
 #endif
 				Context.InstructionIndex++;
@@ -2477,6 +2514,14 @@ bool URigVM::Execute(TArrayView<URigVMMemoryStorage*> Memory, TArrayView<void*>
 					const uint8* SourceMemory = ArrayHelper.GetRawPtr(Index);
 					URigVMMemoryStorage::CopyProperty(ElementHandle.GetProperty(), TargetMemory, ArrayProperty->Inner, SourceMemory);
 				}
+
+				if(DebugMemoryStorageObject->Num() > 0)
+				{
+					const FRigVMTernaryOp& Op = ByteCode.GetOpAt<FRigVMTernaryOp>(Instruction);
+					CopyOperandForDebuggingIfNeeded(Op.ArgA, ArrayHandle);
+					CopyOperandForDebuggingIfNeeded(Op.ArgB, CachedMemoryHandles[FirstHandleForInstruction[Context.InstructionIndex] + 1]);
+					CopyOperandForDebuggingIfNeeded(Op.ArgC, CachedMemoryHandles[FirstHandleForInstruction[Context.InstructionIndex] + 2]);
+				}
 #endif
 				Context.InstructionIndex++;
 				break;
@@ -2495,6 +2540,14 @@ bool URigVM::Execute(TArrayView<URigVMMemoryStorage*> Memory, TArrayView<void*>
 					const uint8* SourceMemory = ElementHandle.GetData();
 					URigVMMemoryStorage::CopyProperty(ArrayProperty->Inner, TargetMemory, ElementHandle.GetProperty(), SourceMemory);
 				}
+
+				if(DebugMemoryStorageObject->Num() > 0)
+				{
+					const FRigVMTernaryOp& Op = ByteCode.GetOpAt<FRigVMTernaryOp>(Instruction);
+					CopyOperandForDebuggingIfNeeded(Op.ArgA, ArrayHandle);
+					CopyOperandForDebuggingIfNeeded(Op.ArgB, CachedMemoryHandles[FirstHandleForInstruction[Context.InstructionIndex] + 1]);
+					CopyOperandForDebuggingIfNeeded(Op.ArgC, CachedMemoryHandles[FirstHandleForInstruction[Context.InstructionIndex] + 2]);
+				}
 #endif
 				Context.InstructionIndex++;
 				break;
@@ -2516,6 +2569,14 @@ bool URigVM::Execute(TArrayView<URigVMMemoryStorage*> Memory, TArrayView<void*>
 					const uint8* SourceMemory = ElementHandle.GetData();
 					URigVMMemoryStorage::CopyProperty(ArrayProperty->Inner, TargetMemory, ElementHandle.GetProperty(), SourceMemory);
 				}
+
+				if(DebugMemoryStorageObject->Num() > 0)
+				{
+					const FRigVMTernaryOp& Op = ByteCode.GetOpAt<FRigVMTernaryOp>(Instruction);
+					CopyOperandForDebuggingIfNeeded(Op.ArgA, ArrayHandle);
+					CopyOperandForDebuggingIfNeeded(Op.ArgB, CachedMemoryHandles[FirstHandleForInstruction[Context.InstructionIndex] + 1]);
+					CopyOperandForDebuggingIfNeeded(Op.ArgC, CachedMemoryHandles[FirstHandleForInstruction[Context.InstructionIndex] + 2]);
+				}
 #endif
 				Context.InstructionIndex++;
 				break;
@@ -2531,6 +2592,13 @@ bool URigVM::Execute(TArrayView<URigVMMemoryStorage*> Memory, TArrayView<void*>
 				{
 					ArrayHelper.RemoveValues(Index, 1);
 				}
+
+				if(DebugMemoryStorageObject->Num() > 0)
+				{
+					const FRigVMBinaryOp& Op = ByteCode.GetOpAt<FRigVMBinaryOp>(Instruction);
+					CopyOperandForDebuggingIfNeeded(Op.ArgA, ArrayHandle);
+					CopyOperandForDebuggingIfNeeded(Op.ArgB, CachedMemoryHandles[FirstHandleForInstruction[Context.InstructionIndex] + 1]);
+				}
 #endif
 				Context.InstructionIndex++;
 				break;
@@ -2555,6 +2623,14 @@ bool URigVM::Execute(TArrayView<URigVMMemoryStorage*> Memory, TArrayView<void*>
 				{
 					Index = INDEX_NONE;
 				}
+
+				if(DebugMemoryStorageObject->Num() > 0)
+				{
+					const FRigVMTernaryOp& Op = ByteCode.GetOpAt<FRigVMTernaryOp>(Instruction);
+					CopyOperandForDebuggingIfNeeded(Op.ArgA, ArrayHandle);
+					CopyOperandForDebuggingIfNeeded(Op.ArgB, CachedMemoryHandles[FirstHandleForInstruction[Context.InstructionIndex] + 1]);
+					CopyOperandForDebuggingIfNeeded(Op.ArgC, CachedMemoryHandles[FirstHandleForInstruction[Context.InstructionIndex] + 2]);
+				}
 #endif
 				Context.InstructionIndex++;
 				break;
@@ -2596,6 +2672,15 @@ bool URigVM::Execute(TArrayView<URigVMMemoryStorage*> Memory, TArrayView<void*>
 					static const TCHAR IncompatibleTypes[] = TEXT("Array('%s') doesn't support searching for element('%$s').");
 					Context.Logf(EMessageSeverity::Error, IncompatibleTypes, *PropertyB->GetCPPType(), *PropertyA->GetCPPType());
 				}
+
+				if(DebugMemoryStorageObject->Num() > 0)
+				{
+					const FRigVMQuaternaryOp& Op = ByteCode.GetOpAt<FRigVMQuaternaryOp>(Instruction);
+					CopyOperandForDebuggingIfNeeded(Op.ArgA, ArrayHandle);
+					CopyOperandForDebuggingIfNeeded(Op.ArgB, CachedMemoryHandles[FirstHandleForInstruction[Context.InstructionIndex] + 1]);
+					CopyOperandForDebuggingIfNeeded(Op.ArgC, CachedMemoryHandles[FirstHandleForInstruction[Context.InstructionIndex] + 2]);
+					CopyOperandForDebuggingIfNeeded(Op.ArgD, CachedMemoryHandles[FirstHandleForInstruction[Context.InstructionIndex] + 3]);
+				}
 #endif
 				Context.InstructionIndex++;
 				break;
@@ -2623,12 +2708,208 @@ bool URigVM::Execute(TArrayView<URigVMMemoryStorage*> Memory, TArrayView<void*>
 				else
 				{
 					Ratio = float(Index) / float(Count - 1);
+
+					uint8* TargetMemory = ElementHandle.GetData();
+					const uint8* SourceMemory = ArrayHelper.GetRawPtr(Index);
+					URigVMMemoryStorage::CopyProperty(ElementHandle.GetProperty(), TargetMemory, ArrayProperty->Inner, SourceMemory);
+
+					if(DebugMemoryStorageObject->Num() > 0)
+					{
+						const FRigVMSenaryOp& Op = ByteCode.GetOpAt<FRigVMSenaryOp>(Instruction);
+						CopyOperandForDebuggingIfNeeded(Op.ArgA, ArrayHandle); // array
+						CopyOperandForDebuggingIfNeeded(Op.ArgD, CachedMemoryHandles[FirstHandleForInstruction[Context.InstructionIndex] + 3]); // count
+
+						Context.BeginSlice(Count, Index);
+						CopyOperandForDebuggingIfNeeded(Op.ArgB, CachedMemoryHandles[FirstHandleForInstruction[Context.InstructionIndex] + 1]); // element
+						CopyOperandForDebuggingIfNeeded(Op.ArgC, CachedMemoryHandles[FirstHandleForInstruction[Context.InstructionIndex] + 2]); // index
+						CopyOperandForDebuggingIfNeeded(Op.ArgE, CachedMemoryHandles[FirstHandleForInstruction[Context.InstructionIndex] + 4]); // ratio
+						Context.EndSlice();
+					}
 				}
 
 #endif
 				Context.InstructionIndex++;
 				break;
 			}
+			case ERigVMOpCode::ArrayUnion:
+			case ERigVMOpCode::ArrayDifference:
+			case ERigVMOpCode::ArrayIntersection:
+			{
+#if !UE_RIGVM_UCLASS_BASED_STORAGE_DISABLED
+				FRigVMMemoryHandle& ArrayHandleA = CachedMemoryHandles[FirstHandleForInstruction[Context.InstructionIndex]]; 					
+				FRigVMMemoryHandle& ArrayHandleB = CachedMemoryHandles[FirstHandleForInstruction[Context.InstructionIndex] + 1]; 					
+				FScriptArrayHelper ArrayHelperA(CastFieldChecked<FArrayProperty>(ArrayHandleA.GetProperty()), ArrayHandleA.GetData());
+				FScriptArrayHelper ArrayHelperB(CastFieldChecked<FArrayProperty>(ArrayHandleB.GetProperty()), ArrayHandleB.GetData());
+				const FArrayProperty* ArrayPropertyA = CastFieldChecked<FArrayProperty>(ArrayHandleA.GetProperty());
+				const FArrayProperty* ArrayPropertyB = CastFieldChecked<FArrayProperty>(ArrayHandleB.GetProperty());
+				const FProperty* ElementPropertyA = ArrayPropertyA->Inner;
+				const FProperty* ElementPropertyB = ArrayPropertyB->Inner;
+
+				TMap<uint32, int32> HashA, HashB;
+				HashA.Reserve(ArrayHelperA.Num());
+				HashB.Reserve(ArrayHelperB.Num());
+
+				for(int32 Index = 0; Index < ArrayHelperA.Num(); Index++)
+				{
+					const uint32 HashValue = ElementPropertyA->GetValueTypeHash(ArrayHelperA.GetRawPtr(Index));
+					if(!HashA.Contains(HashValue))
+					{
+						HashA.Add(HashValue, Index);
+					}
+				}
+				for(int32 Index = 0; Index < ArrayHelperB.Num(); Index++)
+				{
+					const uint32 HashValue = ElementPropertyB->GetValueTypeHash(ArrayHelperB.GetRawPtr(Index));
+					if(!HashB.Contains(HashValue))
+					{
+						HashB.Add(HashValue, Index);
+					}
+				}
+
+				if(Instruction.OpCode == ERigVMOpCode::ArrayUnion)
+				{
+					// copy the complete array to a temp storage
+					TArray<uint8, TAlignedHeapAllocator<16>> TempStorage;
+					const int32 NumElementsA = ArrayHelperA.Num();
+					TempStorage.AddZeroed(NumElementsA * ElementPropertyA->GetSize());
+					uint8* TempMemory = TempStorage.GetData();
+					for(int32 Index = 0; Index < NumElementsA; Index++)
+					{
+						ElementPropertyA->InitializeValue(TempMemory);
+						ElementPropertyA->CopyCompleteValue(TempMemory, ArrayHelperA.GetRawPtr(Index));
+						TempMemory += ElementPropertyA->GetSize();
+					}
+
+					ArrayHelperA.Resize(0);
+
+					for(const TPair<uint32, int32>& Pair : HashA)
+					{
+						int32 AddedIndex = ArrayHelperA.AddValue();
+						TempMemory = TempStorage.GetData() + Pair.Value * ElementPropertyA->GetSize();
+						
+						URigVMMemoryStorage::CopyProperty(
+							ElementPropertyA,
+							ArrayHelperA.GetRawPtr(AddedIndex),
+							ElementPropertyA,
+							TempMemory
+						);
+					}
+
+					TempMemory = TempStorage.GetData();
+					for(int32 Index = 0; Index < NumElementsA; Index++)
+					{
+						ElementPropertyA->DestroyValue(TempMemory);
+						TempMemory += ElementPropertyA->GetSize();
+					}
+
+					for(const TPair<uint32, int32>& Pair : HashB)
+					{
+						if(!HashA.Contains(Pair.Key))
+						{
+							int32 AddedIndex = ArrayHelperA.AddValue();
+							
+							URigVMMemoryStorage::CopyProperty(
+								ElementPropertyA,
+								ArrayHelperA.GetRawPtr(AddedIndex),
+								ElementPropertyB,
+								ArrayHelperB.GetRawPtr(Pair.Value)
+							);
+						}
+					}
+					
+					if(DebugMemoryStorageObject->Num() > 0)
+					{
+						const FRigVMBinaryOp& Op = ByteCode.GetOpAt<FRigVMBinaryOp>(Instruction);
+						CopyOperandForDebuggingIfNeeded(Op.ArgA, ArrayHandleA);
+						CopyOperandForDebuggingIfNeeded(Op.ArgB, ArrayHandleB);
+					}
+				}
+				else
+				{
+					FRigVMMemoryHandle& ResultArrayHandle = CachedMemoryHandles[FirstHandleForInstruction[Context.InstructionIndex] + 2]; 					
+					FScriptArrayHelper ResultArrayHelper(CastFieldChecked<FArrayProperty>(ResultArrayHandle.GetProperty()), ResultArrayHandle.GetData());
+					const FArrayProperty* ResultArrayProperty = CastFieldChecked<FArrayProperty>(ResultArrayHandle.GetProperty());
+					const FProperty* ResultElementProperty = ResultArrayProperty->Inner;
+
+					ResultArrayHelper.Resize(0);
+					
+					if(Instruction.OpCode == ERigVMOpCode::ArrayDifference)
+					{
+						for(const TPair<uint32, int32>& Pair : HashA)
+						{
+							if(!HashB.Contains(Pair.Key))
+							{
+								int32 AddedIndex = ResultArrayHelper.AddValue();
+								URigVMMemoryStorage::CopyProperty(
+									ResultElementProperty,
+									ResultArrayHelper.GetRawPtr(AddedIndex),
+									ElementPropertyA,
+									ArrayHelperA.GetRawPtr(Pair.Value)
+								);
+							}
+						}
+						for(const TPair<uint32, int32>& Pair : HashB)
+						{
+							if(!HashA.Contains(Pair.Key))
+							{
+								int32 AddedIndex = ResultArrayHelper.AddValue();
+								URigVMMemoryStorage::CopyProperty(
+									ResultElementProperty,
+									ResultArrayHelper.GetRawPtr(AddedIndex),
+									ElementPropertyB,
+									ArrayHelperB.GetRawPtr(Pair.Value)
+								);
+							}
+						}
+					}
+					else // intersection
+					{
+						for(const TPair<uint32, int32>& Pair : HashA)
+						{
+							if(HashB.Contains(Pair.Key))
+							{
+								int32 AddedIndex = ResultArrayHelper.AddValue();
+								URigVMMemoryStorage::CopyProperty(
+									ResultElementProperty,
+									ResultArrayHelper.GetRawPtr(AddedIndex),
+									ElementPropertyA,
+									ArrayHelperA.GetRawPtr(Pair.Value)
+								);
+							}
+						}
+					}
+
+					if(DebugMemoryStorageObject->Num() > 0)
+					{
+						const FRigVMTernaryOp& Op = ByteCode.GetOpAt<FRigVMTernaryOp>(Instruction);
+						CopyOperandForDebuggingIfNeeded(Op.ArgA, ArrayHandleA);
+						CopyOperandForDebuggingIfNeeded(Op.ArgB, ArrayHandleB);
+						CopyOperandForDebuggingIfNeeded(Op.ArgC, ResultArrayHandle);
+					}
+				}
+#endif
+				Context.InstructionIndex++;
+				break;
+			}
+			case ERigVMOpCode::ArrayReverse:
+			{
+#if !UE_RIGVM_UCLASS_BASED_STORAGE_DISABLED
+				FRigVMMemoryHandle& ArrayHandle = CachedMemoryHandles[FirstHandleForInstruction[Context.InstructionIndex]]; 					
+				FScriptArrayHelper ArrayHelper(CastFieldChecked<FArrayProperty>(ArrayHandle.GetProperty()), ArrayHandle.GetData());
+				for(int32 A=0, B=ArrayHelper.Num()-1; A<B; A++, B--)
+				{
+					ArrayHelper.SwapValues(A, B);
+				}
+
+				if(DebugMemoryStorageObject->Num() > 0)
+				{
+					const FRigVMUnaryOp& Op = ByteCode.GetOpAt<FRigVMUnaryOp>(Instruction);
+					CopyOperandForDebuggingIfNeeded(Op.Arg, ArrayHandle);
+				}
+#endif	
+				Context.InstructionIndex++;
+				break;
+			}
 			case ERigVMOpCode::Invalid:
 			{
 				ensure(false);
@@ -3375,6 +3656,26 @@ void URigVM::RefreshExternalPropertyPaths()
 	}
 }
 
+void URigVM::CopyArray(FScriptArrayHelper& TargetHelper, FRigVMMemoryHandle& TargetHandle,
+	FScriptArrayHelper& SourceHelper, FRigVMMemoryHandle& SourceHandle)
+{
+	const FArrayProperty* TargetArrayProperty = CastFieldChecked<FArrayProperty>(TargetHandle.GetProperty());
+	const FArrayProperty* SourceArrayProperty = CastFieldChecked<FArrayProperty>(SourceHandle.GetProperty());
+
+	TargetHelper.Resize(SourceHelper.Num());
+	if(SourceHelper.Num() > 0)
+	{
+		const FProperty* TargetProperty = TargetArrayProperty->Inner;
+		const FProperty* SourceProperty = SourceArrayProperty->Inner;
+		for(int32 ElementIndex = 0; ElementIndex < SourceHelper.Num(); ElementIndex++)
+		{
+			uint8* TargetMemory = TargetHelper.GetRawPtr(ElementIndex);
+			const uint8* SourceMemory = SourceHelper.GetRawPtr(ElementIndex);
+			URigVMMemoryStorage::CopyProperty(TargetProperty, TargetMemory, SourceProperty, SourceMemory);
+		}
+	}
+}
+
 #endif
 
 #if UE_RIGVM_UCLASS_BASED_STORAGE_DISABLED
diff --git a/Engine/Source/Runtime/RigVM/Private/RigVMCore/RigVMByteCode.cpp b/Engine/Source/Runtime/RigVM/Private/RigVMCore/RigVMByteCode.cpp
--- a/Engine/Source/Runtime/RigVM/Private/RigVMCore/RigVMByteCode.cpp
+++ b/Engine/Source/Runtime/RigVM/Private/RigVMCore/RigVMByteCode.cpp
@@ -324,6 +324,7 @@ void FRigVMByteCode::Save(FArchive& Ar) const
 			case ERigVMOpCode::Increment:
 			case ERigVMOpCode::Decrement:
 			case ERigVMOpCode::ArrayReset:
+			case ERigVMOpCode::ArrayReverse:
 			{
 				FRigVMUnaryOp Op = GetOpAt<FRigVMUnaryOp>(Instruction.ByteCodeIndex);
 				Ar << Op;
@@ -364,6 +365,7 @@ void FRigVMByteCode::Save(FArchive& Ar) const
 			case ERigVMOpCode::ArrayAppend:
 			case ERigVMOpCode::ArrayClone:
 			case ERigVMOpCode::ArrayRemove:
+			case ERigVMOpCode::ArrayUnion:
 			{
 				FRigVMBinaryOp Op = GetOpAt<FRigVMBinaryOp>(Instruction.ByteCodeIndex);
 				Ar << Op;
@@ -373,6 +375,8 @@ void FRigVMByteCode::Save(FArchive& Ar) const
 			case ERigVMOpCode::ArrayGetAtIndex:
 			case ERigVMOpCode::ArraySetAtIndex:
 			case ERigVMOpCode::ArrayInsert:
+			case ERigVMOpCode::ArrayDifference:
+			case ERigVMOpCode::ArrayIntersection:
 			{
 				FRigVMTernaryOp Op = GetOpAt<FRigVMTernaryOp>(Instruction.ByteCodeIndex);
 				Ar << Op;
@@ -538,6 +542,7 @@ void FRigVMByteCode::Load(FArchive& Ar)
 			case ERigVMOpCode::Increment:
 			case ERigVMOpCode::Decrement:
 			case ERigVMOpCode::ArrayReset:
+			case ERigVMOpCode::ArrayReverse:
 			{
 				FRigVMUnaryOp Op;
 				Ar << Op;
@@ -581,6 +586,7 @@ void FRigVMByteCode::Load(FArchive& Ar)
 			case ERigVMOpCode::ArrayAppend:
 			case ERigVMOpCode::ArrayClone:
 			case ERigVMOpCode::ArrayRemove:
+			case ERigVMOpCode::ArrayUnion:
 			{
 				FRigVMBinaryOp Op;
 				Ar << Op;
@@ -591,6 +597,8 @@ void FRigVMByteCode::Load(FArchive& Ar)
 			case ERigVMOpCode::ArrayGetAtIndex:
 			case ERigVMOpCode::ArraySetAtIndex:
 			case ERigVMOpCode::ArrayInsert:
+			case ERigVMOpCode::ArrayDifference:
+			case ERigVMOpCode::ArrayIntersection:
 			{
 				FRigVMTernaryOp Op;
 				Ar << Op;
@@ -807,6 +815,7 @@ uint64 FRigVMByteCode::GetOpNumBytesAt(uint64 InByteCodeIndex, bool bIncludeOper
 		case ERigVMOpCode::Increment:
 		case ERigVMOpCode::Decrement:
 		case ERigVMOpCode::ArrayReset:
+		case ERigVMOpCode::ArrayReverse:
 		{
 			return (uint64)sizeof(FRigVMUnaryOp);
 		}
@@ -846,6 +855,7 @@ uint64 FRigVMByteCode::GetOpNumBytesAt(uint64 InByteCodeIndex, bool bIncludeOper
 		case ERigVMOpCode::ArrayAppend:
 		case ERigVMOpCode::ArrayClone:
 		case ERigVMOpCode::ArrayRemove:
+		case ERigVMOpCode::ArrayUnion:
 		{
 			return (uint64)sizeof(FRigVMBinaryOp);
 		}
@@ -853,6 +863,8 @@ uint64 FRigVMByteCode::GetOpNumBytesAt(uint64 InByteCodeIndex, bool bIncludeOper
 		case ERigVMOpCode::ArrayGetAtIndex:
 		case ERigVMOpCode::ArraySetAtIndex:
 		case ERigVMOpCode::ArrayInsert:
+		case ERigVMOpCode::ArrayDifference:
+		case ERigVMOpCode::ArrayIntersection:
 		{				
 			return (uint64)sizeof(FRigVMTernaryOp);
 		}
@@ -912,6 +924,7 @@ uint64 FRigVMByteCode::AddCopyOp(const FRigVMOperand& InSource, const FRigVMOper
 uint64 FRigVMByteCode::AddCopyOp(const FRigVMOperand& InSource, const FRigVMOperand& InTarget)
 {
 	check(InTarget.GetMemoryType() != ERigVMMemoryType::Literal);
+	check(InSource != InTarget);
 
 	FRigVMCopyOp Op(InSource, InTarget);
 	return AddOp(Op);
@@ -1082,7 +1095,7 @@ FString FRigVMByteCode::DumpToText() const
 					for (const FRigVMOperand& Operand : Operands)
 					{
 						FString OperandContent;
-						FRigVMOperand::StaticStruct()->ExportText(OperandContent, &Operand, nullptr, nullptr, PPF_None, nullptr);
+						FRigVMOperand::StaticStruct()->ExportText(OperandContent, &Operand, &Operand, nullptr, PPF_None, nullptr);
 						OperandsContent.Add(FString::Printf(TEXT("\t%s"), *OperandContent));
 					}
 
@@ -1094,10 +1107,10 @@ FString FRigVMByteCode::DumpToText() const
 			{
 				const FRigVMCopyOp& Op = GetOpAt<FRigVMCopyOp>(Instruction.ByteCodeIndex);
 				FString SourceContent;
-				FRigVMOperand::StaticStruct()->ExportText(SourceContent, &Op.Source, nullptr, nullptr, PPF_None, nullptr);
+				FRigVMOperand::StaticStruct()->ExportText(SourceContent, &Op.Source, &Op.Source, nullptr, PPF_None, nullptr);
 				Line += FString::Printf(TEXT(", Source %s"), *SourceContent);
 				FString TargetContent;
-				FRigVMOperand::StaticStruct()->ExportText(TargetContent, &Op.Source, nullptr, nullptr, PPF_None, nullptr);
+				FRigVMOperand::StaticStruct()->ExportText(TargetContent, &Op.Target, &Op.Target, nullptr, PPF_None, nullptr);
 				Line += FString::Printf(TEXT(", Target %s"), *TargetContent);
 #if UE_RIGVM_UCLASS_BASED_STORAGE_DISABLED
 				Line += FString::Printf(TEXT(", NumBytes %d"), (int32)Op.NumBytes);
@@ -1110,10 +1123,11 @@ FString FRigVMByteCode::DumpToText() const
 			case ERigVMOpCode::Increment:
 			case ERigVMOpCode::Decrement:
 			case ERigVMOpCode::ArrayReset:
+			case ERigVMOpCode::ArrayReverse:
 			{
 				const FRigVMUnaryOp& Op = GetOpAt<FRigVMUnaryOp>(Instruction.ByteCodeIndex);
 				FString ArgContent;
-				FRigVMOperand::StaticStruct()->ExportText(ArgContent, &Op.Arg, nullptr, nullptr, PPF_None, nullptr);
+				FRigVMOperand::StaticStruct()->ExportText(ArgContent, &Op.Arg, &Op.Arg, nullptr, PPF_None, nullptr);
 				Line += FString::Printf(TEXT(", Source %s"), *ArgContent);
 				break;
 			}
@@ -1122,13 +1136,13 @@ FString FRigVMByteCode::DumpToText() const
 			{
 				const FRigVMComparisonOp& Op = GetOpAt<FRigVMComparisonOp>(Instruction.ByteCodeIndex);
 				FString AContent;
-				FRigVMOperand::StaticStruct()->ExportText(AContent, &Op.A, nullptr, nullptr, PPF_None, nullptr);
+				FRigVMOperand::StaticStruct()->ExportText(AContent, &Op.A, &Op.A, nullptr, PPF_None, nullptr);
 				Line += FString::Printf(TEXT(", A %s"), *AContent);
 				FString BContent;
-				FRigVMOperand::StaticStruct()->ExportText(BContent, &Op.B, nullptr, nullptr, PPF_None, nullptr);
+				FRigVMOperand::StaticStruct()->ExportText(BContent, &Op.B, &Op.B, nullptr, PPF_None, nullptr);
 				Line += FString::Printf(TEXT(", B %s"), *BContent);
 				FString ResultContent;
-				FRigVMOperand::StaticStruct()->ExportText(ResultContent, &Op.Result, nullptr, nullptr, PPF_None, nullptr);
+				FRigVMOperand::StaticStruct()->ExportText(ResultContent, &Op.Result, &Op.Result, nullptr, PPF_None, nullptr);
 				Line += FString::Printf(TEXT(", Result %s"), *ResultContent);
 				break;
 			}
@@ -1147,7 +1161,7 @@ FString FRigVMByteCode::DumpToText() const
 				const FRigVMJumpIfOp& Op = GetOpAt<FRigVMJumpIfOp>(Instruction.ByteCodeIndex);
 				Line += FString::Printf(TEXT(", InstructionIndex %d"), (int32)Op.InstructionIndex);
 				FString ArgContent;
-				FRigVMOperand::StaticStruct()->ExportText(ArgContent, &Op.Arg, nullptr, nullptr, PPF_None, nullptr);
+				FRigVMOperand::StaticStruct()->ExportText(ArgContent, &Op.Arg, &Op.Arg, nullptr, PPF_None, nullptr);
 				Line += FString::Printf(TEXT(", Source %s"), *ArgContent);
 				Line += FString::Printf(TEXT(", Condition %d"), (int32)(Op.Condition ? 0 : 1));
 				break;
@@ -1162,11 +1176,12 @@ FString FRigVMByteCode::DumpToText() const
 			case ERigVMOpCode::ArrayAppend:
 			case ERigVMOpCode::ArrayClone:
 			case ERigVMOpCode::ArrayRemove:
+			case ERigVMOpCode::ArrayUnion:
 			{
 				const FRigVMBinaryOp& Op = GetOpAt<FRigVMBinaryOp>(Instruction.ByteCodeIndex);
 				FString ArgA, ArgB;
-				FRigVMOperand::StaticStruct()->ExportText(ArgA, &Op.ArgA, nullptr, nullptr, PPF_None, nullptr);
-				FRigVMOperand::StaticStruct()->ExportText(ArgB, &Op.ArgB, nullptr, nullptr, PPF_None, nullptr);
+				FRigVMOperand::StaticStruct()->ExportText(ArgA, &Op.ArgA, &Op.ArgA, nullptr, PPF_None, nullptr);
+				FRigVMOperand::StaticStruct()->ExportText(ArgB, &Op.ArgB, &Op.ArgB, nullptr, PPF_None, nullptr);
 				Line += FString::Printf(TEXT(", ArgA %s"), *ArgA);
 				Line += FString::Printf(TEXT(", ArgB %s"), *ArgB);
 				break;
@@ -1175,12 +1190,14 @@ FString FRigVMByteCode::DumpToText() const
 			case ERigVMOpCode::ArrayGetAtIndex:
 			case ERigVMOpCode::ArraySetAtIndex:
 			case ERigVMOpCode::ArrayInsert:
+			case ERigVMOpCode::ArrayDifference:
+			case ERigVMOpCode::ArrayIntersection:
 			{
 				const FRigVMTernaryOp& Op = GetOpAt<FRigVMTernaryOp>(Instruction.ByteCodeIndex);
 				FString ArgA, ArgB, ArgC;
-				FRigVMOperand::StaticStruct()->ExportText(ArgA, &Op.ArgA, nullptr, nullptr, PPF_None, nullptr);
-				FRigVMOperand::StaticStruct()->ExportText(ArgB, &Op.ArgB, nullptr, nullptr, PPF_None, nullptr);
-				FRigVMOperand::StaticStruct()->ExportText(ArgC, &Op.ArgC, nullptr, nullptr, PPF_None, nullptr);
+				FRigVMOperand::StaticStruct()->ExportText(ArgA, &Op.ArgA, &Op.ArgA, nullptr, PPF_None, nullptr);
+				FRigVMOperand::StaticStruct()->ExportText(ArgB, &Op.ArgB, &Op.ArgB, nullptr, PPF_None, nullptr);
+				FRigVMOperand::StaticStruct()->ExportText(ArgC, &Op.ArgC, &Op.ArgC, nullptr, PPF_None, nullptr);
 				Line += FString::Printf(TEXT(", ArgA %s"), *ArgA);
 				Line += FString::Printf(TEXT(", ArgB %s"), *ArgB);
 				Line += FString::Printf(TEXT(", ArgC %s"), *ArgC);
@@ -1190,10 +1207,10 @@ FString FRigVMByteCode::DumpToText() const
 			{
 				const FRigVMQuaternaryOp& Op = GetOpAt<FRigVMQuaternaryOp>(Instruction.ByteCodeIndex);
 				FString ArgA, ArgB, ArgC, ArgD;
-				FRigVMOperand::StaticStruct()->ExportText(ArgA, &Op.ArgA, nullptr, nullptr, PPF_None, nullptr);
-				FRigVMOperand::StaticStruct()->ExportText(ArgB, &Op.ArgB, nullptr, nullptr, PPF_None, nullptr);
-				FRigVMOperand::StaticStruct()->ExportText(ArgC, &Op.ArgC, nullptr, nullptr, PPF_None, nullptr);
-				FRigVMOperand::StaticStruct()->ExportText(ArgD, &Op.ArgD, nullptr, nullptr, PPF_None, nullptr);
+				FRigVMOperand::StaticStruct()->ExportText(ArgA, &Op.ArgA, &Op.ArgA, nullptr, PPF_None, nullptr);
+				FRigVMOperand::StaticStruct()->ExportText(ArgB, &Op.ArgB, &Op.ArgB, nullptr, PPF_None, nullptr);
+				FRigVMOperand::StaticStruct()->ExportText(ArgC, &Op.ArgC, &Op.ArgC, nullptr, PPF_None, nullptr);
+				FRigVMOperand::StaticStruct()->ExportText(ArgD, &Op.ArgD, &Op.ArgD, nullptr, PPF_None, nullptr);
 				Line += FString::Printf(TEXT(", ArgA %s"), *ArgA);
 				Line += FString::Printf(TEXT(", ArgB %s"), *ArgB);
 				Line += FString::Printf(TEXT(", ArgC %s"), *ArgC);
@@ -1204,12 +1221,12 @@ FString FRigVMByteCode::DumpToText() const
 			{
 				const FRigVMSenaryOp& Op = GetOpAt<FRigVMSenaryOp>(Instruction.ByteCodeIndex);
 				FString ArgA, ArgB, ArgC, ArgD, ArgE, ArgF;
-				FRigVMOperand::StaticStruct()->ExportText(ArgA, &Op.ArgA, nullptr, nullptr, PPF_None, nullptr);
-				FRigVMOperand::StaticStruct()->ExportText(ArgB, &Op.ArgB, nullptr, nullptr, PPF_None, nullptr);
-				FRigVMOperand::StaticStruct()->ExportText(ArgC, &Op.ArgC, nullptr, nullptr, PPF_None, nullptr);
-				FRigVMOperand::StaticStruct()->ExportText(ArgD, &Op.ArgD, nullptr, nullptr, PPF_None, nullptr);
-				FRigVMOperand::StaticStruct()->ExportText(ArgE, &Op.ArgE, nullptr, nullptr, PPF_None, nullptr);
-				FRigVMOperand::StaticStruct()->ExportText(ArgF, &Op.ArgF, nullptr, nullptr, PPF_None, nullptr);
+				FRigVMOperand::StaticStruct()->ExportText(ArgA, &Op.ArgA, &Op.ArgA, nullptr, PPF_None, nullptr);
+				FRigVMOperand::StaticStruct()->ExportText(ArgB, &Op.ArgB, &Op.ArgB, nullptr, PPF_None, nullptr);
+				FRigVMOperand::StaticStruct()->ExportText(ArgC, &Op.ArgC, &Op.ArgC, nullptr, PPF_None, nullptr);
+				FRigVMOperand::StaticStruct()->ExportText(ArgD, &Op.ArgD, &Op.ArgD, nullptr, PPF_None, nullptr);
+				FRigVMOperand::StaticStruct()->ExportText(ArgE, &Op.ArgE, &Op.ArgE, nullptr, PPF_None, nullptr);
+				FRigVMOperand::StaticStruct()->ExportText(ArgF, &Op.ArgF, &Op.ArgF, nullptr, PPF_None, nullptr);
 				Line += FString::Printf(TEXT(", ArgA %s"), *ArgA);
 				Line += FString::Printf(TEXT(", ArgB %s"), *ArgB);
 				Line += FString::Printf(TEXT(", ArgC %s"), *ArgC);
@@ -1311,6 +1328,26 @@ uint64 FRigVMByteCode::AddArrayIteratorOp(FRigVMOperand InArrayArg, FRigVMOperan
 	return AddOp(FRigVMSenaryOp(ERigVMOpCode::ArrayIterator, InArrayArg, InElementArg, InIndexArg, InCountArg, InRatioArg, InContinueArg));
 }
 
+uint64 FRigVMByteCode::AddArrayUnionOp(FRigVMOperand InArrayArg, FRigVMOperand InOtherArrayArg)
+{
+	return AddOp(FRigVMBinaryOp(ERigVMOpCode::ArrayUnion, InArrayArg, InOtherArrayArg));
+}
+	
+uint64 FRigVMByteCode::AddArrayDifferenceOp(FRigVMOperand InArrayArg, FRigVMOperand InOtherArrayArg, FRigVMOperand InResultArrayArg)
+{
+	return AddOp(FRigVMTernaryOp(ERigVMOpCode::ArrayDifference, InArrayArg, InOtherArrayArg, InResultArrayArg));
+}
+	
+uint64 FRigVMByteCode::AddArrayIntersectionOp(FRigVMOperand InArrayArg, FRigVMOperand InOtherArrayArg, FRigVMOperand InResultArrayArg)
+{
+	return AddOp(FRigVMTernaryOp(ERigVMOpCode::ArrayIntersection, InArrayArg, InOtherArrayArg, InResultArrayArg));
+}
+	
+uint64 FRigVMByteCode::AddArrayReverseOp(FRigVMOperand InArrayArg)
+{
+	return AddOp(FRigVMUnaryOp(ERigVMOpCode::ArrayReverse, InArrayArg));
+}
+
 uint64 FRigVMByteCode::GetOpAlignment(ERigVMOpCode InOpCode) const
 {
 	switch (InOpCode)
@@ -1395,6 +1432,7 @@ uint64 FRigVMByteCode::GetOpAlignment(ERigVMOpCode InOpCode) const
 		case ERigVMOpCode::Increment:
 		case ERigVMOpCode::Decrement:
 		case ERigVMOpCode::ArrayReset:
+		case ERigVMOpCode::ArrayReverse:
 		{
 			static const uint64 Alignment = FRigVMUnaryOp::StaticStruct()->GetCppStructOps()->GetAlignment();
 			return Alignment;
@@ -1440,6 +1478,7 @@ uint64 FRigVMByteCode::GetOpAlignment(ERigVMOpCode InOpCode) const
 		case ERigVMOpCode::ArrayAppend:
 		case ERigVMOpCode::ArrayClone:
 		case ERigVMOpCode::ArrayRemove:
+		case ERigVMOpCode::ArrayUnion:
 		{
 			static const uint64 Alignment = FRigVMBinaryOp::StaticStruct()->GetCppStructOps()->GetAlignment();
 			return Alignment;
@@ -1448,6 +1487,8 @@ uint64 FRigVMByteCode::GetOpAlignment(ERigVMOpCode InOpCode) const
 		case ERigVMOpCode::ArrayGetAtIndex:
 		case ERigVMOpCode::ArraySetAtIndex:
 		case ERigVMOpCode::ArrayInsert:
+		case ERigVMOpCode::ArrayDifference:
+		case ERigVMOpCode::ArrayIntersection:
 		{				
 			static const uint64 Alignment = FRigVMTernaryOp::StaticStruct()->GetCppStructOps()->GetAlignment();
 			return Alignment;
diff --git a/Engine/Source/Runtime/RigVM/Private/RigVMCore/RigVMMemoryStorage.cpp b/Engine/Source/Runtime/RigVM/Private/RigVMCore/RigVMMemoryStorage.cpp
--- a/Engine/Source/Runtime/RigVM/Private/RigVMCore/RigVMMemoryStorage.cpp
+++ b/Engine/Source/Runtime/RigVM/Private/RigVMCore/RigVMMemoryStorage.cpp
@@ -327,6 +327,26 @@ URigVMMemoryStorageGeneratorClass* URigVMMemoryStorageGeneratorClass::CreateStor
 	return Class;
 }
 
+bool URigVMMemoryStorageGeneratorClass::RemoveStorageClass(UObject* InOuter, ERigVMMemoryType InMemoryType)
+{
+	check(InOuter);
+	UPackage* Package = InOuter->GetOutermost();
+
+	const FString ClassName = GetClassName(InMemoryType);
+
+	// if there's an old class - remove it from the package and mark it to destroy
+	URigVMMemoryStorageGeneratorClass* OldClass = FindObject<URigVMMemoryStorageGeneratorClass>(Package, *ClassName);
+	if(OldClass)
+	{
+		OldClass->RemoveFromRoot();
+		OldClass->Rename(nullptr, GetTransientPackage(), REN_ForceNoResetLoaders | REN_DoNotDirty | REN_DontCreateRedirectors | REN_NonTransactional);
+		OldClass->MarkPendingKill();
+		return true;
+	}
+
+	return false;
+}
+
 FRigVMMemoryStatistics URigVMMemoryStorageGeneratorClass::GetStatistics() const
 {
 	FRigVMMemoryStatistics Statistics;
diff --git a/Engine/Source/Runtime/RigVM/Public/RigVMCore/RigVM.h b/Engine/Source/Runtime/RigVM/Public/RigVMCore/RigVM.h
--- a/Engine/Source/Runtime/RigVM/Public/RigVMCore/RigVM.h
+++ b/Engine/Source/Runtime/RigVM/Public/RigVMCore/RigVM.h
@@ -849,6 +849,7 @@ private:
 	UScriptStruct* GetScripStructForOperand(const FRigVMOperand& InOperand) const;
 #else
 	void RefreshExternalPropertyPaths();
+	static void CopyArray(FScriptArrayHelper& TargetHelper, FRigVMMemoryHandle& TargetHandle, FScriptArrayHelper& SourceHelper, FRigVMMemoryHandle& SourceHandle);
 #endif
 	
 	TMap<FRigVMOperand, TArray<FRigVMOperand>> OperandToDebugRegisters;
diff --git a/Engine/Source/Runtime/RigVM/Public/RigVMCore/RigVMByteCode.h b/Engine/Source/Runtime/RigVM/Public/RigVMCore/RigVMByteCode.h
--- a/Engine/Source/Runtime/RigVM/Public/RigVMCore/RigVMByteCode.h
+++ b/Engine/Source/Runtime/RigVM/Public/RigVMCore/RigVMByteCode.h
@@ -99,7 +99,7 @@ enum class ERigVMOpCode : uint8
 	Exit, // exit the execution loop
 	BeginBlock, // begins a new memory slice / block
 	EndBlock, // ends the last memory slice / block
-	ArrayReset, // reads and returns the size of an array (unary op, in out array)
+	ArrayReset, // clears an array and resets its content
 	ArrayGetNum, // reads and returns the size of an array (binary op, in array, out int32) 
 	ArraySetNum, // resizes an array (binary op, in out array, in int32)
 	ArrayGetAtIndex, // returns an array element by index (ternary op, in array, in int32, out element)  
@@ -111,7 +111,13 @@ enum class ERigVMOpCode : uint8
 	ArrayAppend, // appends an array to another (binary op, in out array, in array)
 	ArrayClone, // clones an array (binary op, in array, out array)
 	ArrayIterator, // iterates over an array (senary op, in array, out element, out index, out count, out ratio, out continue)
-	Invalid
+	ArrayUnion, // merges two arrays while avoiding duplicates (binary op, in out array, in other array)
+	ArrayDifference, // returns a new array containing elements only found in one array (ternary op, in array, in array, out result)
+	ArrayIntersection, // returns a new array containing elements found in both of the input arrays (ternary op, in array, in array, out result)
+	ArrayReverse, // returns the reverse of the input array (unary op, in out array)
+	Invalid,
+	FirstArrayOpCode = ArrayReset,
+	LastArrayOpCode = ArrayReverse,
 };
 
 // Base class for all VM operations
@@ -185,7 +191,8 @@ struct RIGVM_API FRigVMUnaryOp : public FRigVMBaseOp
 			uint8(InOpCode) == uint8(ERigVMOpCode::JumpForwardIf) ||
 			uint8(InOpCode) == uint8(ERigVMOpCode::JumpBackwardIf) ||
 			uint8(InOpCode) == uint8(ERigVMOpCode::ChangeType) ||
-			uint8(InOpCode) == uint8(ERigVMOpCode::ArrayReset)
+			uint8(InOpCode) == uint8(ERigVMOpCode::ArrayReset) ||
+			uint8(InOpCode) == uint8(ERigVMOpCode::ArrayReverse)
 		);
 	}
 
@@ -223,7 +230,8 @@ struct RIGVM_API FRigVMBinaryOp : public FRigVMBaseOp
 			uint8(InOpCode) == uint8(ERigVMOpCode::ArraySetNum) ||
 			uint8(InOpCode) == uint8(ERigVMOpCode::ArrayAppend) ||
 			uint8(InOpCode) == uint8(ERigVMOpCode::ArrayClone) ||
-			uint8(InOpCode) == uint8(ERigVMOpCode::ArrayRemove)
+			uint8(InOpCode) == uint8(ERigVMOpCode::ArrayRemove) ||
+			uint8(InOpCode) == uint8(ERigVMOpCode::ArrayUnion)
 		);
 	}
 
@@ -262,7 +270,9 @@ struct RIGVM_API FRigVMTernaryOp : public FRigVMBaseOp
 			uint8(InOpCode) == uint8(ERigVMOpCode::ArrayAdd) ||
 			uint8(InOpCode) == uint8(ERigVMOpCode::ArrayGetAtIndex) ||
 			uint8(InOpCode) == uint8(ERigVMOpCode::ArraySetAtIndex) ||
-			uint8(InOpCode) == uint8(ERigVMOpCode::ArrayInsert)
+			uint8(InOpCode) == uint8(ERigVMOpCode::ArrayInsert) ||
+			uint8(InOpCode) == uint8(ERigVMOpCode::ArrayDifference) ||
+			uint8(InOpCode) == uint8(ERigVMOpCode::ArrayIntersection)
 		);
 	}
 
@@ -457,6 +467,15 @@ public:
 	{
 	}
 
+	bool IsValid() const
+	{
+		return
+			Source.IsValid() &&
+			Target.IsValid() &&
+			(Source != Target) &&
+			(Target.GetMemoryType() != ERigVMMemoryType::Literal);
+	}
+
 	FRigVMOperand Source;
 	FRigVMOperand Target;
 #if !UE_RIGVM_UCLASS_BASED_STORAGE_DISABLED
@@ -850,7 +869,19 @@ public:
 
 	// adds an array iterator operator
 	uint64 AddArrayIteratorOp(FRigVMOperand InArrayArg, FRigVMOperand InElementArg, FRigVMOperand InIndexArg, FRigVMOperand InCountArg, FRigVMOperand InRatioArg, FRigVMOperand InContinueArg);
+
+	// adds an array union operator
+	uint64 AddArrayUnionOp(FRigVMOperand InArrayArg, FRigVMOperand InOtherArrayArg);
 	
+	// adds an array difference operator
+	uint64 AddArrayDifferenceOp(FRigVMOperand InArrayArg, FRigVMOperand InOtherArrayArg, FRigVMOperand InResultArrayArg);
+	
+	// adds an array intersection operator
+	uint64 AddArrayIntersectionOp(FRigVMOperand InArrayArg, FRigVMOperand InOtherArrayArg, FRigVMOperand InResultArrayArg);
+	
+	// adds an array reverse operator
+	uint64 AddArrayReverseOp(FRigVMOperand InArrayArg);
+
 	// returns an instruction array for iterating over all operators
 	FORCEINLINE FRigVMInstructionArray GetInstructions() const
 	{
diff --git a/Engine/Source/Runtime/RigVM/Public/RigVMCore/RigVMExecuteContext.h b/Engine/Source/Runtime/RigVM/Public/RigVMCore/RigVMExecuteContext.h
--- a/Engine/Source/Runtime/RigVM/Public/RigVMCore/RigVMExecuteContext.h
+++ b/Engine/Source/Runtime/RigVM/Public/RigVMCore/RigVMExecuteContext.h
@@ -10,6 +10,8 @@
 #include "Logging/TokenizedMessage.h"
 #include "RigVMExecuteContext.generated.h"
 
+struct FRigVMExecuteContext;
+
 USTRUCT()
 struct RIGVM_API FRigVMSlice
 {
@@ -164,7 +166,7 @@ struct RIGVM_API FRigVMRuntimeSettings
 	/*
 	 * The function to use for logging anything from the VM to the host
 	 */
-	TFunction<void(EMessageSeverity::Type,int32,const FString&)> LogFunction = nullptr;
+	TFunction<void(EMessageSeverity::Type,const FRigVMExecuteContext*,const FString&)> LogFunction = nullptr;
 };
 
 /**
@@ -269,7 +271,7 @@ struct RIGVM_API FRigVMExecuteContext
 	{
 		if(RuntimeSettings.LogFunction)
 		{
-			RuntimeSettings.LogFunction(InSeverity, InstructionIndex, InMessage);
+			RuntimeSettings.LogFunction(InSeverity, this, InMessage);
 		}
 		else
 		{
diff --git a/Engine/Source/Runtime/RigVM/Public/RigVMCore/RigVMMemoryStorage.h b/Engine/Source/Runtime/RigVM/Public/RigVMCore/RigVMMemoryStorage.h
--- a/Engine/Source/Runtime/RigVM/Public/RigVMCore/RigVMMemoryStorage.h
+++ b/Engine/Source/Runtime/RigVM/Public/RigVMCore/RigVMMemoryStorage.h
@@ -246,6 +246,15 @@ public:
 		const TArray<FRigVMPropertyDescription>& InProperties,
 		const TArray<FRigVMPropertyPathDescription>& InPropertyPaths = TArray<FRigVMPropertyPathDescription>());
 
+	/**
+	 * Removes an existing storage class
+	 * @param InOuter The package or the outer object of this class
+	 * @param InMemoryType The memory type to use for this class (affects class name)
+	 */
+	static bool RemoveStorageClass(
+		UObject* InOuter,
+		ERigVMMemoryType InMemoryType);
+
 	// Returns the type of memory of this class (literal, work, etc)
 	ERigVMMemoryType GetMemoryType() const { return MemoryType; }
 
