Control Rig: Implement in editor profiling

diff --git a/Engine/Plugins/Animation/ControlRig/Source/ControlRig/Public/ControlRigDefines.h b/Engine/Plugins/Animation/ControlRig/Source/ControlRig/Public/ControlRigDefines.h
--- a/Engine/Plugins/Animation/ControlRig/Source/ControlRig/Public/ControlRigDefines.h
+++ b/Engine/Plugins/Animation/ControlRig/Source/ControlRig/Public/ControlRigDefines.h
@@ -9,8 +9,9 @@
 #include "RigVMCore/RigVMExecuteContext.h"
 #include "ControlRigDefines.generated.h"
 
-#define DECLARE_SCOPE_HIERARCHICAL_COUNTER_RIGUNIT() \
-	DECLARE_SCOPE_HIERARCHICAL_COUNTER_FUNC()
+#define DECLARE_SCOPE_HIERARCHICAL_COUNTER_RIGUNIT()
+//#define DECLARE_SCOPE_HIERARCHICAL_COUNTER_RIGUNIT() \
+//	DECLARE_SCOPE_HIERARCHICAL_COUNTER_FUNC()
 
 USTRUCT()
 struct FControlRigExecuteContext : public FRigVMExecuteContext
diff --git a/Engine/Plugins/Animation/ControlRig/Source/ControlRigDeveloper/Private/ControlRigBlueprint.cpp b/Engine/Plugins/Animation/ControlRig/Source/ControlRigDeveloper/Private/ControlRigBlueprint.cpp
--- a/Engine/Plugins/Animation/ControlRig/Source/ControlRigDeveloper/Private/ControlRigBlueprint.cpp
+++ b/Engine/Plugins/Animation/ControlRig/Source/ControlRigDeveloper/Private/ControlRigBlueprint.cpp
@@ -530,6 +530,9 @@ void UControlRigBlueprint::RecompileVM()
 	
 	bErrorsDuringCompilation = false;
 
+	RigGraphDisplaySettings.MinMicroSeconds = RigGraphDisplaySettings.LastMinMicroSeconds = DBL_MAX;
+	RigGraphDisplaySettings.MaxMicroSeconds = RigGraphDisplaySettings.LastMaxMicroSeconds = (double)INDEX_NONE;
+
 	UControlRigBlueprintGeneratedClass* RigClass = GetControlRigBlueprintGeneratedClass();
 	UControlRig* CDO = Cast<UControlRig>(RigClass->GetDefaultObject(true /* create if needed */));
 	if (CDO->VM != nullptr)
@@ -1790,6 +1793,11 @@ void UControlRigBlueprint::PostTransacted(const FTransactionObjectEvent& Transac
 		{
 			PropagateDrawInstructionsFromBPToInstances();
 		}
+
+		else if (TransactionEvent.GetChangedProperties().Contains(TEXT("VMRuntimeSettings")))
+		{
+			PropagateRuntimeSettingsFromBPToInstances();
+		}
 	}
 }
 
@@ -3347,6 +3355,43 @@ void UControlRigBlueprint::PropagateDrawInstructionsFromBPToInstances()
 	}
 }
 
+void UControlRigBlueprint::PropagateRuntimeSettingsFromBPToInstances()
+{
+	if (UClass* MyControlRigClass = GeneratedClass)
+	{
+		if (UControlRig* DefaultObject = Cast<UControlRig>(MyControlRigClass->GetDefaultObject(false)))
+		{
+			DefaultObject->VMRuntimeSettings = VMRuntimeSettings;
+
+			TArray<UObject*> ArchetypeInstances;
+			DefaultObject->GetArchetypeInstances(ArchetypeInstances);
+
+			for (UObject* ArchetypeInstance : ArchetypeInstances)
+			{
+				if (UControlRig* InstanceRig = Cast<UControlRig>(ArchetypeInstance))
+				{
+					InstanceRig->VMRuntimeSettings = VMRuntimeSettings;
+				}
+			}
+		}
+	}
+
+	TArray<UEdGraph*> EdGraphs;
+	GetAllGraphs(EdGraphs);
+
+	for (UEdGraph* Graph : EdGraphs)
+	{
+		TArray<UEdGraphNode*> Nodes = Graph->Nodes;
+		for (UEdGraphNode* Node : Nodes)
+		{
+			if(UControlRigGraphNode* RigNode = Cast<UControlRigGraphNode>(Node))
+			{
+				RigNode->ReconstructNode_Internal(true);
+			}
+		}
+	}
+}
+
 void UControlRigBlueprint::PropagatePropertyFromBPToInstances(FRigElementKey InRigElement, const FProperty* InProperty)
 {
 	int32 ElementIndex = Hierarchy->GetIndex(InRigElement);
diff --git a/Engine/Plugins/Animation/ControlRig/Source/ControlRigDeveloper/Private/Graph/ControlRigGraphNode.cpp b/Engine/Plugins/Animation/ControlRig/Source/ControlRigDeveloper/Private/Graph/ControlRigGraphNode.cpp
--- a/Engine/Plugins/Animation/ControlRig/Source/ControlRigDeveloper/Private/Graph/ControlRigGraphNode.cpp
+++ b/Engine/Plugins/Animation/ControlRig/Source/ControlRigDeveloper/Private/Graph/ControlRigGraphNode.cpp
@@ -41,6 +41,9 @@ UControlRigGraphNode::UControlRigGraphNode()
 , NodeTopologyVersion(INDEX_NONE)
 , CachedTitleColor(FLinearColor(0.f, 0.f, 0.f, 0.f))
 , CachedNodeColor(FLinearColor(0.f, 0.f, 0.f, 0.f))
+#if WITH_EDITOR
+, bEnableProfiling(false)
+#endif
 {
 	bHasCompilerMessage = false;
 	ErrorType = (int32)EMessageSeverity::Info + 1;
@@ -181,6 +184,17 @@ void UControlRigGraphNode::ReconstructNode_Internal(bool bForce)
 		}
 	}
 
+#if WITH_EDITOR
+	bEnableProfiling = false;
+	if(RigGraph)
+	{
+		if(UControlRigBlueprint* RigBlueprint = RigGraph->GetBlueprint())
+		{
+			bEnableProfiling = RigBlueprint->VMRuntimeSettings.bEnableProfiling;
+		}
+	}
+#endif
+
 	// Clear previously set messages
 	ErrorMsg.Reset();
 
@@ -370,6 +384,53 @@ int32 UControlRigGraphNode::GetInstructionIndex()
 	return INDEX_NONE;
 }
 
+FLinearColor UControlRigGraphNode::GetNodeProfilingColor() const
+{
+#if WITH_EDITOR
+	if(bEnableProfiling)
+	{
+		if(UControlRigBlueprint* Blueprint = Cast<UControlRigBlueprint>(GetOuter()->GetOuter()))
+		{
+			if(UControlRig* DebuggedControlRig = Cast<UControlRig>(Blueprint->GetObjectBeingDebugged()))
+			{
+				if(URigVMNode* ModelNode = GetModelNode())
+				{
+					const double MicroSeconds = ModelNode->GetInstructionMicroSeconds(DebuggedControlRig->GetVM(), FRigVMASTProxy());
+					if(MicroSeconds >= 0.0)
+					{
+						if(Blueprint->RigGraphDisplaySettings.bAutoDetermineRange)
+						{
+							if(MicroSeconds < Blueprint->RigGraphDisplaySettings.MinMicroSeconds)
+							{
+								Blueprint->RigGraphDisplaySettings.MinMicroSeconds = MicroSeconds;
+							}
+							if(MicroSeconds > Blueprint->RigGraphDisplaySettings.MaxMicroSeconds)
+							{
+								Blueprint->RigGraphDisplaySettings.MaxMicroSeconds = MicroSeconds;
+							}
+						}
+							
+						const double MinMicroSeconds = Blueprint->RigGraphDisplaySettings.LastMinMicroSeconds;
+						const double MaxMicroSeconds = Blueprint->RigGraphDisplaySettings.LastMaxMicroSeconds;
+						if(MaxMicroSeconds <= MinMicroSeconds)
+						{
+							return FLinearColor::Black;
+						}
+			
+						const FLinearColor& MinColor = Blueprint->RigGraphDisplaySettings.MinDurationColor;
+						const FLinearColor& MaxColor = Blueprint->RigGraphDisplaySettings.MaxDurationColor;
+
+						const double T = (MicroSeconds - MinMicroSeconds) / (MaxMicroSeconds - MinMicroSeconds);
+						return FMath::Lerp<FLinearColor>(MinColor, MaxColor, (float)T);
+					}
+				}
+			}
+		}
+	}
+#endif
+	return FLinearColor::Black;
+}
+
 void UControlRigGraphNode::AllocateDefaultPins()
 {
 	ExecutePins.Reset();
@@ -641,6 +702,13 @@ FLinearColor UControlRigGraphNode::GetNodeTitleColor() const
 
 FLinearColor UControlRigGraphNode::GetNodeBodyTintColor() const
 {
+#if WITH_EDITOR
+	if(bEnableProfiling)
+	{
+		return GetNodeProfilingColor();
+	}
+#endif
+	
 	return CachedNodeColor * GetNodeOpacityColor();
 }
 
diff --git a/Engine/Plugins/Animation/ControlRig/Source/ControlRigDeveloper/Public/ControlRigBlueprint.h b/Engine/Plugins/Animation/ControlRig/Source/ControlRigDeveloper/Public/ControlRigBlueprint.h
--- a/Engine/Plugins/Animation/ControlRig/Source/ControlRigDeveloper/Public/ControlRigBlueprint.h
+++ b/Engine/Plugins/Animation/ControlRig/Source/ControlRigDeveloper/Public/ControlRigBlueprint.h
@@ -112,6 +112,13 @@ struct CONTROLRIGDEVELOPER_API FRigGraphDisplaySettings
 		: bShowNodeRunCounts(false)
 		, NodeRunLowerBound(1)
 		, NodeRunLimit(64)
+		, MinMicroSeconds(0.0)
+		, MaxMicroSeconds(1.0)
+		, bAutoDetermineRange(true)
+		, LastMinMicroSeconds(0.0)
+		, LastMaxMicroSeconds(1.0)
+		, MinDurationColor(FLinearColor::Green)
+		, MaxDurationColor(FLinearColor::Red)
 	{
 	}
 
@@ -136,6 +143,31 @@ struct CONTROLRIGDEVELOPER_API FRigGraphDisplaySettings
 	// Note: The count limit doesn't apply to functions / collapse nodes.
 	UPROPERTY(EditAnywhere, Category = "Graph Display Settings")
 	int32 NodeRunLimit;
+
+	// The duration in microseconds of the fastest instruction / node
+	UPROPERTY(EditAnywhere, Category = "Graph Display Settings", transient, meta = (EditCondition = "!bAutoDetermineRange"))
+	double MinMicroSeconds;
+
+	// The duration in microseconds of the slowest instruction / node
+	UPROPERTY(EditAnywhere, Category = "Graph Display Settings", transient, meta = (EditCondition = "!bAutoDetermineRange"))
+	double MaxMicroSeconds;
+
+	UPROPERTY(EditAnywhere, Category = "Graph Display Settings")
+	bool bAutoDetermineRange;
+
+	UPROPERTY(transient)
+	double LastMinMicroSeconds;
+
+	UPROPERTY(transient)
+	double LastMaxMicroSeconds;
+
+	// The color of the fastest instruction / node
+	UPROPERTY(EditAnywhere, Category = "Graph Display Settings")
+	FLinearColor MinDurationColor;
+
+	// The color of the slowest instruction / node
+	UPROPERTY(EditAnywhere, Category = "Graph Display Settings")
+	FLinearColor MaxDurationColor;
 };
 
 USTRUCT()
@@ -503,6 +535,7 @@ public:
 	void PropagatePoseFromBPToInstances();
 	void PropagateHierarchyFromBPToInstances();
 	void PropagateDrawInstructionsFromBPToInstances();
+	void PropagateRuntimeSettingsFromBPToInstances();
 	void PropagatePropertyFromBPToInstances(FRigElementKey InRigElement, const FProperty* InProperty);
 	void PropagatePropertyFromInstanceToBP(FRigElementKey InRigElement, const FProperty* InProperty, UControlRig* InInstance);
 
diff --git a/Engine/Plugins/Animation/ControlRig/Source/ControlRigDeveloper/Public/Graph/ControlRigGraphNode.h b/Engine/Plugins/Animation/ControlRig/Source/ControlRigDeveloper/Public/Graph/ControlRigGraphNode.h
--- a/Engine/Plugins/Animation/ControlRig/Source/ControlRigDeveloper/Public/Graph/ControlRigGraphNode.h
+++ b/Engine/Plugins/Animation/ControlRig/Source/ControlRigDeveloper/Public/Graph/ControlRigGraphNode.h
@@ -153,6 +153,7 @@ public:
 
 protected:
 
+	FLinearColor GetNodeProfilingColor() const;
 	FLinearColor GetNodeOpacityColor() const;
 
 	/** Helper function for AllocateDefaultPins */
@@ -195,6 +196,10 @@ private:
 	FLinearColor CachedTitleColor;
 	FLinearColor CachedNodeColor;
 
+#if WITH_EDITOR
+	bool bEnableProfiling;
+#endif
+
 	TArray<URigVMPin*> ExecutePins;
 	TArray<URigVMPin*> InputOutputPins;
 	TArray<URigVMPin*> InputPins;
diff --git a/Engine/Plugins/Animation/ControlRig/Source/ControlRigEditor/Private/Editor/ControlRigEditor.cpp b/Engine/Plugins/Animation/ControlRig/Source/ControlRigEditor/Private/Editor/ControlRigEditor.cpp
--- a/Engine/Plugins/Animation/ControlRig/Source/ControlRigEditor/Private/Editor/ControlRigEditor.cpp
+++ b/Engine/Plugins/Animation/ControlRig/Source/ControlRigEditor/Private/Editor/ControlRigEditor.cpp
@@ -2827,6 +2827,32 @@ void FControlRigEditor::HandleControlRigExecutedEvent(UControlRig* InControlRig,
 				}
 			}
 		}
+
+		if(ControlRigBP->VMRuntimeSettings.bEnableProfiling)
+		{
+			if(ControlRigBP->RigGraphDisplaySettings.bAutoDetermineRange)
+			{
+				if(ControlRigBP->RigGraphDisplaySettings.LastMaxMicroSeconds < 0.0)
+				{
+					ControlRigBP->RigGraphDisplaySettings.LastMinMicroSeconds = ControlRigBP->RigGraphDisplaySettings.MinMicroSeconds; 
+					ControlRigBP->RigGraphDisplaySettings.LastMaxMicroSeconds = ControlRigBP->RigGraphDisplaySettings.MaxMicroSeconds;
+				}
+				else if(ControlRigBP->RigGraphDisplaySettings.MaxMicroSeconds >= 0.0)
+				{
+					const double T = 0.05;
+					ControlRigBP->RigGraphDisplaySettings.LastMinMicroSeconds = FMath::Lerp<double>(ControlRigBP->RigGraphDisplaySettings.LastMinMicroSeconds, ControlRigBP->RigGraphDisplaySettings.MinMicroSeconds, T); 
+					ControlRigBP->RigGraphDisplaySettings.LastMaxMicroSeconds = FMath::Lerp<double>(ControlRigBP->RigGraphDisplaySettings.LastMaxMicroSeconds, ControlRigBP->RigGraphDisplaySettings.MaxMicroSeconds, T); 
+				}
+
+				ControlRigBP->RigGraphDisplaySettings.MinMicroSeconds = DBL_MAX; 
+				ControlRigBP->RigGraphDisplaySettings.MaxMicroSeconds = (double)INDEX_NONE;
+			}
+			else
+			{
+				ControlRigBP->RigGraphDisplaySettings.LastMinMicroSeconds = ControlRigBP->RigGraphDisplaySettings.MinMicroSeconds; 
+				ControlRigBP->RigGraphDisplaySettings.LastMaxMicroSeconds = ControlRigBP->RigGraphDisplaySettings.MaxMicroSeconds;
+			}
+		}
 	}
 
 	UpdateGraphCompilerErrors();
@@ -4191,6 +4217,12 @@ void FControlRigEditor::OnFinishedChangingProperties(const FPropertyChangedEvent
 			return;
 		}
 
+		if (PropertyChangedEvent.MemberProperty->GetNameCPP() == TEXT("VMRuntimeSettings"))
+		{
+			ControlRigBP->PropagateRuntimeSettingsFromBPToInstances();
+			return;
+		}
+
 		if (PropertyChangedEvent.MemberProperty->GetNameCPP() == TEXT("DrawContainer"))
 		{
 			ControlRigBP->PropagateDrawInstructionsFromBPToInstances();
diff --git a/Engine/Plugins/Animation/ControlRig/Source/ControlRigEditor/Private/Editor/SControlRigStackView.cpp b/Engine/Plugins/Animation/ControlRig/Source/ControlRigEditor/Private/Editor/SControlRigStackView.cpp
--- a/Engine/Plugins/Animation/ControlRig/Source/ControlRigEditor/Private/Editor/SControlRigStackView.cpp
+++ b/Engine/Plugins/Animation/ControlRig/Source/ControlRigEditor/Private/Editor/SControlRigStackView.cpp
@@ -130,6 +130,16 @@ void SRigStackItem::Construct(const FArguments& InArgs, const TSharedRef<STableV
                 SNew(STextBlock)
                 .Text(this, &SRigStackItem::GetVisitedCountText)
                 .Font(IDetailLayoutBuilder::GetDetailFont())
+            ]
+
+			+ SHorizontalBox::Slot()
+			.MaxWidth(80.f)
+            .VAlign(VAlign_Center)
+            .HAlign(HAlign_Right)
+            [
+                SNew(STextBlock)
+                .Text(this, &SRigStackItem::GetDurationText)
+                .Font(IDetailLayoutBuilder::GetDetailFont())
             ]
         ], OwnerTable);
 }
@@ -179,6 +189,31 @@ FText SRigStackItem::GetVisitedCountText() const
 	return FText();
 }
 
+FText SRigStackItem::GetDurationText() const
+{
+	if(WeakStackEntry.IsValid() && WeakBlueprint.IsValid())
+	{
+		if(WeakBlueprint->VMRuntimeSettings.bEnableProfiling)
+		{
+			if(WeakStackEntry.Pin()->EntryType == ERigStackEntry::Operator)
+			{
+				if(UControlRig* ControlRig = Cast<UControlRig>(WeakBlueprint->GetObjectBeingDebugged()))
+				{
+					if(URigVM* VM = ControlRig->GetVM())
+					{
+						const double MicroSeconds = VM->GetInstructionMicroSeconds(WeakStackEntry.Pin()->InstructionIndex);
+						if(MicroSeconds > 0.0)
+						{
+							return FText::FromString(FString::Printf(TEXT("%d µs"), (int32)MicroSeconds));
+						}
+					}
+				}
+			}
+		}
+	}
+	return FText();
+}
+
 //////////////////////////////////////////////////////////////
 /// SControlRigStackView
 ///////////////////////////////////////////////////////////
diff --git a/Engine/Plugins/Animation/ControlRig/Source/ControlRigEditor/Private/Editor/SControlRigStackView.h b/Engine/Plugins/Animation/ControlRig/Source/ControlRigEditor/Private/Editor/SControlRigStackView.h
--- a/Engine/Plugins/Animation/ControlRig/Source/ControlRigEditor/Private/Editor/SControlRigStackView.h
+++ b/Engine/Plugins/Animation/ControlRig/Source/ControlRigEditor/Private/Editor/SControlRigStackView.h
@@ -59,6 +59,7 @@ private:
 	FText GetLabelText() const;
 	FSlateFontInfo GetLabelFont() const;
 	FText GetVisitedCountText() const;
+	FText GetDurationText() const;
 };
 
 class SControlRigStackView : public SCompoundWidget
diff --git a/Engine/Plugins/Animation/ControlRig/Source/ControlRigEditor/Private/Graph/SControlRigGraphNode.cpp b/Engine/Plugins/Animation/ControlRig/Source/ControlRigEditor/Private/Graph/SControlRigGraphNode.cpp
--- a/Engine/Plugins/Animation/ControlRig/Source/ControlRigEditor/Private/Graph/SControlRigGraphNode.cpp
+++ b/Engine/Plugins/Animation/ControlRig/Source/ControlRigEditor/Private/Graph/SControlRigGraphNode.cpp
@@ -498,6 +498,15 @@ void SControlRigGraphNode::Construct( const FArguments& InArgs )
 	.Visibility(EVisibility::Visible)
 	.ToolTipText(LOCTEXT("NodeHitCountToolTip", "This number represents the hit count for a node.\nFor functions / collapse nodes it represents the sum of all hit counts of contained nodes.\n\nYou can enable / disable the display of the number in the Class Settings\n(Rig Graph Display Settings -> Show Node Run Counts)"));
 
+	SAssignNew(InstructionDurationTextBlockWidget, STextBlock)
+	.Margin(FMargin(2.0f, 2.0f, 2.0f, 1.0f))
+	.Text(this, &SControlRigGraphNode::GetInstructionDurationText)
+	.Font(IDetailLayoutBuilder::GetDetailFont())
+	.ColorAndOpacity(FLinearColor::White)
+	.ShadowColorAndOpacity(FLinearColor(0.1f, 0.1f, 0.1f, 1.f))
+	.Visibility(EVisibility::Visible)
+	.ToolTipText(LOCTEXT("NodeDurationToolTip", "This number represents the duration in microseconds for a node.\nFor functions / collapse nodes it represents the accumulated time of contained nodes.\n\nYou can enable / disable the display of the number in the Class Settings\n(VM Runtime Settings -> Enable Profiling)"));
+
 	EdGraphNode->GetNodeTitleDirtied().BindSP(this, &SControlRigGraphNode::HandleNodeTitleDirtied);
 
 	LastHighDetailSize = FVector2D::ZeroVector;
@@ -654,6 +663,13 @@ void SControlRigGraphNode::AddPin(const TSharedRef<SGraphPin>& PinToAdd)
 
 const FSlateBrush * SControlRigGraphNode::GetNodeBodyBrush() const
 {
+	if(UControlRigGraphNode* RigNode = Cast<UControlRigGraphNode>(GetNodeObj()))
+	{
+		if(RigNode->bEnableProfiling)
+		{
+			return FEditorStyle::GetBrush("Graph.Node.TintedBody");
+		}
+	}
 	return FEditorStyle::GetBrush("Graph.Node.Body");
 }
 
@@ -1002,20 +1018,41 @@ TArray<FOverlayWidgetInfo> SControlRigGraphNode::GetOverlayWidgets(bool bSelecte
 
 		if(Blueprint.IsValid())
 		{
-			if(Blueprint->RigGraphDisplaySettings.bShowNodeRunCounts)
+			const bool bShowNodeCounts = Blueprint->RigGraphDisplaySettings.bShowNodeRunCounts;
+			const bool bEnableProfiling = Blueprint->VMRuntimeSettings.bEnableProfiling;
+			
+			if(bShowNodeCounts || bEnableProfiling)
 			{
 				if(UControlRig* DebuggedControlRig = Cast<UControlRig>(Blueprint->GetObjectBeingDebugged()))
 				{
-					const int32 Count = ModelNode->GetInstructionVisitedCount(DebuggedControlRig->GetVM(), FRigVMASTProxy(), false);
-					if(Count > Blueprint->RigGraphDisplaySettings.NodeRunLowerBound)
+					if(bShowNodeCounts)
 					{
-						const int32 VOffset = bSelected ? -2 : 2;
-						const FVector2D TextSize = InstructionCountTextBlockWidget->GetDesiredSize();
-						FOverlayWidgetInfo Info;
-						Info.OverlayOffset = FVector2D(WidgetSize.X - TextSize.X - 8.f, VOffset - TextSize.Y);
-						Info.Widget = InstructionCountTextBlockWidget;
-						Widgets.Add(Info);
+						const int32 Count = ModelNode->GetInstructionVisitedCount(DebuggedControlRig->GetVM(), FRigVMASTProxy());
+						if(Count > Blueprint->RigGraphDisplaySettings.NodeRunLowerBound)
+						{
+							const int32 VOffset = bSelected ? -2 : 2;
+							const FVector2D TextSize = InstructionCountTextBlockWidget->GetDesiredSize();
+							FOverlayWidgetInfo Info;
+							Info.OverlayOffset = FVector2D(WidgetSize.X - TextSize.X - 8.f, VOffset - TextSize.Y);
+							Info.Widget = InstructionCountTextBlockWidget;
+							Widgets.Add(Info);
+						}
+					}
+
+					if(bEnableProfiling)
+					{
+						const double MicroSeconds = ModelNode->GetInstructionMicroSeconds(DebuggedControlRig->GetVM(), FRigVMASTProxy());
+						if(MicroSeconds >= 0.0)
+						{
+							const int32 VOffset = bSelected ? -2 : 2;
+							const FVector2D TextSize = InstructionDurationTextBlockWidget->GetDesiredSize();
+							FOverlayWidgetInfo Info;
+							Info.OverlayOffset = FVector2D(8.f, VOffset - TextSize.Y);
+							Info.Widget = InstructionDurationTextBlockWidget;
+							Widgets.Add(Info);
+						}
 					}
+					
 				}
 			}
 		}
@@ -1071,7 +1108,7 @@ FText SControlRigGraphNode::GetInstructionCountText() const
 			{
 				if(UControlRig* DebuggedControlRig = Cast<UControlRig>(Blueprint->GetObjectBeingDebugged()))
 				{
-					const int32 Count = ModelNode->GetInstructionVisitedCount(DebuggedControlRig->GetVM(), FRigVMASTProxy(), true);
+					const int32 Count = ModelNode->GetInstructionVisitedCount(DebuggedControlRig->GetVM(), FRigVMASTProxy());
 					if(Count > Blueprint->RigGraphDisplaySettings.NodeRunLowerBound)
 					{
 						return FText::FromString(FString::FromInt(Count));
@@ -1084,6 +1121,29 @@ FText SControlRigGraphNode::GetInstructionCountText() const
 	return FText();
 }
 
+FText SControlRigGraphNode::GetInstructionDurationText() const
+{
+	if(Blueprint.IsValid())
+	{
+		if(Blueprint->VMRuntimeSettings.bEnableProfiling)
+		{
+			if (ModelNode.IsValid())
+			{
+				if(UControlRig* DebuggedControlRig = Cast<UControlRig>(Blueprint->GetObjectBeingDebugged()))
+				{
+					const double MicroSeconds = ModelNode->GetInstructionMicroSeconds(DebuggedControlRig->GetVM(), FRigVMASTProxy());
+					if(MicroSeconds >= 0)
+					{
+						return FText::FromString(FString::Printf(TEXT("%d µs"), (int32)MicroSeconds));
+					}
+				}
+			}
+		}
+	}
+
+	return FText();
+}
+
 int32 SControlRigGraphNode::GetNodeTopologyVersion() const
 {
 	if(UControlRigGraphNode* ControlRigGraphNode = Cast<UControlRigGraphNode>(GraphNode))
diff --git a/Engine/Plugins/Animation/ControlRig/Source/ControlRigEditor/Private/Graph/SControlRigGraphNode.h b/Engine/Plugins/Animation/ControlRig/Source/ControlRigEditor/Private/Graph/SControlRigGraphNode.h
--- a/Engine/Plugins/Animation/ControlRig/Source/ControlRigEditor/Private/Graph/SControlRigGraphNode.h
+++ b/Engine/Plugins/Animation/ControlRig/Source/ControlRigEditor/Private/Graph/SControlRigGraphNode.h
@@ -72,6 +72,7 @@ private:
 	void HandleNodeTitleDirtied();
 
 	FText GetInstructionCountText() const;
+	FText GetInstructionDurationText() const;
 
 private:
 
@@ -88,6 +89,7 @@ private:
 
 	TSharedPtr<SImage> VisualDebugIndicatorWidget;
 	TSharedPtr<STextBlock> InstructionCountTextBlockWidget;
+	TSharedPtr<STextBlock> InstructionDurationTextBlockWidget;
 
 	static const FSlateBrush* CachedImg_CR_Pin_Connected;
 	static const FSlateBrush* CachedImg_CR_Pin_Disconnected;
diff --git a/Engine/Source/Developer/RigVMDeveloper/Private/RigVMCompiler/RigVMCompiler.cpp b/Engine/Source/Developer/RigVMDeveloper/Private/RigVMCompiler/RigVMCompiler.cpp
--- a/Engine/Source/Developer/RigVMDeveloper/Private/RigVMCompiler/RigVMCompiler.cpp
+++ b/Engine/Source/Developer/RigVMDeveloper/Private/RigVMCompiler/RigVMCompiler.cpp
@@ -1635,58 +1635,61 @@ void URigVMCompiler::InitializeLocalVariables(const FRigVMExprAST* InExpr, FRigV
 	if (!WorkData.bSetupMemory)
 	{
 		FRigVMByteCode& ByteCode = WorkData.VM->GetByteCode();
-		FRigVMCallstack Callstack;
+		const FRigVMASTProxy* Proxy = nullptr;
 		switch (InExpr->GetType())
 		{
 			case FRigVMExprAST::EType::CallExtern:
 			{
-				Callstack = InExpr->To<FRigVMCallExternExprAST>()->GetProxy().GetCallstack();
+				Proxy = &InExpr->To<FRigVMCallExternExprAST>()->GetProxy();
 				break;
 			}
 			case FRigVMExprAST::EType::NoOp:
 			{
-				Callstack = InExpr->To<FRigVMNoOpExprAST>()->GetProxy().GetCallstack();
+				Proxy = &InExpr->To<FRigVMNoOpExprAST>()->GetProxy();
 				break;
 			}
 			case FRigVMExprAST::EType::Var:
 			{
-				Callstack = InExpr->To<FRigVMVarExprAST>()->GetProxy().GetCallstack();
+				Proxy = &InExpr->To<FRigVMVarExprAST>()->GetProxy();
 				break;
 			}
 			case FRigVMExprAST::EType::Literal:
 			{
-				Callstack = InExpr->To<FRigVMLiteralExprAST>()->GetProxy().GetCallstack();
+				Proxy = &InExpr->To<FRigVMLiteralExprAST>()->GetProxy();
 				break;
 			}
 			case FRigVMExprAST::EType::ExternalVar:
 			{
-				Callstack = InExpr->To<FRigVMExternalVarExprAST>()->GetProxy().GetCallstack();
+				Proxy = &InExpr->To<FRigVMExternalVarExprAST>()->GetProxy();
 				break;
 			}
 			case FRigVMExprAST::EType::Branch:
 			{
-				Callstack = InExpr->To<FRigVMBranchExprAST>()->GetProxy().GetCallstack();
+				Proxy = &InExpr->To<FRigVMBranchExprAST>()->GetProxy();
 				break;
 			}
 			case FRigVMExprAST::EType::If:
 			{
-				Callstack = InExpr->To<FRigVMIfExprAST>()->GetProxy().GetCallstack();
+				Proxy = &InExpr->To<FRigVMIfExprAST>()->GetProxy();
 				break;
 			}
 			case FRigVMExprAST::EType::Select:
 			{
-				Callstack = InExpr->To<FRigVMSelectExprAST>()->GetProxy().GetCallstack();
+				Proxy = &InExpr->To<FRigVMSelectExprAST>()->GetProxy();
 				break;
 			}
 			case FRigVMExprAST::EType::Array:
 			{
-				Callstack = InExpr->To<FRigVMArrayExprAST>()->GetProxy().GetCallstack();
+				Proxy = &InExpr->To<FRigVMArrayExprAST>()->GetProxy();
 				break;
 			}
 		}
 
-		if (Callstack.Num() > 0)
+		if(Proxy != nullptr)
 		{
+			const FRigVMCallstack& Callstack = Proxy->GetCallstack();
+			ensure(Callstack.Num() > 0);
+
 			// Look at previous instructions to get the last callstack
 			const TArray<UObject*>* PreviousCallstack = nullptr;
 			int32 PreviousIndex = ByteCode.GetNumInstructions()-1;
@@ -1756,40 +1759,50 @@ void URigVMCompiler::InitializeLocalVariables(const FRigVMExprAST* InExpr, FRigV
 								if (!bAlreadyCopied)
 								{
 									ByteCode.AddCopyOp(WorkData.VM->GetCopyOpForOperands(Source, Target));
+									if(Settings.SetupNodeInstructionIndex)
+									{
+										const int32 InstructionIndex = ByteCode.GetNumInstructions() - 1;
+										WorkData.VM->GetByteCode().SetSubject(InstructionIndex, Callstack.GetCallPath(), Callstack.GetStack());
+									}
 								}	
 							}					
 						}
 					}
 				}
 			}
-		}
 
-		// We also need to initialize local variables of the root graph (after the entry instruction is inserted)
-		// If the parent of InExpr is an entry, initialize the root graph's local variables
-		if(InExpr->NumParents() > 0 && InExpr->ParentAt(0)->GetType() == FRigVMExprAST::Entry)
-		{
-			if (UObject* Subject = ByteCode.GetSubjectForInstruction(ByteCode.GetNumInstructions()-1))
+			// We also need to initialize local variables of the root graph (after the entry instruction is inserted)
+			// If the parent of InExpr is an entry, initialize the root graph's local variables
+			if(InExpr->NumParents() > 0 && InExpr->ParentAt(0)->GetType() == FRigVMExprAST::Entry)
 			{
-				URigVMNode* Node = Cast<URigVMNode>(Subject);
-				if (!Node)
+				if (UObject* Subject = ByteCode.GetSubjectForInstruction(ByteCode.GetNumInstructions()-1))
 				{
-					URigVMPin* Pin = Cast<URigVMPin>(Subject);
-					Node = Pin->GetNode();
-				}
-				if (Node)
-				{
-					for (FRigVMGraphVariableDescription Variable : Node->GetGraph()->LocalVariables)
+					URigVMNode* Node = Cast<URigVMNode>(Subject);
+					if (!Node)
+					{
+						URigVMPin* Pin = Cast<URigVMPin>(Subject);
+						Node = Pin->GetNode();
+					}
+					if (Node)
 					{
-						FString TargetPath = FString::Printf(TEXT("LocalVariable::%s"), *Variable.Name.ToString());
-						FString SourcePath = FString::Printf(TEXT("LocalVariableDefault::|%s::Const"), *Variable.Name.ToString());
-						FRigVMOperand* TargetPtr = WorkData.PinPathToOperand->Find(TargetPath);
-						FRigVMOperand* SourcePtr = WorkData.PinPathToOperand->Find(SourcePath);
-						if (SourcePtr && TargetPtr)
+						for (FRigVMGraphVariableDescription Variable : Node->GetGraph()->LocalVariables)
 						{
-							const FRigVMOperand& Source = *SourcePtr;
-							const FRigVMOperand& Target = *TargetPtr;
-							ByteCode.AddCopyOp(WorkData.VM->GetCopyOpForOperands(Source, Target));
-						}					
+							FString TargetPath = FString::Printf(TEXT("LocalVariable::%s"), *Variable.Name.ToString());
+							FString SourcePath = FString::Printf(TEXT("LocalVariableDefault::|%s::Const"), *Variable.Name.ToString());
+							FRigVMOperand* TargetPtr = WorkData.PinPathToOperand->Find(TargetPath);
+							FRigVMOperand* SourcePtr = WorkData.PinPathToOperand->Find(SourcePath);
+							if (SourcePtr && TargetPtr)
+							{
+								const FRigVMOperand& Source = *SourcePtr;
+								const FRigVMOperand& Target = *TargetPtr;
+								ByteCode.AddCopyOp(WorkData.VM->GetCopyOpForOperands(Source, Target));
+								if(Settings.SetupNodeInstructionIndex)
+								{
+									const int32 InstructionIndex = ByteCode.GetNumInstructions() - 1;
+									WorkData.VM->GetByteCode().SetSubject(InstructionIndex, Callstack.GetCallPath(), Callstack.GetStack());
+								}
+							}					
+						}
 					}
 				}
 			}
diff --git a/Engine/Source/Developer/RigVMDeveloper/Private/RigVMModel/Nodes/RigVMLibraryNode.cpp b/Engine/Source/Developer/RigVMDeveloper/Private/RigVMModel/Nodes/RigVMLibraryNode.cpp
--- a/Engine/Source/Developer/RigVMDeveloper/Private/RigVMModel/Nodes/RigVMLibraryNode.cpp
+++ b/Engine/Source/Developer/RigVMDeveloper/Private/RigVMModel/Nodes/RigVMLibraryNode.cpp
@@ -31,7 +31,7 @@ bool URigVMLibraryNode::IsDefinedAsVarying() const
 	return false;
 }
 
-TArray<int32> URigVMLibraryNode::GetInstructionsForVM(URigVM* InVM, const FRigVMASTProxy& InProxy) const
+TArray<int32> URigVMLibraryNode::GetInstructionsForVMImpl(URigVM* InVM, const FRigVMASTProxy& InProxy) const
 {
 	TArray<int32> Instructions;
 
@@ -53,6 +53,7 @@ TArray<int32> URigVMLibraryNode::GetInstructionsForVM(URigVM* InVM, const FRigVM
 	return Instructions;
 }
 
+/*
 int32 URigVMLibraryNode::GetInstructionVisitedCount(URigVM* InVM, const FRigVMASTProxy& InProxy, bool bConsolidatePerNode) const
 {
 	int32 Count = 0;
@@ -75,6 +76,7 @@ int32 URigVMLibraryNode::GetInstructionVisitedCount(URigVM* InVM, const FRigVMAS
 
 	return Count;
 }
+*/
 
 const TArray<URigVMNode*>& URigVMLibraryNode::GetContainedNodes() const
 {
diff --git a/Engine/Source/Developer/RigVMDeveloper/Private/RigVMModel/RigVMNode.cpp b/Engine/Source/Developer/RigVMDeveloper/Private/RigVMModel/RigVMNode.cpp
--- a/Engine/Source/Developer/RigVMDeveloper/Private/RigVMModel/RigVMNode.cpp
+++ b/Engine/Source/Developer/RigVMDeveloper/Private/RigVMModel/RigVMNode.cpp
@@ -7,6 +7,10 @@
 
 const FString URigVMNode::NodeColorName = TEXT("NodeColor");
 
+#if WITH_EDITOR
+TArray<int32> URigVMNode::EmptyInstructionArray;
+#endif
+
 URigVMNode::URigVMNode()
 : UObject()
 , Position(FVector2D::ZeroVector)
@@ -14,6 +18,9 @@ URigVMNode::URigVMNode()
 , NodeColor(FLinearColor::White)
 , bHasBreakpoint(false)
 , bHaltedAtThisNode(false)
+#if WITH_EDITOR
+, ProfilingHash(0)
+#endif
 {
 
 }
@@ -446,7 +453,16 @@ void URigVMNode::GetLinkedNodesRecursive(URigVMPin* InPin, bool bLookForSources,
 	}
 }
 
-TArray<int32> URigVMNode::GetInstructionsForVM(URigVM* InVM, const FRigVMASTProxy& InProxy) const
+const TArray<int32>& URigVMNode::GetInstructionsForVM(URigVM* InVM, const FRigVMASTProxy& InProxy) const
+{
+	if(const FProfilingCache* Cache = UpdateProfilingCacheIfNeeded(InVM, InProxy))
+	{
+		return Cache->Instructions;
+	}
+	return EmptyInstructionArray;
+}
+
+TArray<int32> URigVMNode::GetInstructionsForVMImpl(URigVM* InVM, const FRigVMASTProxy& InProxy) const
 {
 	TArray<int32> Instructions;
 
@@ -473,30 +489,88 @@ TArray<int32> URigVMNode::GetInstructionsForVM(URigVM* InVM, const FRigVMASTProx
 	return Instructions;
 }
 
-int32 URigVMNode::GetInstructionVisitedCount(URigVM* InVM, const FRigVMASTProxy& InProxy, bool bConsolidatePerNode) const
+int32 URigVMNode::GetInstructionVisitedCount(URigVM* InVM, const FRigVMASTProxy& InProxy) const
 {
 #if WITH_EDITOR
 	if(InVM)
 	{
-		TArray<int32> Instructions = GetInstructionsForVM(InVM, InProxy);
-		if(Instructions.Num() > 0)
+		if(const FProfilingCache* Cache = UpdateProfilingCacheIfNeeded(InVM, InProxy))
 		{
-			int32 Count = 0;
-			for(int32 Instruction : Instructions)
+			return Cache->VisitedCount;
+		}
+	}
+#endif
+	return 0;
+}
+
+double URigVMNode::GetInstructionMicroSeconds(URigVM* InVM, const FRigVMASTProxy& InProxy) const
+{
+#if WITH_EDITOR
+	if(InVM)
+	{
+		if(const FProfilingCache* Cache = UpdateProfilingCacheIfNeeded(InVM, InProxy))
+		{
+			return Cache->MicroSeconds;
+		}
+	}
+#endif
+	return -1.0;
+}
+
+#if WITH_EDITOR
+
+const URigVMNode::FProfilingCache* URigVMNode::UpdateProfilingCacheIfNeeded(URigVM* InVM, const FRigVMASTProxy& InProxy) const
+{
+	if(InVM == nullptr)
+	{
+		return nullptr;
+	}
+	
+	const uint32 VMHash = HashCombine(GetTypeHash(InVM), GetTypeHash(InVM->GetNumExecutions()));
+	if(VMHash != ProfilingHash)
+	{
+		ProfilingCache.Reset();
+	}
+	ProfilingHash = VMHash;
+
+	const uint32 ProxyHash = InProxy.IsValid() ? GetTypeHash(InProxy) : GetTypeHash(this);
+
+	const TSharedPtr<FProfilingCache>* ExistingCache = ProfilingCache.Find(ProxyHash);
+	if(ExistingCache)
+	{
+		return ExistingCache->Get();
+	}
+
+	TSharedPtr<FProfilingCache> Cache(new FProfilingCache);
+
+	Cache->Instructions = GetInstructionsForVMImpl(InVM, InProxy);
+	Cache->VisitedCount = 0;
+	Cache->MicroSeconds = -1.0;
+
+	if(Cache->Instructions.Num() > 0)
+	{
+		for(const int32 Instruction : Cache->Instructions)
+		{
+			const int32 CountPerInstruction = InVM->GetInstructionVisitedCount(Instruction);
+			Cache->VisitedCount += CountPerInstruction;
+
+			const double MicroSecondsPerInstruction = InVM->GetInstructionMicroSeconds(Instruction);
+			if(MicroSecondsPerInstruction >= 0.0)
 			{
-				const int32 CountPerInstruction = InVM->GetInstructionVisitedCount(Instruction);
-				if(bConsolidatePerNode)
+				if(Cache->MicroSeconds < 0.0)
 				{
-					Count = FMath::Max<int32>(Count, CountPerInstruction);
+					Cache->MicroSeconds = MicroSecondsPerInstruction;
 				}
 				else
 				{
-					Count += CountPerInstruction;
+					Cache->MicroSeconds += MicroSecondsPerInstruction;
 				}
 			}
-			return Count;
 		}
 	}
-#endif
-	return 0;
+
+	ProfilingCache.Add(ProxyHash, Cache);
+	return Cache.Get();;
 }
+
+#endif
\ No newline at end of file
diff --git a/Engine/Source/Developer/RigVMDeveloper/Public/RigVMModel/Nodes/RigVMLibraryNode.h b/Engine/Source/Developer/RigVMDeveloper/Public/RigVMModel/Nodes/RigVMLibraryNode.h
--- a/Engine/Source/Developer/RigVMDeveloper/Public/RigVMModel/Nodes/RigVMLibraryNode.h
+++ b/Engine/Source/Developer/RigVMDeveloper/Public/RigVMModel/Nodes/RigVMLibraryNode.h
@@ -26,8 +26,7 @@ public:
 	// Override node functions
 	virtual bool IsDefinedAsConstant() const override;
 	virtual bool IsDefinedAsVarying() const override;
-	virtual TArray<int32> GetInstructionsForVM(URigVM* InVM, const FRigVMASTProxy& InProxy = FRigVMASTProxy()) const override; 
-	virtual int32 GetInstructionVisitedCount(URigVM* InVM, const FRigVMASTProxy& InProxy = FRigVMASTProxy(), bool bConsolidatePerNode = false) const override; 
+	//virtual int32 GetInstructionVisitedCount(URigVM* InVM, const FRigVMASTProxy& InProxy = FRigVMASTProxy(), bool bConsolidatePerNode = false) const override; 
 
 	// Library node interface
 	virtual FString GetNodeCategory() const { return FString(); }
@@ -49,6 +48,7 @@ public:
 
 protected:
 
+	virtual TArray<int32> GetInstructionsForVMImpl(URigVM* InVM, const FRigVMASTProxy& InProxy = FRigVMASTProxy()) const override; 
 	const static TArray<URigVMNode*> EmptyNodes;
 	const static TArray<URigVMLink*> EmptyLinks;
 
diff --git a/Engine/Source/Developer/RigVMDeveloper/Public/RigVMModel/RigVMNode.h b/Engine/Source/Developer/RigVMDeveloper/Public/RigVMModel/RigVMNode.h
--- a/Engine/Source/Developer/RigVMDeveloper/Public/RigVMModel/RigVMNode.h
+++ b/Engine/Source/Developer/RigVMDeveloper/Public/RigVMModel/RigVMNode.h
@@ -189,10 +189,13 @@ public:
 	FName GetPreviousFName() const { return PreviousName; }
 
 	// Returns the indices of associated instructions for this node
-	virtual TArray<int32> GetInstructionsForVM(URigVM* InVM, const FRigVMASTProxy& InProxy = FRigVMASTProxy()) const; 
+	const TArray<int32>& GetInstructionsForVM(URigVM* InVM, const FRigVMASTProxy& InProxy = FRigVMASTProxy()) const;
 
 	// Returns the number of visited / run instructions for this node
-	virtual int32 GetInstructionVisitedCount(URigVM* InVM, const FRigVMASTProxy& InProxy = FRigVMASTProxy(), bool bConsolidatePerNode = false) const; 
+	virtual int32 GetInstructionVisitedCount(URigVM* InVM, const FRigVMASTProxy& InProxy = FRigVMASTProxy()) const;
+
+	// Returns the accumulated duration of all of instructions for this node 
+	double GetInstructionMicroSeconds(URigVM* InVM, const FRigVMASTProxy& InProxy = FRigVMASTProxy()) const;
 
 	// return true if this node is a loop node
 	virtual bool IsLoopNode() const { return false; }
@@ -217,6 +220,7 @@ private:
 
 protected:
 
+	virtual TArray<int32> GetInstructionsForVMImpl(URigVM* InVM, const FRigVMASTProxy& InProxy = FRigVMASTProxy()) const; 
 	virtual FText GetToolTipTextForPin(const URigVMPin* InPin) const;
 	virtual bool AllowsLinksOn(const URigVMPin* InPin) const { return true; }
 
@@ -249,6 +253,19 @@ private:
 	UPROPERTY()
 	TArray<TObjectPtr<URigVMPin>> OrphanedPins;
 
+#if WITH_EDITOR
+	struct FProfilingCache
+	{
+		mutable int32 VisitedCount;
+		mutable double MicroSeconds;
+		mutable TArray<int32> Instructions;
+	};
+	const FProfilingCache* UpdateProfilingCacheIfNeeded(URigVM* InVM, const FRigVMASTProxy& InProxy) const;
+	mutable uint32 ProfilingHash;
+	mutable TMap<uint32, TSharedPtr<FProfilingCache>> ProfilingCache;
+	static TArray<int32> EmptyInstructionArray;
+#endif
+	
 	friend class URigVMController;
 	friend class URigVMGraph;
 	friend class URigVMPin;
diff --git a/Engine/Source/Runtime/RigVM/Private/RigVMCore/RigVM.cpp b/Engine/Source/Runtime/RigVM/Private/RigVMCore/RigVM.cpp
--- a/Engine/Source/Runtime/RigVM/Private/RigVMCore/RigVM.cpp
+++ b/Engine/Source/Runtime/RigVM/Private/RigVMCore/RigVM.cpp
@@ -74,6 +74,7 @@ URigVM::URigVM()
 	, LiteralMemoryStorageObject(nullptr)
 	, DebugMemoryStorageObject(nullptr)
 	, ByteCodePtr(&ByteCodeStorage)
+	, NumExecutions(0)
 #if WITH_EDITOR
 	, DebugInfo(nullptr)
 	, HaltedAtBreakpoint(nullptr)
@@ -551,6 +552,7 @@ void URigVM::Reset()
 	InvalidateCachedMemory();
 	
 	OperandToDebugRegisters.Reset();
+	NumExecutions = 0;
 }
 
 void URigVM::Empty()
@@ -1860,6 +1862,15 @@ bool URigVM::Execute(TArrayView<URigVMMemoryStorage*> Memory, TArrayView<void*>
 		InstructionVisitedDuringLastRun.Reset();
 		InstructionVisitOrder.Reset();
 		InstructionVisitedDuringLastRun.SetNumZeroed(Instructions.Num());
+		InstructionCyclesDuringLastRun.Reset();
+		if(Context.RuntimeSettings.bEnableProfiling)
+		{
+			InstructionCyclesDuringLastRun.SetNumUninitialized(Instructions.Num());
+			for(int32 DurationIndex=0;DurationIndex<InstructionCyclesDuringLastRun.Num();DurationIndex++)
+			{
+				InstructionCyclesDuringLastRun[DurationIndex] = UINT64_MAX;
+			}
+		}
 	}
 #endif
 
@@ -1887,6 +1898,14 @@ bool URigVM::Execute(TArrayView<URigVMMemoryStorage*> Memory, TArrayView<void*>
 	}
 #endif
 
+	NumExecutions++;
+
+	uint64 StartCycles = 0;
+	if(Context.RuntimeSettings.bEnableProfiling)
+	{
+		StartCycles = FPlatformTime::Cycles64();
+	}
+
 	while (Instructions.IsValidIndex(Context.InstructionIndex))
 	{
 #if WITH_EDITOR
@@ -1894,9 +1913,11 @@ bool URigVM::Execute(TArrayView<URigVMMemoryStorage*> Memory, TArrayView<void*>
 		{
 			return true;
 		}
-		
+
+		const int32 CurrentInstructionIndex = Context.InstructionIndex;
 		InstructionVisitedDuringLastRun[Context.InstructionIndex]++;
 		InstructionVisitOrder.Add(Context.InstructionIndex);
+	
 #endif
 
 		const FRigVMInstruction& Instruction = Instructions[Context.InstructionIndex];
@@ -2916,6 +2937,24 @@ bool URigVM::Execute(TArrayView<URigVMMemoryStorage*> Memory, TArrayView<void*>
 				return false;
 			}
 		}
+
+#if WITH_EDITOR
+		if(Context.RuntimeSettings.bEnableProfiling && !InstructionVisitOrder.IsEmpty())
+		{
+			const uint64 EndCycles = FPlatformTime::Cycles64();
+			const uint64 Cycles = EndCycles - StartCycles;
+			if(InstructionCyclesDuringLastRun[CurrentInstructionIndex] == UINT64_MAX)
+			{
+				InstructionCyclesDuringLastRun[CurrentInstructionIndex] = Cycles;
+			}
+			else
+			{
+				InstructionCyclesDuringLastRun[CurrentInstructionIndex] += Cycles;
+			}
+
+			StartCycles = EndCycles;
+		}
+#endif
 	}
 
 #if WITH_EDITOR
diff --git a/Engine/Source/Runtime/RigVM/Public/RigVMCore/RigVM.h b/Engine/Source/Runtime/RigVM/Public/RigVMCore/RigVM.h
--- a/Engine/Source/Runtime/RigVM/Public/RigVMCore/RigVM.h
+++ b/Engine/Source/Runtime/RigVM/Public/RigVMCore/RigVM.h
@@ -12,6 +12,7 @@
 #include "RigVMStatistics.h"
 #if WITH_EDITOR
 #include "RigVMDebugInfo.h"
+#include "HAL/PlatformTime.h"
 #endif
 #include "RigVM.generated.h"
 
@@ -292,6 +293,31 @@ public:
 		return 0;
 	}
 
+	// Returns accumulated cycles spent in an instruction during the last run
+	// This requires bEnabledProfiling to be turned on in the runtime settings.
+	// If there is no information available this function returns UINT64_MAX.
+	FORCEINLINE uint64 GetInstructionCycles(int32 InIndex) const
+	{
+		if (InstructionCyclesDuringLastRun.IsValidIndex(InIndex))
+		{
+			return InstructionCyclesDuringLastRun[InIndex];
+		}
+		return UINT64_MAX;
+	}
+
+	// Returns accumulated duration of the instruction in microseconds during the last run
+	// Note: this requires bEnabledProfiling to be turned on in the runtime settings.
+	// If there is no information available this function returns -1.0.
+	FORCEINLINE double GetInstructionMicroSeconds(int32 InIndex) const
+	{
+		const uint64 Cycles = GetInstructionCycles(InIndex);
+		if(Cycles == UINT64_MAX)
+		{
+			return -1.0;
+		}
+		return double(Cycles) * FPlatformTime::GetSecondsPerCycle() * 1000.0 * 1000.0;
+	}
+
 	// Returns the order of all instructions during the last run
 	FORCEINLINE const TArray<int32> GetInstructionVisitOrder() const { return InstructionVisitOrder; }
 
@@ -718,6 +744,8 @@ public:
 	void SetBreakpointAction(const ERigVMBreakpointAction& Action) { CurrentBreakpointAction = Action; }
 #endif
 
+	uint32 GetNumExecutions() const { return NumExecutions; }
+
 private:
 
 	void ResolveFunctionsIfRequired();
@@ -736,6 +764,9 @@ private:
 	UPROPERTY(transient)
 	FRigVMExecuteContext Context;
 
+	UPROPERTY(transient)
+	uint32 NumExecutions;
+
 #if WITH_EDITOR
 	FRigVMDebugInfo* DebugInfo;
 	TSharedPtr<FRigVMBreakpoint> HaltedAtBreakpoint;
@@ -788,8 +819,10 @@ private:
 	}
 	
 #if WITH_EDITOR
+
 	// stores the number of times each instruction was visited
 	TArray<int32> InstructionVisitedDuringLastRun;
+	TArray<uint64> InstructionCyclesDuringLastRun;
 	TArray<int32> InstructionVisitOrder;
 	
 	// Control Rig can run multiple events per evaluation, such as the Backward&Forward Solve Mode,
diff --git a/Engine/Source/Runtime/RigVM/Public/RigVMCore/RigVMExecuteContext.h b/Engine/Source/Runtime/RigVM/Public/RigVMCore/RigVMExecuteContext.h
--- a/Engine/Source/Runtime/RigVM/Public/RigVMCore/RigVMExecuteContext.h
+++ b/Engine/Source/Runtime/RigVM/Public/RigVMCore/RigVMExecuteContext.h
@@ -157,12 +157,22 @@ struct RIGVM_API FRigVMRuntimeSettings
 	GENERATED_BODY()
 
 	/**
-	* The largest allowed size for arrays within the Control Rig VM.
-	* Accessing or creating larger arrays will cause runtime errors in the rig.
-	*/
+	 * The largest allowed size for arrays within the Control Rig VM.
+	 * Accessing or creating larger arrays will cause runtime errors in the rig.
+	 */
 	UPROPERTY(EditAnywhere, Category = "VM")
 	int32 MaximumArraySize = 2048;
 
+#if WITH_EDITORONLY_DATA
+	// When enabled records the timing of each instruction / node
+	// on each node and within the execution stack window.
+	// Keep in mind when looking at nodes in a function the duration
+	// represents the accumulated duration of all invocations
+	// of the function currently running.
+	UPROPERTY(EditAnywhere, Category = "VM")
+	bool bEnableProfiling = false;
+#endif
+
 	/*
 	 * The function to use for logging anything from the VM to the host
 	 */
