Control Rig: Added null checks when collapsing to function

diff --git a/Engine/Source/Developer/RigVMDeveloper/Private/RigVMModel/RigVMController.cpp b/Engine/Source/Developer/RigVMDeveloper/Private/RigVMModel/RigVMController.cpp
--- a/Engine/Source/Developer/RigVMDeveloper/Private/RigVMModel/RigVMController.cpp
+++ b/Engine/Source/Developer/RigVMDeveloper/Private/RigVMModel/RigVMController.cpp
@@ -4081,9 +4081,15 @@ URigVMFunctionReferenceNode* URigVMController::PromoteCollapseNodeToFunctionRefe
 		{
 			FRigVMControllerGraphGuard GraphGuard(this, FunctionDefinition->GetContainedGraph(), bSetupUndoRedo);
 			ImportNodesFromText(TextContent, bSetupUndoRedo);
-
-			SetNodePosition(FunctionDefinition->GetContainedGraph()->GetEntryNode(), InCollapseNode->GetContainedGraph()->GetEntryNode()->GetPosition());
-			SetNodePosition(FunctionDefinition->GetContainedGraph()->GetReturnNode(), InCollapseNode->GetContainedGraph()->GetReturnNode()->GetPosition());
+			if (FunctionDefinition->GetContainedGraph()->GetEntryNode() && InCollapseNode->GetContainedGraph()->GetEntryNode())
+			{ 
+				SetNodePosition(FunctionDefinition->GetContainedGraph()->GetEntryNode(), InCollapseNode->GetContainedGraph()->GetEntryNode()->GetPosition());
+			}
+			
+			if (FunctionDefinition->GetContainedGraph()->GetReturnNode() && InCollapseNode->GetContainedGraph()->GetReturnNode())
+			{ 
+				SetNodePosition(FunctionDefinition->GetContainedGraph()->GetReturnNode(), InCollapseNode->GetContainedGraph()->GetReturnNode()->GetPosition());
+			}
 
 			for (const URigVMLink* InnerLink : InCollapseNode->GetContainedGraph()->GetLinks())
 			{
