Control Rig: Low performance when reordering function pins

diff --git a/Engine/Source/Developer/RigVMDeveloper/Private/RigVMModel/RigVMController.cpp b/Engine/Source/Developer/RigVMDeveloper/Private/RigVMModel/RigVMController.cpp
--- a/Engine/Source/Developer/RigVMDeveloper/Private/RigVMModel/RigVMController.cpp
+++ b/Engine/Source/Developer/RigVMDeveloper/Private/RigVMModel/RigVMController.cpp
@@ -7326,6 +7326,8 @@ bool URigVMController::SetExposedPinIndex(const FName& InPinName, int32 InNewInd
 		return false;
 	}
 
+	Notify(ERigVMGraphNotifType::InteractionBracketOpened, nullptr);
+
 	FRigVMSetPinIndexAction PinIndexAction(Pin, InNewIndex);
 	{
 		LibraryNode->Pins.Remove(Pin);
@@ -7344,6 +7346,8 @@ bool URigVMController::SetExposedPinIndex(const FName& InPinName, int32 InNewInd
 		ActionStack->AddAction(PinIndexAction);
 	}
 
+	Notify(ERigVMGraphNotifType::InteractionBracketClosed, nullptr);
+
 	if (bPrintPythonCommand)
 	{
 		RigVMPythonUtils::Print(GetGraph()->GetRootGraph()->GetOuter()->GetFName().ToString(), 
