Control Rig: Introduce pin validity checking + fixup

diff --git a/Engine/Source/Developer/RigVMDeveloper/Private/RigVMCompiler/RigVMCompiler.cpp b/Engine/Source/Developer/RigVMDeveloper/Private/RigVMCompiler/RigVMCompiler.cpp
--- a/Engine/Source/Developer/RigVMDeveloper/Private/RigVMCompiler/RigVMCompiler.cpp
+++ b/Engine/Source/Developer/RigVMDeveloper/Private/RigVMCompiler/RigVMCompiler.cpp
@@ -199,6 +199,14 @@ bool URigVMCompiler::Compile(URigVMGraph* InGraph, URigVMController* InControlle
 					VisitedGraphs.AddUnique(ContainedGraph);
 				}
 			}
+
+			for(URigVMPin* Pin : ModelNode->Pins)
+			{
+				if(!URigVMController::EnsurePinValidity(Pin, true))
+				{
+					return false;
+				}
+			}
 		}
 	}
 
diff --git a/Engine/Source/Developer/RigVMDeveloper/Private/RigVMModel/RigVMController.cpp b/Engine/Source/Developer/RigVMDeveloper/Private/RigVMModel/RigVMController.cpp
--- a/Engine/Source/Developer/RigVMDeveloper/Private/RigVMModel/RigVMController.cpp
+++ b/Engine/Source/Developer/RigVMDeveloper/Private/RigVMModel/RigVMController.cpp
@@ -870,7 +870,7 @@ URigVMVariableNode* URigVMController::AddVariableNode(const FName& InVariableNam
 		ExecutePin->CPPTypeObject = ExecuteContextStruct;
 		ExecutePin->CPPTypeObjectPath = *ExecutePin->CPPTypeObject->GetPathName();
 		ExecutePin->Direction = ERigVMPinDirection::IO;
-		Node->Pins.Add(ExecutePin);
+		AddNodePin(Node, ExecutePin);
 	}
 
 	URigVMPin* VariablePin = NewObject<URigVMPin>(Node, *URigVMVariableNode::VariableName);
@@ -878,7 +878,7 @@ URigVMVariableNode* URigVMController::AddVariableNode(const FName& InVariableNam
 	VariablePin->Direction = ERigVMPinDirection::Hidden;
 	VariablePin->DefaultValue = InVariableName.ToString();
 	VariablePin->CustomWidgetName = TEXT("VariableName");
-	Node->Pins.Add(VariablePin);
+	AddNodePin(Node, VariablePin);
 
 	URigVMPin* ValuePin = NewObject<URigVMPin>(Node, *URigVMVariableNode::ValueName);
 
@@ -911,7 +911,7 @@ URigVMVariableNode* URigVMController::AddVariableNode(const FName& InVariableNam
 	}
 
 	ValuePin->Direction = bIsGetter ? ERigVMPinDirection::Output : ERigVMPinDirection::Input;
-	Node->Pins.Add(ValuePin);
+	AddNodePin(Node, ValuePin);
 
 	Graph->Nodes.Add(Node);
 
@@ -1435,7 +1435,7 @@ URigVMParameterNode* URigVMController::AddParameterNode(const FName& InParameter
 		ExecutePin->CPPTypeObject = ExecuteContextStruct;
 		ExecutePin->CPPTypeObjectPath = *ExecutePin->CPPTypeObject->GetPathName();
 		ExecutePin->Direction = ERigVMPinDirection::IO;
-		Node->Pins.Add(ExecutePin);
+		AddNodePin(Node, ExecutePin);
 	}
 
 	URigVMPin* ParameterPin = NewObject<URigVMPin>(Node, *URigVMParameterNode::ParameterName);
@@ -1444,7 +1444,7 @@ URigVMParameterNode* URigVMController::AddParameterNode(const FName& InParameter
 	ParameterPin->DefaultValue = InParameterName.ToString();
 	ParameterPin->CustomWidgetName = TEXT("ParameterName");
 
-	Node->Pins.Add(ParameterPin);
+	AddNodePin(Node, ParameterPin);
 
 	URigVMPin* DefaultValuePin = nullptr;
 	if (bIsInput)
@@ -1496,9 +1496,9 @@ URigVMParameterNode* URigVMController::AddParameterNode(const FName& InParameter
 
 	if (DefaultValuePin)
 	{
-		Node->Pins.Add(DefaultValuePin);
+		AddNodePin(Node, DefaultValuePin);
 	}
-	Node->Pins.Add(ValuePin);
+	AddNodePin(Node, ValuePin);
 
 	Graph->Nodes.Add(Node);
 
@@ -1752,7 +1752,7 @@ URigVMRerouteNode* URigVMController::AddRerouteNodeOnPin(const FString& InPinPat
 	URigVMPin* ValuePin = NewObject<URigVMPin>(Node, *URigVMRerouteNode::ValueName);
 	ConfigurePinFromPin(ValuePin, Pin);
 	ValuePin->Direction = ERigVMPinDirection::IO;
-	Node->Pins.Add(ValuePin);
+	AddNodePin(Node, ValuePin);
 
 	if (ValuePin->IsStruct())
 	{
@@ -3173,7 +3173,7 @@ URigVMCollapseNode* URigVMController::CollapseNodes(const TArray<URigVMNode*>& I
 
 		bContainsOutputs = bContainsOutputs || bSourceToBeCollapsed;
 
-		CollapseNode->Pins.Add(CollapsedPin);
+		AddNodePin(CollapseNode, CollapsedPin);
 
 		FPinState PinState = GetPinState(PinToCollapse);
 		ApplyPinState(CollapsedPin, PinState);
@@ -6875,7 +6875,7 @@ FName URigVMController::AddExposedPin(const FName& InPinName, ERigVMPinDirection
 	Pin->CPPTypeObjectPath = InCPPTypeObjectPath;
 	Pin->bIsConstant = false;
 	Pin->Direction = InDirection;
-	LibraryNode->Pins.Add(Pin);
+	AddNodePin(LibraryNode, Pin);
 
 	if (Pin->IsStruct())
 	{
@@ -7582,7 +7582,7 @@ URigVMLibraryNode* URigVMController::AddFunctionToLibrary(const FName& InFunctio
 		ExecutePin->CPPTypeObject = ExecuteContextStruct;
 		ExecutePin->CPPTypeObjectPath = *ExecutePin->CPPTypeObject->GetPathName();
 		ExecutePin->Direction = ERigVMPinDirection::IO;
-		CollapseNode->Pins.Add(ExecutePin);
+		AddNodePin(CollapseNode, ExecutePin);
 	}
 
 	Notify(ERigVMGraphNotifType::NodeAdded, CollapseNode);
@@ -8373,7 +8373,7 @@ URigVMRerouteNode* URigVMController::AddFreeRerouteNode(bool bShowAsFullNode, co
 	ValuePin->bIsConstant = bIsConstant;
 	ValuePin->CustomWidgetName = InCustomWidgetName;
 	ValuePin->Direction = ERigVMPinDirection::IO;
-	Node->Pins.Add(ValuePin);
+	AddNodePin(Node, ValuePin);
 	Graph->Nodes.Add(Node);
 
 	if (ValuePin->IsStruct())
@@ -8422,26 +8422,26 @@ URigVMBranchNode* URigVMController::AddBranchNode(const FVector2D& InPosition, c
 	ExecutePin->CPPTypeObject = ExecuteContextStruct;
 	ExecutePin->CPPTypeObjectPath = *ExecutePin->CPPTypeObject->GetPathName();
 	ExecutePin->Direction = ERigVMPinDirection::Input;
-	Node->Pins.Add(ExecutePin);
+	AddNodePin(Node, ExecutePin);
 
 	URigVMPin* ConditionPin = NewObject<URigVMPin>(Node, *URigVMBranchNode::ConditionName);
 	ConditionPin->CPPType = TEXT("bool");
 	ConditionPin->Direction = ERigVMPinDirection::Input;
-	Node->Pins.Add(ConditionPin);
+	AddNodePin(Node, ConditionPin);
 
 	URigVMPin* TruePin = NewObject<URigVMPin>(Node, *URigVMBranchNode::TrueName);
 	TruePin->CPPType = ExecutePin->CPPType;
 	TruePin->CPPTypeObject = ExecutePin->CPPTypeObject;
 	TruePin->CPPTypeObjectPath = ExecutePin->CPPTypeObjectPath;
 	TruePin->Direction = ERigVMPinDirection::Output;
-	Node->Pins.Add(TruePin);
+	AddNodePin(Node, TruePin);
 
 	URigVMPin* FalsePin = NewObject<URigVMPin>(Node, *URigVMBranchNode::FalseName);
 	FalsePin->CPPType = ExecutePin->CPPType;
 	FalsePin->CPPTypeObject = ExecutePin->CPPTypeObject;
 	FalsePin->CPPTypeObjectPath = ExecutePin->CPPTypeObjectPath;
 	FalsePin->Direction = ERigVMPinDirection::Output;
-	Node->Pins.Add(FalsePin);
+	AddNodePin(Node, FalsePin);
 
 	Graph->Nodes.Add(Node);
 
@@ -8509,7 +8509,7 @@ URigVMIfNode* URigVMController::AddIfNode(const FString& InCPPType, const FName&
 	URigVMPin* ConditionPin = NewObject<URigVMPin>(Node, *URigVMIfNode::ConditionName);
 	ConditionPin->CPPType = TEXT("bool");
 	ConditionPin->Direction = ERigVMPinDirection::Input;
-	Node->Pins.Add(ConditionPin);
+	AddNodePin(Node, ConditionPin);
 
 	URigVMPin* TruePin = NewObject<URigVMPin>(Node, *URigVMIfNode::TrueName);
 	TruePin->CPPType = CPPType;
@@ -8517,7 +8517,7 @@ URigVMIfNode* URigVMController::AddIfNode(const FString& InCPPType, const FName&
 	TruePin->CPPTypeObjectPath = InCPPTypeObjectPath;
 	TruePin->Direction = ERigVMPinDirection::Input;
 	TruePin->DefaultValue = DefaultValue;
-	Node->Pins.Add(TruePin);
+	AddNodePin(Node, TruePin);
 
 	if (TruePin->IsStruct())
 	{
@@ -8530,7 +8530,7 @@ URigVMIfNode* URigVMController::AddIfNode(const FString& InCPPType, const FName&
 	FalsePin->CPPTypeObjectPath = InCPPTypeObjectPath;
 	FalsePin->Direction = ERigVMPinDirection::Input;
 	FalsePin->DefaultValue = DefaultValue;
-	Node->Pins.Add(FalsePin);
+	AddNodePin(Node, FalsePin);
 
 	if (FalsePin->IsStruct())
 	{
@@ -8542,7 +8542,7 @@ URigVMIfNode* URigVMController::AddIfNode(const FString& InCPPType, const FName&
 	ResultPin->CPPTypeObject = CPPTypeObject;
 	ResultPin->CPPTypeObjectPath = InCPPTypeObjectPath;
 	ResultPin->Direction = ERigVMPinDirection::Output;
-	Node->Pins.Add(ResultPin);
+	AddNodePin(Node, ResultPin);
 
 	if (ResultPin->IsStruct())
 	{
@@ -8615,7 +8615,7 @@ URigVMSelectNode* URigVMController::AddSelectNode(const FString& InCPPType, cons
 	URigVMPin* IndexPin = NewObject<URigVMPin>(Node, *URigVMSelectNode::IndexName);
 	IndexPin->CPPType = TEXT("int32");
 	IndexPin->Direction = ERigVMPinDirection::Input;
-	Node->Pins.Add(IndexPin);
+	AddNodePin(Node, IndexPin);
 
 	URigVMPin* ValuePin = NewObject<URigVMPin>(Node, *URigVMSelectNode::ValueName);
 	ValuePin->CPPType = FString::Printf(TEXT("TArray<%s>"), *CPPType);
@@ -8623,14 +8623,14 @@ URigVMSelectNode* URigVMController::AddSelectNode(const FString& InCPPType, cons
 	ValuePin->CPPTypeObjectPath = InCPPTypeObjectPath;
 	ValuePin->Direction = ERigVMPinDirection::Input;
 	ValuePin->bIsExpanded = true;
-	Node->Pins.Add(ValuePin);
+	AddNodePin(Node, ValuePin);
 
 	URigVMPin* ResultPin = NewObject<URigVMPin>(Node, *URigVMSelectNode::ResultName);
 	ResultPin->CPPType = CPPType;
 	ResultPin->CPPTypeObject = CPPTypeObject;
 	ResultPin->CPPTypeObjectPath = InCPPTypeObjectPath;
 	ResultPin->Direction = ERigVMPinDirection::Output;
-	Node->Pins.Add(ResultPin);
+	AddNodePin(Node, ResultPin);
 
 	Graph->Nodes.Add(Node);
 
@@ -8708,7 +8708,7 @@ URigVMPrototypeNode* URigVMController::AddPrototypeNode(const FName& InNotation,
 			Pin->CPPTypeObjectPath = *Pin->CPPTypeObject->GetPathName();
 		}
 		Pin->Direction = Arg->GetDirection();
-		Node->Pins.Add(Pin);
+		AddNodePin(Node, Pin);
 	}
 
 	Graph->Nodes.Add(Node);
@@ -8768,13 +8768,13 @@ URigVMEnumNode* URigVMController::AddEnumNode(const FName& InCPPTypeObjectPath,
 	EnumValuePin->CPPTypeObjectPath = InCPPTypeObjectPath;
 	EnumValuePin->Direction = ERigVMPinDirection::Visible;
 	EnumValuePin->DefaultValue = Enum->GetNameStringByValue(0);
-	Node->Pins.Add(EnumValuePin);
+	AddNodePin(Node, EnumValuePin);
 
 	URigVMPin* EnumIndexPin = NewObject<URigVMPin>(Node, *URigVMEnumNode::EnumIndexName);
 	EnumIndexPin->CPPType = TEXT("int32");
 	EnumIndexPin->Direction = ERigVMPinDirection::Output;
 	EnumIndexPin->DisplayName = TEXT("Result");
-	Node->Pins.Add(EnumIndexPin);
+	AddNodePin(Node, EnumIndexPin);
 
 	Graph->Nodes.Add(Node);
 
@@ -8894,7 +8894,7 @@ URigVMArrayNode* URigVMController::AddArrayNode(ERigVMOpCode InOpCode, const FSt
 			}
 			Pin->Direction = InDirection;
 			Pin->bIsDynamicArray = bIsArray;
-			InNode->Pins.Add(Pin);
+			AddNodePin(InNode, Pin);
 
 			if(Pin->Direction != ERigVMPinDirection::Hidden && !bIsArray && !Pin->IsExecuteContext())
 			{
@@ -9465,11 +9465,11 @@ void URigVMController::AddPinsForStruct(UStruct* InStruct, URigVMNode* InNode, U
 
 		if (InParentPin)
 		{
-			InParentPin->SubPins.Add(Pin);
+			AddSubPin(InParentPin, Pin);
 		}
 		else
 		{
-			InNode->Pins.Add(Pin);
+			AddNodePin(InNode, Pin);
 		}
 
 		FString* DefaultValuePtr = MemberValues.Find(Pin->GetFName());
@@ -9545,7 +9545,7 @@ void URigVMController::AddPinsForArray(FArrayProperty* InArrayProperty, URigVMNo
 		ConfigurePinFromProperty(InArrayProperty->Inner, Pin, InPinDirection);
 		FString DefaultValue = InDefaultValues[ElementIndex];
 
-		InParentPin->SubPins.Add(Pin);
+		AddSubPin(InParentPin, Pin);
 
 		if (bAutoExpandArrays)
 		{
@@ -10312,6 +10312,7 @@ void URigVMController::RepopulatePinsOnNode(URigVMNode* InNode, bool bFollowCore
 			ValuePin = VariableNode->FindPin(URigVMVariableNode::ValueName);
 		}
 		check(ValuePin);
+		EnsurePinValidity(ValuePin, false);
 		
 		RemovePinsDuringRepopulate(InNode, ValuePin->SubPins, bNotify, bSetupOrphanedPins);
 
@@ -10393,7 +10394,7 @@ void URigVMController::RepopulatePinsOnNode(URigVMNode* InNode, bool bFollowCore
 					ExposedPin->Direction = ERigVMPinDirection::Input;
 				}
 
-				InNode->Pins.Add(ExposedPin);
+				AddNodePin(InNode, ExposedPin);
 
 				if (ExposedPin->IsStruct())
 				{
@@ -10433,7 +10434,7 @@ void URigVMController::RepopulatePinsOnNode(URigVMNode* InNode, bool bFollowCore
 				URigVMPin* NewPin = NewObject<URigVMPin>(InNode, ReferencedPin->GetFName());
 				ConfigurePinFromPin(NewPin, ReferencedPin);
 
-				InNode->Pins.Add(NewPin);
+				AddNodePin(InNode, NewPin);
 
 				if (NewPin->IsStruct())
 				{
@@ -10521,7 +10522,7 @@ void URigVMController::RemovePinsDuringRepopulate(URigVMNode* InNode, TArray<URi
 			{
 				Pin->Rename(nullptr, OrphanedRootPin, REN_ForceNoResetLoaders | REN_DoNotDirty | REN_DontCreateRedirectors | REN_NonTransactional);
 				RootPin->SubPins.Remove(Pin);
-				OrphanedRootPin->SubPins.Add(Pin);
+				AddSubPin(OrphanedRootPin, Pin);
 			}
 		}
 	}
@@ -11170,6 +11171,89 @@ void URigVMController::DestroyObject(UObject* InObjectToDestroy)
 	InObjectToDestroy->RemoveFromRoot();
 }
 
+void URigVMController::AddNodePin(URigVMNode* InNode, URigVMPin* InPin)
+{
+	ValidatePin(InPin);
+	InNode->Pins.Add(InPin);
+}
+
+void URigVMController::AddSubPin(URigVMPin* InParentPin, URigVMPin* InPin)
+{
+	ValidatePin(InPin);
+	InParentPin->SubPins.Add(InPin);
+}
+
+bool URigVMController::EnsurePinValidity(URigVMPin* InPin, bool bRecursive)
+{
+	check(InPin);
+	
+	// check if the CPPTypeObject is set up correctly.
+	if(FRigVMPropertyDescription::RequiresCPPTypeObject(InPin->GetCPPType()))
+	{
+		if(InPin->GetCPPTypeObject() == nullptr)
+		{
+			// try to find the CPPTypeObject by name
+			FString CPPType = InPin->IsArray() ? InPin->GetArrayElementCppType() : InPin->GetCPPType();
+			UObject* CPPTypeObject = FindObject<UObject>(ANY_PACKAGE, *CPPType);
+			if(CPPTypeObject == nullptr)
+			{
+				if(CPPType.StartsWith(TEXT("E"), ESearchCase::CaseSensitive) ||
+					CPPType.StartsWith(TEXT("F"), ESearchCase::CaseSensitive) ||
+					CPPType.StartsWith(TEXT("U"), ESearchCase::CaseSensitive))
+				{
+					CPPType = CPPType.Mid(1);
+					CPPTypeObject = FindObject<UObject>(ANY_PACKAGE, *CPPType);
+				}
+
+				if(CPPTypeObject == nullptr)
+				{
+					const FString Message = FString::Printf(
+						TEXT("Pin '%s' is missing the CPPTypeObject for CPPType '%s'."),
+						*InPin->GetPinPath(), *InPin->GetCPPType());
+					FScriptExceptionHandler::Get().HandleException(ELogVerbosity::Error, *Message, *FString());
+					return false;
+				}
+			}
+			InPin->CPPTypeObject = CPPTypeObject;
+		}
+	}
+
+	if(const UScriptStruct* ScriptStruct = Cast<UScriptStruct>(InPin->GetCPPTypeObject()))
+	{
+		if(InPin->IsArray())
+		{
+			static const TCHAR TemplateTArray[] = TEXT("TArray<%s>");
+			InPin->CPPType = FString::Printf(TemplateTArray, *ScriptStruct->GetStructCPPName());
+		}
+		else
+		{
+			InPin->CPPType = ScriptStruct->GetStructCPPName();
+		}
+	}
+
+	if(bRecursive)
+	{
+		for(URigVMPin* SubPin : InPin->SubPins)
+		{
+			if(!EnsurePinValidity(SubPin, bRecursive))
+			{
+				return false;
+			}
+		}
+	}
+
+	return true;
+}
+
+void URigVMController::ValidatePin(URigVMPin* InPin)
+{
+	check(InPin);
+	
+	// create a property description from the pin here as a test,
+	// since the compiler needs this
+	FRigVMPropertyDescription(InPin->GetFName(), InPin->GetCPPType(), InPin->GetCPPTypeObject(), InPin->GetDefaultValue());
+}
+
 FRigVMExternalVariable URigVMController::GetVariableByName(const FName& InExternalVariableName)
 {
 	TArray<FRigVMExternalVariable> ExternalVariables = GetAllVariables();
diff --git a/Engine/Source/Developer/RigVMDeveloper/Public/RigVMModel/RigVMController.h b/Engine/Source/Developer/RigVMDeveloper/Public/RigVMModel/RigVMController.h
--- a/Engine/Source/Developer/RigVMDeveloper/Public/RigVMModel/RigVMController.h
+++ b/Engine/Source/Developer/RigVMDeveloper/Public/RigVMModel/RigVMController.h
@@ -838,6 +838,10 @@ private:
 #endif
 
 	void DestroyObject(UObject* InObjectToDestroy);
+	static void AddNodePin(URigVMNode* InNode, URigVMPin* InPin);
+	static void AddSubPin(URigVMPin* InParentPin, URigVMPin* InPin);
+	static bool EnsurePinValidity(URigVMPin* InPin, bool bRecursive);
+	static void ValidatePin(URigVMPin* InPin);
 
 	FRigVMExternalVariable GetVariableByName(const FName& InExternalVariableName);
 	TArray<FRigVMExternalVariable> GetAllVariables();
diff --git a/Engine/Source/Runtime/RigVM/Private/RigVMCore/RigVMMemoryStorage.cpp b/Engine/Source/Runtime/RigVM/Private/RigVMCore/RigVMMemoryStorage.cpp
--- a/Engine/Source/Runtime/RigVM/Private/RigVMCore/RigVMMemoryStorage.cpp
+++ b/Engine/Source/Runtime/RigVM/Private/RigVMCore/RigVMMemoryStorage.cpp
@@ -63,6 +63,12 @@ FRigVMPropertyDescription::FRigVMPropertyDescription(const FName& InName, const
 	, Containers()
 	, DefaultValue(InDefaultValue)
 {
+	// sanity check that we have a CPP Type Object for enums, structs and uobjects
+	if(RequiresCPPTypeObject(CPPType))
+	{
+		checkf(CPPTypeObject != nullptr, TEXT("CPPType '%s' requires the CPPTypeObject to be provided."), *CPPType);
+	}
+
 	// only allow valid names
 	SanitizeName();
 
@@ -88,6 +94,13 @@ FRigVMPropertyDescription::FRigVMPropertyDescription(const FName& InName, const
 		}
 	}
 	while(!TailCPPType.IsEmpty());
+
+	// sanity check the CPPType matches the provided struct
+	if(const UScriptStruct* ScriptStruct = Cast<UScriptStruct>(CPPTypeObject))
+	{
+		checkf(TailCPPType == ScriptStruct->GetStructCPPName(), TEXT("CPPType '%s' doesn't match provided Struct '%s'"),
+			*TailCPPType, *ScriptStruct->GetStructCPPName());
+	}
 }
 
 FName FRigVMPropertyDescription::SanitizeName(const FName& InName)
@@ -155,6 +168,37 @@ FString FRigVMPropertyDescription::GetTailCPPType() const
 	return TailCPPType;
 }
 
+bool FRigVMPropertyDescription::RequiresCPPTypeObject(const FString& InCPPType)
+{
+	static const TArray<FString> PrefixesRequiringCPPTypeObject = {
+		TEXT("F"), 
+		TEXT("TArray<F"), 
+		TEXT("E"), 
+		TEXT("TArray<E"), 
+		TEXT("U"), 
+		TEXT("TArray<U")
+	};
+	static const TArray<FString> CPPTypesNotRequiringCPPTypeObject = {
+		TEXT("FString"), 
+		TEXT("TArray<FString>"), 
+		TEXT("FName"), 
+		TEXT("TArray<FName>") 
+	};
+
+	if(!CPPTypesNotRequiringCPPTypeObject.Contains(InCPPType))
+	{
+		for(const FString& Prefix : PrefixesRequiringCPPTypeObject)
+		{
+			if(InCPPType.StartsWith(Prefix, ESearchCase::CaseSensitive))
+			{
+				return true;
+			}
+		}
+	}
+
+	return false;
+}
+
 ////////////////////////////////////////////////////////////////////////////////////////////////////////////////////////
 
 void URigVMMemoryStorageGeneratorClass::PurgeClass(bool bRecompilingOnLoad)
diff --git a/Engine/Source/Runtime/RigVM/Public/RigVMCore/RigVMMemoryStorage.h b/Engine/Source/Runtime/RigVM/Public/RigVMCore/RigVMMemoryStorage.h
--- a/Engine/Source/Runtime/RigVM/Public/RigVMCore/RigVMMemoryStorage.h
+++ b/Engine/Source/Runtime/RigVM/Public/RigVMCore/RigVMMemoryStorage.h
@@ -193,6 +193,8 @@ public:
 	// Returns the CPP type of the tail property, for ex: '[2].Translation' it is 'FVector'
 	FString GetTailCPPType() const;
 
+	static bool RequiresCPPTypeObject(const FString& InCPPType);
+
 private:
 	
 	static const FString ArrayPrefix;
