Control Rig: Fix adding / importing / copy pasting variable nodes

diff --git a/Engine/Source/Developer/RigVMDeveloper/Private/RigVMModel/RigVMController.cpp b/Engine/Source/Developer/RigVMDeveloper/Private/RigVMModel/RigVMController.cpp
--- a/Engine/Source/Developer/RigVMDeveloper/Private/RigVMModel/RigVMController.cpp
+++ b/Engine/Source/Developer/RigVMDeveloper/Private/RigVMModel/RigVMController.cpp
@@ -9366,22 +9366,14 @@ bool URigVMController::CanAddNode(URigVMNode* InNode, bool bReportErrors, bool b
 			FString VarName = NamePin->GetDefaultValue();
 			if (!VarName.IsEmpty())
 			{
-				for (auto LocalVar : Graph->GetLocalVariables())
+				TArray<FRigVMExternalVariable> AllVariables = URigVMController::GetAllVariables();
+				for(const FRigVMExternalVariable& Variable : AllVariables)
 				{
-					if (LocalVar.Name.ToString() == VarName)
+					if(Variable.Name.ToString() == VarName)
 					{
 						return true;
 					}
 				}
-
-				for (auto ExternalVar : Graph->GetExternalVariables())
-				{
-					if (ExternalVar.Name.ToString() == VarName)
-					{
-						return true;
-					}
-				}
-
 				return false;
 			}
 		}
