Control Rig: Fix Local variables default value overwritting previous results

diff --git a/Engine/Source/Developer/RigVMDeveloper/Private/RigVMCompiler/RigVMCompiler.cpp b/Engine/Source/Developer/RigVMDeveloper/Private/RigVMCompiler/RigVMCompiler.cpp
--- a/Engine/Source/Developer/RigVMDeveloper/Private/RigVMCompiler/RigVMCompiler.cpp
+++ b/Engine/Source/Developer/RigVMDeveloper/Private/RigVMCompiler/RigVMCompiler.cpp
@@ -1698,50 +1698,25 @@ void URigVMCompiler::InitializeLocalVariables(const FRigVMExprAST* InExpr, FRigV
 			const FRigVMCallstack& Callstack = Proxy->GetCallstack();
 			ensure(Callstack.Num() > 0);
 
-			// Look at previous instructions to get the last callstack
-			const TArray<UObject*>* PreviousCallstack = nullptr;
-			int32 PreviousIndex = ByteCode.GetNumInstructions()-1;
-			TArray<FRigVMCopyOp> PreviousCopyInstructions;
-			while (PreviousCallstack == nullptr)
+			// Find all function references in the callstack and initialize their local variables if necessary
+			for (int32 SubjectIndex=0; SubjectIndex<Callstack.Num(); ++SubjectIndex)
 			{
-				if (PreviousIndex < 0)
+				if (const URigVMLibraryNode* Node = Cast<const URigVMLibraryNode>(Callstack[SubjectIndex]))
 				{
-					break;
-				}
-				PreviousCallstack = ByteCode.GetCallstackForInstruction(PreviousIndex);
-
-				// Other recursive calls might have already initialized local variables, we don't want
-				// to initialize them again
-				uint64 ByteCodeIndex = ByteCode.GetInstructions()[PreviousIndex].ByteCodeIndex;
-				ERigVMOpCode Code = ByteCode.GetOpCodeAt(ByteCodeIndex);
-				if (ByteCode.GetOpCodeAt(ByteCodeIndex) == ERigVMOpCode::Copy)
-				{
-					const FRigVMCopyOp& CopyOp = ByteCode.GetOpAt<FRigVMCopyOp>(ByteCodeIndex);
-					PreviousCopyInstructions.Add(CopyOp);
-				}					
-			
-				PreviousIndex--;
-			}
-
-			int32 FirstDifferenceIndex = 0;
-			if (PreviousCallstack)
-			{
-				// Compare the two callstacks to find the first difference in the path				
-				for (; FirstDifferenceIndex < PreviousCallstack->Num() && FirstDifferenceIndex < Callstack.Num(); ++FirstDifferenceIndex)
-				{
-					if (PreviousCallstack->operator[](FirstDifferenceIndex) != Callstack[FirstDifferenceIndex])
+					// Check if this is the first time we are accessing this function reference
+					bool bFound = false;
+					for (int32 i=ByteCode.GetNumInstructions()-1; i>0; --i)
 					{
-						break;
+						const TArray<UObject*>* PreviousCallstack = ByteCode.GetCallstackForInstruction(i);
+						if (PreviousCallstack && PreviousCallstack->Contains(Node))
+						{
+							bFound = true;
+							break;
+						}
 					}
-				}
 
-				// Starting from the first difference in the path, look for function references or collapse nodes
-				// to initialize their local variables. 
-				for (int32 Index = FirstDifferenceIndex; Index < Callstack.Num(); ++Index)
-				{
-					// If it is a library node (which will handle function references and collapsed nodes), try to initialize local variables
-					const URigVMLibraryNode* Node = Cast<const URigVMLibraryNode>(Callstack[Index]);
-					if (Node)
+					// If it is the first time we access this function reference, initialize all local variables
+					if (!bFound)
 					{
 						for (FRigVMGraphVariableDescription Variable : Node->GetContainedGraph()->LocalVariables)
 						{
@@ -1753,62 +1728,13 @@ void URigVMCompiler::InitializeLocalVariables(const FRigVMExprAST* InExpr, FRigV
 							{
 								const FRigVMOperand& Source = *SourcePtr;
 								const FRigVMOperand& Target = *TargetPtr;
-
-								bool bAlreadyCopied = false;
-								for (const FRigVMCopyOp& CopyOp : PreviousCopyInstructions)
-								{
-									if (CopyOp.Source == Source && CopyOp.Target == Target)
-									{
-										bAlreadyCopied = true;
-										break;
-									}
-								}
-
-								if (!bAlreadyCopied)
-								{
-									ByteCode.AddCopyOp(WorkData.VM->GetCopyOpForOperands(Source, Target));
-									if(Settings.SetupNodeInstructionIndex)
-									{
-										const int32 InstructionIndex = ByteCode.GetNumInstructions() - 1;
-										WorkData.VM->GetByteCode().SetSubject(InstructionIndex, Callstack.GetCallPath(), Callstack.GetStack());
-									}
-								}	
-							}					
-						}
-					}
-				}
-			}
-
-			// We also need to initialize local variables of the root graph (after the entry instruction is inserted)
-			// If the parent of InExpr is an entry, initialize the root graph's local variables
-			if(InExpr->NumParents() > 0 && InExpr->ParentAt(0)->GetType() == FRigVMExprAST::Entry)
-			{
-				if (UObject* Subject = ByteCode.GetSubjectForInstruction(ByteCode.GetNumInstructions()-1))
-				{
-					URigVMNode* Node = Cast<URigVMNode>(Subject);
-					if (!Node)
-					{
-						URigVMPin* Pin = Cast<URigVMPin>(Subject);
-						Node = Pin->GetNode();
-					}
-					if (Node)
-					{
-						for (FRigVMGraphVariableDescription Variable : Node->GetGraph()->LocalVariables)
-						{
-							FString TargetPath = FString::Printf(TEXT("LocalVariable::%s"), *Variable.Name.ToString());
-							FString SourcePath = FString::Printf(TEXT("LocalVariableDefault::|%s::Const"), *Variable.Name.ToString());
-							FRigVMOperand* TargetPtr = WorkData.PinPathToOperand->Find(TargetPath);
-							FRigVMOperand* SourcePtr = WorkData.PinPathToOperand->Find(SourcePath);
-							if (SourcePtr && TargetPtr)
-							{
-								const FRigVMOperand& Source = *SourcePtr;
-								const FRigVMOperand& Target = *TargetPtr;
+								
 								ByteCode.AddCopyOp(WorkData.VM->GetCopyOpForOperands(Source, Target));
 								if(Settings.SetupNodeInstructionIndex)
 								{
 									const int32 InstructionIndex = ByteCode.GetNumInstructions() - 1;
 									WorkData.VM->GetByteCode().SetSubject(InstructionIndex, Callstack.GetCallPath(), Callstack.GetStack());
-								}
+								}							
 							}					
 						}
 					}
diff --git a/Engine/Source/Developer/RigVMDeveloper/Private/RigVMModel/RigVMController.cpp b/Engine/Source/Developer/RigVMDeveloper/Private/RigVMModel/RigVMController.cpp
--- a/Engine/Source/Developer/RigVMDeveloper/Private/RigVMModel/RigVMController.cpp
+++ b/Engine/Source/Developer/RigVMDeveloper/Private/RigVMModel/RigVMController.cpp
@@ -7689,6 +7689,21 @@ FRigVMGraphVariableDescription URigVMController::AddLocalVariable(const FName& I
 	URigVMGraph* Graph = GetGraph();
 	check(Graph);
 
+	// Check this is the main graph of a function
+	{
+		if (URigVMLibraryNode* LibraryNode = Cast<URigVMLibraryNode>(Graph->GetOuter()))
+		{
+			if (!LibraryNode->GetOuter()->IsA<URigVMFunctionLibrary>())
+			{
+				return NewVariable;
+			}
+		}
+		else
+		{
+			return NewVariable;
+		}
+	}
+
 	FName VariableName = GetUniqueName(InVariableName, [Graph](const FName& InName) {
 		for (FRigVMGraphVariableDescription LocalVariable : Graph->GetLocalVariables())
 		{
