Control Rig: Ensure pin validity after copy & paste

diff --git a/Engine/Source/Developer/RigVMDeveloper/Private/RigVMModel/RigVMController.cpp b/Engine/Source/Developer/RigVMDeveloper/Private/RigVMModel/RigVMController.cpp
--- a/Engine/Source/Developer/RigVMDeveloper/Private/RigVMModel/RigVMController.cpp
+++ b/Engine/Source/Developer/RigVMDeveloper/Private/RigVMModel/RigVMController.cpp
@@ -2466,6 +2466,11 @@ TArray<FName> URigVMController::ImportNodesFromText(const FString& InText, bool
 					}
 				}
 			}
+
+			for(URigVMPin* Pin : SubNode->Pins)
+			{
+				EnsurePinValidity(Pin, true);
+			}
 		}
 
 		Notify(ERigVMGraphNotifType::NodeAdded, CreatedNode);
