Control Rig: Guard against array types for local variables

diff --git a/Engine/Source/Developer/RigVMDeveloper/Private/RigVMModel/RigVMController.cpp b/Engine/Source/Developer/RigVMDeveloper/Private/RigVMModel/RigVMController.cpp
--- a/Engine/Source/Developer/RigVMDeveloper/Private/RigVMModel/RigVMController.cpp
+++ b/Engine/Source/Developer/RigVMDeveloper/Private/RigVMModel/RigVMController.cpp
@@ -7990,9 +7990,7 @@ bool URigVMController::SetLocalVariableType(const FName& InVariableName, const F
 		FRigVMInverseAction InverseAction;
 		InverseAction.Title = FString::Printf(TEXT("Change Local Variable type %s to %s"), *InVariableName.ToString(), *InCPPType);
 
-		ActionStack->BeginAction(InverseAction);
 		ActionStack->AddAction(FRigVMChangeLocalVariableTypeAction(LocalVariables[FoundIndex], InCPPType, InCPPTypeObject));
-		ActionStack->EndAction(InverseAction);
 	}	
 	
 	LocalVariables[FoundIndex].CPPType = InCPPType;
diff --git a/Engine/Source/Developer/RigVMDeveloper/Private/RigVMModel/RigVMPin.cpp b/Engine/Source/Developer/RigVMDeveloper/Private/RigVMModel/RigVMPin.cpp
--- a/Engine/Source/Developer/RigVMDeveloper/Private/RigVMModel/RigVMPin.cpp
+++ b/Engine/Source/Developer/RigVMDeveloper/Private/RigVMModel/RigVMPin.cpp
@@ -613,6 +613,21 @@ bool URigVMPin::IsValidDefaultValue(const FString& InDefaultValue) const
 
 	if (IsArray())
 	{
+		if(InDefaultValue.IsEmpty())
+		{
+			return false;
+		}
+		
+		if(InDefaultValue[0] != TCHAR('('))
+		{
+			return false;
+		}
+
+		if(InDefaultValue[InDefaultValue.Len() - 1] != TCHAR(')'))
+		{
+			return false;
+		}
+
 		DefaultValues = URigVMPin::SplitDefaultValue(InDefaultValue);
 	}
 	else
diff --git a/Engine/Source/Developer/RigVMDeveloper/Public/RigVMModel/RigVMVariableDescription.h b/Engine/Source/Developer/RigVMDeveloper/Public/RigVMModel/RigVMVariableDescription.h
--- a/Engine/Source/Developer/RigVMDeveloper/Public/RigVMModel/RigVMVariableDescription.h
+++ b/Engine/Source/Developer/RigVMDeveloper/Public/RigVMModel/RigVMVariableDescription.h
@@ -72,27 +72,39 @@ public:
 		PinType.ResetToDefaults();
 		PinType.PinCategory = NAME_None;
 
-		if (CPPType == TEXT("bool"))
+		FString CurrentCPPType = CPPType;
+		if (CurrentCPPType.StartsWith(TEXT("TArray<")))
+		{
+			PinType.ContainerType = EPinContainerType::Array;
+			CurrentCPPType.RemoveFromStart(TEXT("TArray<"));
+			CurrentCPPType.RemoveFromEnd(TEXT(">"));
+		}
+		else
+		{
+			PinType.ContainerType = EPinContainerType::None;
+		}
+
+		if (CurrentCPPType == TEXT("bool"))
 		{
 			PinType.PinCategory = UEdGraphSchema_K2::PC_Boolean;
 		}
-		else if (CPPType == TEXT("int32"))
+		else if (CurrentCPPType == TEXT("int32"))
 		{
 			PinType.PinCategory = UEdGraphSchema_K2::PC_Int;
 		}
-		else if (CPPType == TEXT("float"))
+		else if (CurrentCPPType == TEXT("float"))
 		{
 			PinType.PinCategory = UEdGraphSchema_K2::PC_Float;
 		}
-		else if (CPPType == TEXT("double"))
+		else if (CurrentCPPType == TEXT("double"))
 		{
 			PinType.PinCategory = UEdGraphSchema_K2::PC_Double;
 		}
-		else if (CPPType == TEXT("FName"))
+		else if (CurrentCPPType == TEXT("FName"))
 		{
 			PinType.PinCategory = UEdGraphSchema_K2::PC_Name;
 		}
-		else if (CPPType == TEXT("FString"))
+		else if (CurrentCPPType == TEXT("FString"))
 		{
 			PinType.PinCategory = UEdGraphSchema_K2::PC_String;
 		}
@@ -107,15 +119,6 @@ public:
 			PinType.PinSubCategoryObject = CPPTypeObject;
 		}
 
-		if (CPPType.StartsWith(TEXT("TArray<")))
-		{
-			PinType.ContainerType = EPinContainerType::Array;
-		}
-		else
-		{
-			PinType.ContainerType = EPinContainerType::None;
-		}
-
 		return PinType;
 	}
 
