Control Rig: Add warning message for unconnected execute pins on Entry / Return nodes

diff --git a/Engine/Source/Developer/RigVMDeveloper/Private/RigVMCompiler/RigVMCompiler.cpp b/Engine/Source/Developer/RigVMDeveloper/Private/RigVMCompiler/RigVMCompiler.cpp
--- a/Engine/Source/Developer/RigVMDeveloper/Private/RigVMCompiler/RigVMCompiler.cpp
+++ b/Engine/Source/Developer/RigVMDeveloper/Private/RigVMCompiler/RigVMCompiler.cpp
@@ -192,6 +192,19 @@ bool URigVMCompiler::Compile(URigVMGraph* InGraph, URigVMController* InControlle
 				}
 			}
 
+			if(ModelNode->IsA<URigVMFunctionEntryNode>() || ModelNode->IsA<URigVMFunctionReturnNode>())
+			{
+				static const FString ExecuteContextName = FRigVMStruct::ExecuteContextName.ToString();
+				if(URigVMPin* ExecutePin = ModelNode->FindPin(ExecuteContextName))
+				{
+					if(ExecutePin->GetLinks().Num() == 0)
+					{
+						static const FString UnlinkedExecuteMessage = TEXT("Node @@ has an unconnected Execute pin.\nThe function might cause unexpected behavior.");
+						Settings.ASTSettings.Report(EMessageSeverity::Warning, ModelNode, UnlinkedExecuteMessage);
+					}
+				}
+			}
+
 			if(URigVMLibraryNode* LibraryNode = Cast<URigVMLibraryNode>(ModelNode))
 			{
 				if(URigVMGraph* ContainedGraph = LibraryNode->GetContainedGraph())
