Control Rig: Copy Python Command leads to infinite loop for graphs that contains dangling reroutes

diff --git a/Engine/Source/Developer/RigVMDeveloper/Private/RigVMModel/RigVMController.cpp b/Engine/Source/Developer/RigVMDeveloper/Private/RigVMModel/RigVMController.cpp
--- a/Engine/Source/Developer/RigVMDeveloper/Private/RigVMModel/RigVMController.cpp
+++ b/Engine/Source/Developer/RigVMDeveloper/Private/RigVMModel/RigVMController.cpp
@@ -330,9 +330,18 @@ TArray<FString> URigVMController::GeneratePythonCommands()
 				{
 					SourcePin = Links[0]->GetSourcePin();
 				}
+				else
+				{
+					break;
+				}
 			}
 		}
 
+		if (SourcePin->GetNode()->IsA<URigVMRerouteNode>())
+		{
+			continue;
+		}
+
 		//bool AddLink(const FString& InOutputPinPath, const FString& InInputPinPath, bool bSetupUndoRedo = true);
 		Commands.Add(FString::Printf(TEXT("blueprint.get_controller_by_name('%s').add_link('%s', '%s')"),
 					*GetGraph()->GetGraphName(),
@@ -347,54 +356,108 @@ TArray<FString> URigVMController::GeneratePythonCommands()
 		{
 			if (URigVMRerouteNode* Reroute = Cast<URigVMRerouteNode>(Node))
 			{
-				Reroutes.Add(Reroute);
+				if (Reroute->Pins[0]->GetTargetLinks().Num() > 0)
+				{
+					Reroutes.Add(Reroute);
+				}
 			}
 		}
 
 		TArray<FString> ReroutesAdded;
-		while (ReroutesAdded.Num() < Reroutes.Num())
+		bool bNodesAdded = false;
+		do
 		{
+			bNodesAdded = false;
 			for (URigVMRerouteNode* Reroute : Reroutes)
 			{
+				if (ReroutesAdded.Contains(Reroute->GetName()))
+				{
+					continue;
+				}
+
+				// If this reroute has no target links, we will ignore it (it will not be created)
 				if (Reroute->Pins.IsEmpty() || Reroute->Pins[0]->GetTargetLinks().IsEmpty())
 				{
 					continue;
 				}
-				
-				URigVMPin* TargetPin = Reroute->Pins[0]->GetTargetLinks()[0]->GetTargetPin();
-				if (!TargetPin->GetNode()->IsA<URigVMRerouteNode>() ||
-					ReroutesAdded.Contains(TargetPin->GetNode()->GetName()))
+
+				TArray<URigVMPin*> TargetPins;
+				for (URigVMLink* Link : Reroute->Pins[0]->GetTargetLinks())
 				{
-					// Iterate upstream until the source pin is not on a reroute node
-					URigVMPin* SourcePin = Reroute->Pins[0]->GetSourceLinks()[0]->GetSourcePin();
-					while (SourcePin && SourcePin->GetNode()->IsA<URigVMRerouteNode>())
+					URigVMPin* TargetPin = Link->GetTargetPin();
+					if (!TargetPin->GetNode()->IsA<URigVMRerouteNode>() || ReroutesAdded.Contains(TargetPin->GetNode()->GetName()))
 					{
-						const URigVMNode* Node = SourcePin->GetNode();
-						if (Node->Pins.Num() > 0)
+						TargetPins.Add(TargetPin);
+						continue;
+					}
+				}
+
+				if (!TargetPins.IsEmpty())
+				{
+					URigVMPin* TargetPin = TargetPins[0];
+					// AddRerouteNodeOnPin(const FString& InPinPath, bool bAsInput, bool bShowAsFullNode, const FVector2D& InPosition, const FString& InNodeName, bool bSetupUndoRedo, bool bPrintPythonCommand)
+					Commands.Add(FString::Printf(TEXT("blueprint.get_controller_by_name('%s').add_reroute_node_on_pin('%s', True, %s, %s, '%s')"),
+							*GetGraph()->GetGraphName(),
+							*TargetPin->GetPinPath(),
+							Reroute->GetShowsAsFullNode() ? TEXT("True") : TEXT("False"),
+							*RigVMPythonUtils::Vector2DToPythonString(Reroute->GetPosition()),
+							*Reroute->GetName()));
+
+					ReroutesAdded.Add(Reroute->GetName());
+					bNodesAdded = true;
+
+					// Add the rest of target links
+					for (int32 i=1; i<TargetPins.Num(); ++i)
+					{
+						URigVMPin* OtherTargetPin = Reroute->Pins[0]->GetTargetLinks()[i]->GetTargetPin();
+		
+						if (OtherTargetPin->GetNode()->IsA<URigVMRerouteNode>())
 						{
-							const TArray<URigVMLink*>& Links = Node->Pins[0]->GetSourceLinks();
-							if (Links.Num() > 0)
-							{
-								SourcePin = Links[0]->GetSourcePin();
-							}
+							continue;
 						}
+
+						if (OtherTargetPin->GetInjectedNodes().Num() > 0)
+						{
+							continue;
+						}
+
+						//bool AddLink(const FString& InOutputPinPath, const FString& InInputPinPath, bool bSetupUndoRedo = true);
+						Commands.Add(FString::Printf(TEXT("blueprint.get_controller_by_name('%s').add_link('%s', '%s')"),
+									*GetGraph()->GetGraphName(),
+									*Reroute->Pins[0]->GetPinPath(),
+									*OtherTargetPin->GetPinPath()));
 					}
 
-					if(SourcePin)
+					// Set default value in input if necessary
+					URigVMPin* Pin = Reroute->Pins[0];
+					const FString DefaultValue = Pin->GetDefaultValue();
+					if (!DefaultValue.IsEmpty() && DefaultValue != TEXT("()"))
 					{
-						// AddRerouteNodeOnLinkPath(const FString& InLinkPinPathRepresentation, bool bShowAsFullNode, const FVector2D& InPosition = FVector2D::ZeroVector, const FString& InNodeName = TEXT(""), bool bSetupUndoRedo = true);
-						Commands.Add(FString::Printf(TEXT("blueprint.get_controller_by_name('%s').add_reroute_node_on_link_path('%s', %s, %s, '%s')"),
-								*GetGraph()->GetGraphName(),
-								*FString::Printf(TEXT("%s -> %s"), *SourcePin->GetPinPath(), *TargetPin->GetPinPath()),
-								Reroute->GetShowsAsFullNode() ? TEXT("True") : TEXT("False"),
-								*RigVMPythonUtils::Vector2DToPythonString(Reroute->GetPosition()),
-								*Reroute->GetName()));
-
-						ReroutesAdded.Add(Reroute->GetName());
+						Commands.Add(FString::Printf(TEXT("blueprint.get_controller_by_name('%s').set_pin_default_value('%s', '%s')"),
+									*GetGraph()->GetGraphName(),
+									*Pin->GetPinPath(),
+									*Pin->GetDefaultValue()));
 					}
-				}
+					
+					// Add source links					
+					TArray<URigVMLink*> SourceLinks = Pin->GetSourceLinks(true);
+					for (URigVMLink* Link : SourceLinks)
+					{
+						URigVMPin* SourcePin = Link->GetSourcePin(); 
+						if (SourcePin->GetNode()->IsA<URigVMRerouteNode>())
+						{
+							continue;
+						}
+
+						//bool AddLink(const FString& InOutputPinPath, const FString& InInputPinPath, bool bSetupUndoRedo = true);
+						Commands.Add(FString::Printf(TEXT("blueprint.get_controller_by_name('%s').add_link('%s', '%s')"),
+									*GetGraph()->GetGraphName(),
+									*SourcePin->GetPinPath(),
+									*Link->GetTargetPin()->GetPinPath()));
+					}				
+				}				
 			}
-		}
+		} while (bNodesAdded);
 	}
 
 	return Commands;
