Control Rig: Fix Visual Debug node becomes unresponsive in collapsed function

diff --git a/Engine/Source/Developer/RigVMDeveloper/Private/RigVMModel/RigVMController.cpp b/Engine/Source/Developer/RigVMDeveloper/Private/RigVMModel/RigVMController.cpp
--- a/Engine/Source/Developer/RigVMDeveloper/Private/RigVMModel/RigVMController.cpp
+++ b/Engine/Source/Developer/RigVMDeveloper/Private/RigVMModel/RigVMController.cpp
@@ -4142,6 +4142,11 @@ URigVMFunctionReferenceNode* URigVMController::PromoteCollapseNodeToFunctionRefe
 			TArray<FName> NodeNames;
 			for (const URigVMNode* Node : InCollapseNode->GetContainedNodes())
 			{
+				if (Node->IsInjected())
+				{
+					continue;
+				}
+				
 				NodeNames.Add(Node->GetFName());
 			}
 			TextContent = ExportNodesToText(NodeNames);
