Control Rig: Fix crash in unit test

diff --git a/Engine/Source/Developer/RigVMDeveloper/Private/RigVMCompiler/RigVMCompiler.cpp b/Engine/Source/Developer/RigVMDeveloper/Private/RigVMCompiler/RigVMCompiler.cpp
--- a/Engine/Source/Developer/RigVMDeveloper/Private/RigVMCompiler/RigVMCompiler.cpp
+++ b/Engine/Source/Developer/RigVMDeveloper/Private/RigVMCompiler/RigVMCompiler.cpp
@@ -264,13 +264,16 @@ bool URigVMCompiler::Compile(URigVMGraph* InGraph, URigVMController* InControlle
 	// If in editor, make sure we visit all the graphs to initialize local variables
 	// in case the user wants to edit default values
 	URigVMFunctionLibrary* FunctionLibrary = InGraph->GetDefaultFunctionLibrary();
-	for (URigVMLibraryNode* LibraryNode : FunctionLibrary->GetFunctions())
+	if (FunctionLibrary)
 	{
-		for (FRigVMGraphVariableDescription& Variable : LibraryNode->GetContainedGraph()->LocalVariables)
+		for (URigVMLibraryNode* LibraryNode : FunctionLibrary->GetFunctions())
 		{
-			FString Path = FString::Printf(TEXT("LocalVariableDefault::%s|%s::Const"), *LibraryNode->GetFName().ToString(), *Variable.Name.ToString());
-			FRigVMOperand Operand = WorkData.AddProperty(ERigVMMemoryType::Literal, *Path, Variable.CPPType, Variable.CPPTypeObject, Variable.DefaultValue);
-			WorkData.PinPathToOperand->Add(Path, Operand);
+			for (FRigVMGraphVariableDescription& Variable : LibraryNode->GetContainedGraph()->LocalVariables)
+			{
+				FString Path = FString::Printf(TEXT("LocalVariableDefault::%s|%s::Const"), *LibraryNode->GetFName().ToString(), *Variable.Name.ToString());
+				FRigVMOperand Operand = WorkData.AddProperty(ERigVMMemoryType::Literal, *Path, Variable.CPPType, Variable.CPPTypeObject, Variable.DefaultValue);
+				WorkData.PinPathToOperand->Add(Path, Operand);
+			}
 		}
 	}
 #endif
