Control Rig: Fix unit test promote to function failing

diff --git a/Engine/Source/Developer/RigVMDeveloper/Private/RigVMModel/RigVMController.cpp b/Engine/Source/Developer/RigVMDeveloper/Private/RigVMModel/RigVMController.cpp
--- a/Engine/Source/Developer/RigVMDeveloper/Private/RigVMModel/RigVMController.cpp
+++ b/Engine/Source/Developer/RigVMDeveloper/Private/RigVMModel/RigVMController.cpp
@@ -2568,6 +2568,26 @@ TArray<FName> URigVMController::ImportNodesFromText(const FString& InText, bool
 					CreatedLink->TargetPinPath = URigVMPin::JoinPinPath(NewTargetNodeName->ToString(), TargetRight);
 					URigVMPin* SourcePin = CreatedLink->GetSourcePin();
 					URigVMPin* TargetPin = CreatedLink->GetTargetPin();
+
+					if (SourcePin == nullptr)
+					{
+						URigVMNode* OriginalNode = Graph->FindNode(SourceLeft);
+						if (OriginalNode && OriginalNode->IsA<URigVMFunctionEntryNode>())
+						{
+							CreatedLink->SourcePinPath = URigVMPin::JoinPinPath(SourceLeft, SourceRight);
+							SourcePin = CreatedLink->GetSourcePin();							
+						}
+					}
+					if (TargetPin == nullptr)
+					{
+						URigVMNode* OriginalNode = Graph->FindNode(TargetLeft);
+						if (OriginalNode && OriginalNode->IsA<URigVMFunctionReturnNode>())
+						{
+							CreatedLink->TargetPinPath = URigVMPin::JoinPinPath(TargetLeft, TargetRight);
+							TargetPin = CreatedLink->GetTargetPin();							
+						}
+					}
+					
 					if (SourcePin && TargetPin)
 					{
 						Graph->Links.Add(CreatedLink);
@@ -4166,7 +4186,15 @@ URigVMFunctionReferenceNode* URigVMController::PromoteCollapseNodeToFunctionRefe
 
 			for (const URigVMLink* InnerLink : InCollapseNode->GetContainedGraph()->GetLinks())
 			{
-				AddLink(InnerLink->SourcePinPath, InnerLink->TargetPinPath, bSetupUndoRedo);
+				URigVMPin* SourcePin = InCollapseNode->GetGraph()->FindPin(InnerLink->SourcePinPath);
+				URigVMPin* TargetPin = InCollapseNode->GetGraph()->FindPin(InnerLink->TargetPinPath);
+				if (SourcePin && TargetPin)
+				{
+					if (!SourcePin->IsLinkedTo(TargetPin))
+					{
+						AddLink(InnerLink->SourcePinPath, InnerLink->TargetPinPath, bSetupUndoRedo);	
+					}
+				}				
 			}
 		}
 	}
