Control Rig: Accelerate loading of control rig assets

diff --git a/Engine/Plugins/Animation/ControlRig/Source/ControlRigDeveloper/Private/ControlRigBlueprint.cpp b/Engine/Plugins/Animation/ControlRig/Source/ControlRigDeveloper/Private/ControlRigBlueprint.cpp
--- a/Engine/Plugins/Animation/ControlRig/Source/ControlRigDeveloper/Private/ControlRigBlueprint.cpp
+++ b/Engine/Plugins/Animation/ControlRig/Source/ControlRigDeveloper/Private/ControlRigBlueprint.cpp
@@ -344,6 +344,8 @@ void UControlRigBlueprint::PostLoad()
 {
 	Super::PostLoad();
 
+	IncrementVMRecompileBracket();
+
 	TArray<UControlRigBlueprint*> ReferencedBlueprints = GetReferencedControlRigBlueprints();
 
 	// PostLoad all referenced BPs so that their function graphs are fully loaded 
@@ -503,7 +505,9 @@ void UControlRigBlueprint::PostLoad()
 
 	CompileLog.Messages.Reset();
 	CompileLog.NumErrors = CompileLog.NumWarnings = 0;
+
 	
+	DecrementVMRecompileBracket();
 	RecompileVM();
 	RequestControlRigInit();
 
diff --git a/Engine/Source/Developer/RigVMDeveloper/Private/RigVMModel/RigVMController.cpp b/Engine/Source/Developer/RigVMDeveloper/Private/RigVMModel/RigVMController.cpp
--- a/Engine/Source/Developer/RigVMDeveloper/Private/RigVMModel/RigVMController.cpp
+++ b/Engine/Source/Developer/RigVMDeveloper/Private/RigVMModel/RigVMController.cpp
@@ -10432,6 +10432,8 @@ void URigVMController::RepopulatePinsOnNode(URigVMNode* InNode, bool bFollowCore
 		return;
 	}
 
+	Notify(ERigVMGraphNotifType::InteractionBracketOpened, nullptr);
+
 	URigVMUnitNode* UnitNode = Cast<URigVMUnitNode>(InNode);
 	URigVMRerouteNode* RerouteNode = Cast<URigVMRerouteNode>(InNode);
 	URigVMFunctionEntryNode* EntryNode = Cast<URigVMFunctionEntryNode>(InNode);
@@ -10657,6 +10659,8 @@ void URigVMController::RepopulatePinsOnNode(URigVMNode* InNode, bool bFollowCore
 		InjectionInfo->InputPin = InNode->FindPin(InjectionInputPinName.ToString());
 		InjectionInfo->OutputPin = InNode->FindPin(InjectionOutputPinName.ToString());
 	}
+
+	Notify(ERigVMGraphNotifType::InteractionBracketClosed, nullptr);
 }
 
 void URigVMController::RemovePinsDuringRepopulate(URigVMNode* InNode, TArray<URigVMPin*>& InPins, bool bNotify, bool bSetupOrphanedPins)
@@ -10928,6 +10932,7 @@ void URigVMController::ApplyPinState(URigVMPin* InPin, const FPinState& InPinSta
 
 void URigVMController::ApplyPinStates(URigVMNode* InNode, const TMap<FString, URigVMController::FPinState>& InPinStates, const TMap<FString, FString>& InRedirectedPinPaths)
 {
+	Notify(ERigVMGraphNotifType::InteractionBracketOpened, nullptr);
 	for (const TPair<FString, FPinState>& PinStatePair : InPinStates)
 	{
 		FString PinPath = PinStatePair.Key;
@@ -10951,6 +10956,7 @@ void URigVMController::ApplyPinStates(URigVMNode* InNode, const TMap<FString, UR
 			}
 		}
 	}
+	Notify(ERigVMGraphNotifType::InteractionBracketClosed, nullptr);
 }
 
 void URigVMController::ReportWarning(const FString& InMessage)
