Control Rig: Support for unknown pins + spawning nodes with unknown pins

diff --git a/Engine/Plugins/Animation/ControlRig/Source/ControlRigDeveloper/Private/Graph/ControlRigGraphSchema.cpp b/Engine/Plugins/Animation/ControlRig/Source/ControlRigDeveloper/Private/Graph/ControlRigGraphSchema.cpp
--- a/Engine/Plugins/Animation/ControlRig/Source/ControlRigDeveloper/Private/Graph/ControlRigGraphSchema.cpp
+++ b/Engine/Plugins/Animation/ControlRig/Source/ControlRigDeveloper/Private/Graph/ControlRigGraphSchema.cpp
@@ -25,6 +25,7 @@
 #include "RigVMModel/Nodes/RigVMFunctionEntryNode.h"
 #include "RigVMModel/Nodes/RigVMFunctionReturnNode.h"
 #include "RigVMModel/RigVMVariableDescription.h"
+#include "RigVMCore/RigVMUnknownType.h"
 #include "Kismet2/Kismet2NameValidators.h"
 
 #if WITH_EDITOR
@@ -629,7 +630,12 @@ FLinearColor UControlRigGraphSchema::GetPinTypeColor(const FEdGraphPinType& PinT
 			{
 				return FLinearColor::White;
 			}
-			
+
+			if (Struct->IsChildOf(FRigVMUnknownType::StaticStruct()))
+			{
+				return FLinearColor(FVector3f::OneVector * 0.25f);
+			}
+
 			if (Struct == FRigElementKey::StaticStruct() || Struct == FRigElementKeyCollection::StaticStruct())
 			{
 				return FLinearColor(0.0, 0.6588, 0.9490);
@@ -672,7 +678,7 @@ void UControlRigGraphSchema::InsertAdditionalActions(TArray<UBlueprint*> InBluep
 		{
 			if(URigVMPin* ModelPin = RigNode->GetModelPinFromPinPath(EdGraphPins[0]->GetName()))
 			{
-				if(!ModelPin->IsExecuteContext())
+				if(!ModelPin->IsExecuteContext() && !ModelPin->IsUnknownType())
 				{
 					if(!ModelPin->GetNode()->IsA<URigVMVariableNode>())
 					{
@@ -1361,14 +1367,14 @@ bool UControlRigGraphSchema::ArePinsCompatible(const UEdGraphPin* PinA, const UE
 	}
 
 	// for reroute nodes - always allow it
-	if (PinA->PinType.PinCategory == TEXT("POLYMORPH"))
+	if (PinA->PinType.PinCategory == TEXT("ANY_TYPE"))
 	{
 		UControlRigGraphSchema* MutableThis = (UControlRigGraphSchema*)this;
 		MutableThis->LastPinForCompatibleCheck = PinB;
 		MutableThis->bLastPinWasInput = PinB->Direction == EGPD_Input;
 		return true;
 	}
-	if (PinB->PinType.PinCategory == TEXT("POLYMORPH"))
+	if (PinB->PinType.PinCategory == TEXT("ANY_TYPE"))
 	{
 		UControlRigGraphSchema* MutableThis = (UControlRigGraphSchema*)this;
 		MutableThis->LastPinForCompatibleCheck = PinA;
@@ -1376,6 +1382,42 @@ bool UControlRigGraphSchema::ArePinsCompatible(const UEdGraphPin* PinA, const UE
 		return true;
 	}
 
+	// if we are looking at a polymorphic node
+	if((PinA->PinType.ContainerType == PinB->PinType.ContainerType) ||
+		(PinA->PinType.PinSubCategoryObject != PinB->PinType.PinSubCategoryObject))
+	{
+		if(PinA->PinType.PinCategory == UEdGraphSchema_K2::PC_Struct && PinA->PinType.PinSubCategoryObject == FRigVMUnknownType::StaticStruct())
+		{
+			bool bIsExecuteContext = false;
+			if(const UScriptStruct* ExecuteContextScriptStruct = Cast<UScriptStruct>(PinB->PinType.PinSubCategoryObject))
+			{
+				bIsExecuteContext = ExecuteContextScriptStruct->IsChildOf(FRigVMExecuteContext::StaticStruct());
+			}
+			if(!bIsExecuteContext)
+			{
+				UControlRigGraphSchema* MutableThis = (UControlRigGraphSchema*)this;
+				MutableThis->LastPinForCompatibleCheck = PinB;
+				MutableThis->bLastPinWasInput = PinB->Direction == EGPD_Input;
+				return true;
+			}
+		}
+ 		else if(PinB->PinType.PinCategory == UEdGraphSchema_K2::PC_Struct && PinB->PinType.PinSubCategoryObject == FRigVMUnknownType::StaticStruct())
+		{
+			bool bIsExecuteContext = false;
+ 			if(const UScriptStruct* ExecuteContextScriptStruct = Cast<UScriptStruct>(PinA->PinType.PinSubCategoryObject))
+ 			{
+ 				bIsExecuteContext = ExecuteContextScriptStruct->IsChildOf(FRigVMExecuteContext::StaticStruct());
+ 			}
+ 			if(!bIsExecuteContext)
+ 			{
+ 				UControlRigGraphSchema* MutableThis = (UControlRigGraphSchema*)this;
+ 				MutableThis->LastPinForCompatibleCheck = PinA;
+ 				MutableThis->bLastPinWasInput = PinA->Direction == EGPD_Input;
+ 				return true;
+ 			}
+		}
+	}
+
 	// if large world coordinate support is enabled we should allow connections
 	// between float and double
 #if !UE_LARGE_WORLD_COORDINATES_DISABLED
diff --git a/Engine/Plugins/Animation/ControlRig/Source/ControlRigEditor/Private/Graph/NodeSpawners/ControlRigArrayNodeSpawner.cpp b/Engine/Plugins/Animation/ControlRig/Source/ControlRigEditor/Private/Graph/NodeSpawners/ControlRigArrayNodeSpawner.cpp
--- a/Engine/Plugins/Animation/ControlRig/Source/ControlRigEditor/Private/Graph/NodeSpawners/ControlRigArrayNodeSpawner.cpp
+++ b/Engine/Plugins/Animation/ControlRig/Source/ControlRigEditor/Private/Graph/NodeSpawners/ControlRigArrayNodeSpawner.cpp
@@ -7,6 +7,7 @@
 #include "Graph/ControlRigGraphSchema.h"
 #include "BlueprintNodeTemplateCache.h"
 #include "RigVMModel/RigVMController.h"
+#include "RigVMCore/RigVMUnknownType.h"
 #include "ControlRigBlueprint.h"
 #include "ControlRigBlueprintUtils.h"
 #include "RigVMModel/Nodes/RigVMArrayNode.h"
@@ -140,8 +141,8 @@ UEdGraphNode* UControlRigArrayNodeSpawner::Invoke(UEdGraph* ParentGraph, FBindin
 		NewNode->Pins.Add(InputValuePin);
 		NewNode->Pins.Add(OutputValuePin);
 
-		InputValuePin->PinType.PinCategory = TEXT("POLYMORPH");
-		OutputValuePin->PinType.PinCategory = TEXT("POLYMORPH");
+		InputValuePin->PinType.PinCategory = TEXT("ANY_TYPE");
+		OutputValuePin->PinType.PinCategory = TEXT("ANY_TYPE");
 		InputValuePin->Direction = EGPD_Input;
 		OutputValuePin->Direction = EGPD_Output;
 		NewNode->SetFlags(RF_Transactional);
@@ -158,8 +159,8 @@ UEdGraphNode* UControlRigArrayNodeSpawner::Invoke(UEdGraph* ParentGraph, FBindin
 	UControlRigBlueprint* RigBlueprint = Cast<UControlRigBlueprint>(FBlueprintEditorUtils::FindBlueprintForGraph(ParentGraph));
 	check(RigBlueprint);
 
-	FString CPPType = "int32";
-	FString CPPTypeObjectPath;
+	FString CPPType = FRigVMUnknownType::StaticStruct()->GetStructCPPName();
+	FString CPPTypeObjectPath = FRigVMUnknownType::StaticStruct()->GetPathName();
 
 	if (UControlRigGraphSchema* RigSchema = Cast<UControlRigGraphSchema>((UEdGraphSchema*)ParentGraph->GetSchema()))
 	{
@@ -174,7 +175,11 @@ UEdGraphNode* UControlRigArrayNodeSpawner::Invoke(UEdGraph* ParentGraph, FBindin
 				}
 				if (ModelPin->GetCPPTypeObject())
 				{
-					CPPTypeObjectPath = ModelPin->GetCPPTypeObject()->GetPathName();
+					CPPTypeObjectPath = *ModelPin->GetCPPTypeObject()->GetPathName();
+				}
+				else
+				{
+					CPPTypeObjectPath = FString();
 				}
 			}
 		}
@@ -234,7 +239,7 @@ bool UControlRigArrayNodeSpawner::IsTemplateNodeFilteredOut(FBlueprintActionFilt
 		{
 			if (Filter.Context.Pins.Num() == 0)
 			{
-				return true;
+				return false;
 			}
 
 			FString PinPath = Filter.Context.Pins[0]->GetName();
diff --git a/Engine/Plugins/Animation/ControlRig/Source/ControlRigEditor/Private/Graph/NodeSpawners/ControlRigIfNodeSpawner.cpp b/Engine/Plugins/Animation/ControlRig/Source/ControlRigEditor/Private/Graph/NodeSpawners/ControlRigIfNodeSpawner.cpp
--- a/Engine/Plugins/Animation/ControlRig/Source/ControlRigEditor/Private/Graph/NodeSpawners/ControlRigIfNodeSpawner.cpp
+++ b/Engine/Plugins/Animation/ControlRig/Source/ControlRigEditor/Private/Graph/NodeSpawners/ControlRigIfNodeSpawner.cpp
@@ -7,6 +7,7 @@
 #include "Graph/ControlRigGraphSchema.h"
 #include "BlueprintNodeTemplateCache.h"
 #include "RigVMModel/RigVMController.h"
+#include "RigVMCore/RigVMUnknownType.h"
 #include "ControlRigBlueprint.h"
 #include "ControlRigBlueprintUtils.h"
 #include "RigVMModel/Nodes/RigVMIfNode.h"
@@ -67,8 +68,8 @@ UEdGraphNode* UControlRigIfNodeSpawner::Invoke(UEdGraph* ParentGraph, FBindingSe
 		NewNode->Pins.Add(InputValuePin);
 		NewNode->Pins.Add(OutputValuePin);
 
-		InputValuePin->PinType.PinCategory = TEXT("POLYMORPH");
-		OutputValuePin->PinType.PinCategory = TEXT("POLYMORPH");
+		InputValuePin->PinType.PinCategory = TEXT("ANY_TYPE");
+		OutputValuePin->PinType.PinCategory = TEXT("ANY_TYPE");
 		InputValuePin->Direction = EGPD_Input;
 		OutputValuePin->Direction = EGPD_Output;
 		NewNode->SetFlags(RF_Transactional);
@@ -85,8 +86,8 @@ UEdGraphNode* UControlRigIfNodeSpawner::Invoke(UEdGraph* ParentGraph, FBindingSe
 	UControlRigBlueprint* RigBlueprint = Cast<UControlRigBlueprint>(FBlueprintEditorUtils::FindBlueprintForGraph(ParentGraph));
 	check(RigBlueprint);
 
-	FString CPPType = "int32";
-	FName CPPTypeObjectPath = NAME_None;
+	FString CPPType = FRigVMUnknownType::StaticStruct()->GetStructCPPName();
+	FName CPPTypeObjectPath = *FRigVMUnknownType::StaticStruct()->GetPathName();
 
 	if (UControlRigGraphSchema* RigSchema = Cast<UControlRigGraphSchema>((UEdGraphSchema*)ParentGraph->GetSchema()))
 	{
@@ -99,6 +100,10 @@ UEdGraphNode* UControlRigIfNodeSpawner::Invoke(UEdGraph* ParentGraph, FBindingSe
 				{
 					CPPTypeObjectPath = *ModelPin->GetCPPTypeObject()->GetPathName();
 				}
+				else
+				{
+					CPPTypeObjectPath = NAME_None;
+				}
 			}
 		}
 	}
@@ -156,7 +161,7 @@ bool UControlRigIfNodeSpawner::IsTemplateNodeFilteredOut(FBlueprintActionFilter
 		{
 			if (Filter.Context.Pins.Num() == 0)
 			{
-				return true;
+				return false;
 			}
 
 			FString PinPath = Filter.Context.Pins[0]->GetName();
diff --git a/Engine/Plugins/Animation/ControlRig/Source/ControlRigEditor/Private/Graph/NodeSpawners/ControlRigRerouteNodeSpawner.cpp b/Engine/Plugins/Animation/ControlRig/Source/ControlRigEditor/Private/Graph/NodeSpawners/ControlRigRerouteNodeSpawner.cpp
--- a/Engine/Plugins/Animation/ControlRig/Source/ControlRigEditor/Private/Graph/NodeSpawners/ControlRigRerouteNodeSpawner.cpp
+++ b/Engine/Plugins/Animation/ControlRig/Source/ControlRigEditor/Private/Graph/NodeSpawners/ControlRigRerouteNodeSpawner.cpp
@@ -92,8 +92,8 @@ UEdGraphNode* UControlRigRerouteNodeSpawner::Invoke(UEdGraph* ParentGraph, FBind
 		NewNode->Pins.Add(InputValuePin);
 		NewNode->Pins.Add(OutputValuePin);
 
-		InputValuePin->PinType.PinCategory = TEXT("POLYMORPH");
-		OutputValuePin->PinType.PinCategory = TEXT("POLYMORPH");
+		InputValuePin->PinType.PinCategory = TEXT("ANY_TYPE");
+		OutputValuePin->PinType.PinCategory = TEXT("ANY_TYPE");
 		InputValuePin->Direction = EGPD_Input;
 		OutputValuePin->Direction = EGPD_Output;
 		NewNode->SetFlags(RF_Transactional);
diff --git a/Engine/Plugins/Animation/ControlRig/Source/ControlRigEditor/Private/Graph/NodeSpawners/ControlRigSelectNodeSpawner.cpp b/Engine/Plugins/Animation/ControlRig/Source/ControlRigEditor/Private/Graph/NodeSpawners/ControlRigSelectNodeSpawner.cpp
--- a/Engine/Plugins/Animation/ControlRig/Source/ControlRigEditor/Private/Graph/NodeSpawners/ControlRigSelectNodeSpawner.cpp
+++ b/Engine/Plugins/Animation/ControlRig/Source/ControlRigEditor/Private/Graph/NodeSpawners/ControlRigSelectNodeSpawner.cpp
@@ -7,6 +7,7 @@
 #include "Graph/ControlRigGraphSchema.h"
 #include "BlueprintNodeTemplateCache.h"
 #include "RigVMModel/RigVMController.h"
+#include "RigVMCore/RigVMUnknownType.h"
 #include "ControlRigBlueprint.h"
 #include "ControlRigBlueprintUtils.h"
 
@@ -68,8 +69,8 @@ UEdGraphNode* UControlRigSelectNodeSpawner::Invoke(UEdGraph* ParentGraph, FBindi
 		NewNode->Pins.Add(InputValuePin);
 		NewNode->Pins.Add(OutputValuePin);
 
-		InputValuePin->PinType.PinCategory = TEXT("POLYMORPH");
-		OutputValuePin->PinType.PinCategory = TEXT("POLYMORPH");
+		InputValuePin->PinType.PinCategory = TEXT("ANY_TYPE");
+		OutputValuePin->PinType.PinCategory = TEXT("ANY_TYPE");
 		InputValuePin->Direction = EGPD_Input;
 		OutputValuePin->Direction = EGPD_Output;
 		NewNode->SetFlags(RF_Transactional);
@@ -86,8 +87,9 @@ UEdGraphNode* UControlRigSelectNodeSpawner::Invoke(UEdGraph* ParentGraph, FBindi
 	UControlRigBlueprint* RigBlueprint = Cast<UControlRigBlueprint>(FBlueprintEditorUtils::FindBlueprintForGraph(ParentGraph));
 	check(RigBlueprint);
 
-	FString CPPType = "int32";
-	FName CPPTypeObjectPath = NAME_None;
+	FString CPPType = FRigVMUnknownType::StaticStruct()->GetStructCPPName();
+	FName CPPTypeObjectPath = *FRigVMUnknownType::StaticStruct()->GetPathName();
+	FString LastOutputPinPath;
 
 	if (UControlRigGraphSchema* RigSchema = Cast<UControlRigGraphSchema>((UEdGraphSchema*)ParentGraph->GetSchema()))
 	{
@@ -95,11 +97,20 @@ UEdGraphNode* UControlRigSelectNodeSpawner::Invoke(UEdGraph* ParentGraph, FBindi
 		{
 			if (URigVMPin* ModelPin = RigBlueprint->GetModel(ParentGraph)->FindPin(LastPin->GetName()))
 			{
+				if(LastPin->Direction == EGPD_Output)
+				{
+					LastOutputPinPath = ModelPin->GetPinPath();
+				}
+				
 				CPPType = ModelPin->GetCPPType();
 				if (ModelPin->GetCPPTypeObject())
 				{
 					CPPTypeObjectPath = *ModelPin->GetCPPTypeObject()->GetPathName();
 				}
+				else
+				{
+					CPPTypeObjectPath = NAME_None;
+				}
 			}
 		}
 	}
@@ -128,6 +139,13 @@ UEdGraphNode* UControlRigSelectNodeSpawner::Invoke(UEdGraph* ParentGraph, FBindi
 
 		if (NewNode && bUndo)
 		{
+			if(!LastOutputPinPath.IsEmpty())
+			{
+				if(const URigVMPin* FirstChoicePin = ModelNode->FindPin(TEXT("Values.0")))
+				{
+					Controller->AddLink(LastOutputPinPath, FirstChoicePin->GetPinPath(), true, true);
+				}
+			}
 			Controller->ClearNodeSelection(true);
 			Controller->SelectNode(ModelNode, true, true);
 		}
@@ -157,7 +175,7 @@ bool UControlRigSelectNodeSpawner::IsTemplateNodeFilteredOut(FBlueprintActionFil
 		{
 			if (Filter.Context.Pins.Num() == 0)
 			{
-				return true;
+				return false;
 			}
 
 			FString PinPath = Filter.Context.Pins[0]->GetName();
diff --git a/Engine/Source/Developer/RigVMDeveloper/Private/RigVMCompiler/RigVMCompiler.cpp b/Engine/Source/Developer/RigVMDeveloper/Private/RigVMCompiler/RigVMCompiler.cpp
--- a/Engine/Source/Developer/RigVMDeveloper/Private/RigVMCompiler/RigVMCompiler.cpp
+++ b/Engine/Source/Developer/RigVMDeveloper/Private/RigVMCompiler/RigVMCompiler.cpp
@@ -642,6 +642,10 @@ void URigVMCompiler::TraverseBlock(const FRigVMBlockExprAST* InExpr, FRigVMCompi
 void URigVMCompiler::TraverseEntry(const FRigVMEntryExprAST* InExpr, FRigVMCompilerWorkData& WorkData)
 {
 	URigVMUnitNode* UnitNode = Cast<URigVMUnitNode>(InExpr->GetNode());
+	if(!ValidateNode(UnitNode))
+	{
+		return;
+	}
 
 	if (WorkData.bSetupMemory)
 	{
@@ -694,6 +698,10 @@ void URigVMCompiler::TraverseEntry(const FRigVMEntryExprAST* InExpr, FRigVMCompi
 int32 URigVMCompiler::TraverseCallExtern(const FRigVMCallExternExprAST* InExpr, FRigVMCompilerWorkData& WorkData)
 {
 	URigVMUnitNode* UnitNode = Cast<URigVMUnitNode>(InExpr->GetNode());
+	if(!ValidateNode(UnitNode))
+	{
+		return INDEX_NONE;
+	}
 
 	int32 InstructionIndex = INDEX_NONE;
 
@@ -748,6 +756,11 @@ void URigVMCompiler::TraverseForLoop(const FRigVMCallExternExprAST* InExpr, FRig
 	}
 
 	URigVMUnitNode* UnitNode = Cast<URigVMUnitNode>(InExpr->GetNode());
+	if(!ValidateNode(UnitNode))
+	{
+		return;
+	}
+
 	const FRigVMCallstack Callstack = InExpr->GetProxy().GetCallstack();
 
 	const FRigVMVarExprAST* CompletedExpr = InExpr->FindVarWithPinName(FRigVMStruct::ForLoopCompletedPinName);
@@ -1109,6 +1122,11 @@ void URigVMCompiler::TraverseIf(const FRigVMIfExprAST* InExpr, FRigVMCompilerWor
 	}
 
 	URigVMIfNode* IfNode = Cast<URigVMIfNode>(InExpr->GetNode());
+	if(!ValidateNode(IfNode))
+	{
+		return;
+	}
+
 	const FRigVMCallstack Callstack = InExpr->GetProxy().GetCallstack();
 
 	const FRigVMVarExprAST* ConditionExpr = InExpr->ChildAt<FRigVMVarExprAST>(0);
@@ -1186,6 +1204,11 @@ void URigVMCompiler::TraverseIf(const FRigVMIfExprAST* InExpr, FRigVMCompilerWor
 void URigVMCompiler::TraverseSelect(const FRigVMSelectExprAST* InExpr, FRigVMCompilerWorkData& WorkData)
 {
 	URigVMSelectNode* SelectNode = Cast<URigVMSelectNode>(InExpr->GetNode());
+	if(!ValidateNode(SelectNode))
+	{
+		return;
+	}
+
 	const FRigVMCallstack Callstack = InExpr->GetProxy().GetCallstack();
 
 	int32 NumCases = SelectNode->FindPin(URigVMSelectNode::ValueName)->GetArraySize();
@@ -1351,7 +1374,10 @@ void URigVMCompiler::TraverseSelect(const FRigVMSelectExprAST* InExpr, FRigVMCom
 void URigVMCompiler::TraverseArray(const FRigVMArrayExprAST* InExpr, FRigVMCompilerWorkData& WorkData)
 {
 	URigVMArrayNode* ArrayNode = Cast<URigVMArrayNode>(InExpr->GetNode());
-	check(ArrayNode);
+	if(!ValidateNode(ArrayNode))
+	{
+		return;
+	}
 	
 	if (WorkData.bSetupMemory)
 	{
@@ -2775,6 +2801,20 @@ FRigVMOperand URigVMCompiler::FindOrAddRegister(const FRigVMVarExprAST* InVarExp
 	return Operand;
 }
 
+bool URigVMCompiler::ValidateNode(URigVMNode* InNode)
+{
+	check(InNode);
+
+	if(InNode->HasUnknownTypePin())
+	{
+		static const FString UnknownTypeMessage = TEXT("Node @@ has unresolved pins of unknown type.");
+		Settings.Report(EMessageSeverity::Error, InNode, UnknownTypeMessage);
+		return false;
+	}
+
+	return true;
+}
+
 void URigVMCompiler::ReportInfo(const FString& InMessage)
 {
 	if (Settings.SurpressInfoMessages)
diff --git a/Engine/Source/Developer/RigVMDeveloper/Private/RigVMModel/RigVMController.cpp b/Engine/Source/Developer/RigVMDeveloper/Private/RigVMModel/RigVMController.cpp
--- a/Engine/Source/Developer/RigVMDeveloper/Private/RigVMModel/RigVMController.cpp
+++ b/Engine/Source/Developer/RigVMDeveloper/Private/RigVMModel/RigVMController.cpp
@@ -7,6 +7,7 @@
 #include "RigVMModel/Nodes/RigVMFunctionReferenceNode.h"
 #include "RigVMCore/RigVMRegistry.h"
 #include "RigVMCore/RigVMExecuteContext.h"
+#include "RigVMCore/RigVMUnknownType.h"
 #include "RigVMCore/RigVMByteCode.h"
 #include "RigVMCompiler/RigVMCompiler.h"
 #include "RigVMDeveloperModule.h"
@@ -6675,6 +6676,23 @@ bool URigVMController::AddLink(URigVMPin* OutputPin, URigVMPin* InputPin, bool b
 		ActionStack->BeginAction(Action);
 	}
 
+	// resolve types on the pins if needed
+	if(InputPin->GetCPPTypeObject() != OutputPin->GetCPPTypeObject())
+	{
+		if(InputPin->GetCPPTypeObject() == FRigVMUnknownType::StaticStruct())
+		{
+			Notify(ERigVMGraphNotifType::InteractionBracketOpened, nullptr);
+			ResolveUnknownTypePin(InputPin, OutputPin, bSetupUndoRedo);
+			Notify(ERigVMGraphNotifType::InteractionBracketClosed, nullptr);
+		}
+		else if(OutputPin->GetCPPTypeObject() == FRigVMUnknownType::StaticStruct())
+		{
+			Notify(ERigVMGraphNotifType::InteractionBracketOpened, nullptr);
+			ResolveUnknownTypePin(OutputPin, InputPin, bSetupUndoRedo);
+			Notify(ERigVMGraphNotifType::InteractionBracketClosed, nullptr);
+		}
+	}
+
 	if (OutputPin->IsExecuteContext())
 	{
 		BreakAllLinks(OutputPin, false, bSetupUndoRedo);
@@ -6711,9 +6729,9 @@ bool URigVMController::AddLink(URigVMPin* OutputPin, URigVMPin* InputPin, bool b
 	UpdateRerouteNodeAfterChangingLinks(OutputPin, bSetupUndoRedo);
 	UpdateRerouteNodeAfterChangingLinks(InputPin, bSetupUndoRedo);
 
-	TArray<URigVMNode*> NodesVisited;
-	PotentiallyResolvePrototypeNode(Cast<URigVMPrototypeNode>(InputPin->GetNode()), bSetupUndoRedo, NodesVisited);
-	PotentiallyResolvePrototypeNode(Cast<URigVMPrototypeNode>(OutputPin->GetNode()), bSetupUndoRedo, NodesVisited);
+	//TArray<URigVMNode*> NodesVisited;
+	//PotentiallyResolvePrototypeNode(Cast<URigVMPrototypeNode>(InputPin->GetNode()), bSetupUndoRedo, NodesVisited);
+	//PotentiallyResolvePrototypeNode(Cast<URigVMPrototypeNode>(OutputPin->GetNode()), bSetupUndoRedo, NodesVisited);
 
 	if (bSetupUndoRedo)
 	{
@@ -9642,6 +9660,11 @@ bool URigVMController::CanAddFunctionRefForDefinition(URigVMLibraryNode* InFunct
 
 void URigVMController::AddPinsForStruct(UStruct* InStruct, URigVMNode* InNode, URigVMPin* InParentPin, ERigVMPinDirection InPinDirection, const FString& InDefaultValue, bool bAutoExpandArrays, bool bNotify)
 {
+	if(!ShouldStructBeUnfolded(InStruct))
+	{
+		return;
+	}
+	
 	TArray<FString> MemberNameValuePairs = URigVMPin::SplitDefaultValue(InDefaultValue);
 	TMap<FName, FString> MemberValues;
 	for (const FString& MemberNameValuePair : MemberNameValuePairs)
@@ -9903,6 +9926,10 @@ bool URigVMController::ShouldStructBeUnfolded(const UStruct* Struct)
 	{
 		return false;
 	}
+	if(Struct->IsChildOf(FRigVMUnknownType::StaticStruct()))
+	{
+		return false;
+	}
 	if (UnfoldStructDelegate.IsBound())
 	{
 		if (!UnfoldStructDelegate.Execute(Struct))
@@ -11085,6 +11112,7 @@ void URigVMController::PostProcessDefaultValue(URigVMPin* Pin, FString& OutDefau
 	}
 }
 
+/*
 void URigVMController::PotentiallyResolvePrototypeNode(URigVMPrototypeNode* InNode, bool bSetupUndoRedo)
 {
 	TArray<URigVMNode*> NodesVisited;
@@ -11199,8 +11227,79 @@ void URigVMController::PotentiallyResolvePrototypeNode(URigVMPrototypeNode* InNo
 		PotentiallyResolvePrototypeNode(Cast<URigVMPrototypeNode>(LinkedNode), bSetupUndoRedo, NodesVisited);
 	}
 }
+*/
 
-bool URigVMController::ChangePinType(const FString& InPinPath, const FString& InCPPType, const FName& InCPPTypeObjectPath, bool bSetupUndoRedo)
+void URigVMController::ResolveUnknownTypePin(URigVMPin* InPinToResolve, const URigVMPin* InTemplatePin,
+	bool bSetupUndoRedo, bool bTraverseNode, bool bTraverseParentPins, bool bTraverseLinks)
+{
+	check(InPinToResolve);
+	check(InTemplatePin);
+	check(InPinToResolve->IsUnknownType());
+	check(!InTemplatePin->IsUnknownType());
+
+	if(bTraverseParentPins)
+	{
+		URigVMPin* RootPin = InPinToResolve->GetRootPin();
+		if(RootPin != InPinToResolve && RootPin->IsUnknownType())
+		{
+			ResolveUnknownTypePin(RootPin, InTemplatePin, bSetupUndoRedo, bTraverseNode, false, bTraverseLinks);
+			return;
+		}
+	}
+
+	if(InPinToResolve->IsArray())
+	{
+		for(URigVMPin* SubPin : InPinToResolve->GetSubPins())
+		{
+			if(SubPin->IsUnknownType())
+			{
+				ResolveUnknownTypePin(SubPin, InTemplatePin, bSetupUndoRedo, false, false, bTraverseLinks);
+			}
+		}
+	}
+
+	FString ResolvedCPPType = InTemplatePin->IsArray() ? InTemplatePin->GetArrayElementCppType() : InTemplatePin->GetCPPType();
+	if(InPinToResolve->IsArray())
+	{
+		ResolvedCPPType = FString::Printf(TEXT("TArray<%s>"), *ResolvedCPPType);
+	}
+
+	if(!ChangePinType(InPinToResolve, ResolvedCPPType, InTemplatePin->GetCPPTypeObject(), bSetupUndoRedo, false, false, false))
+	{
+		return;
+	}
+
+	if(bTraverseNode)
+	{
+		if(const URigVMNode* Node = InPinToResolve->GetNode())
+		{
+			for(URigVMPin* Pin : Node->GetPins())
+			{
+				if(Pin->IsUnknownType())
+				{
+					ResolveUnknownTypePin(Pin, InTemplatePin, bSetupUndoRedo, false, bTraverseParentPins, bTraverseLinks);
+				}
+			}
+		}
+	}
+
+	if(bTraverseLinks)
+	{
+		// we don't need to recurse since unknown type pins don't have sub pins
+		TArray<URigVMPin*> LinkedPins = InPinToResolve->GetLinkedSourcePins(false); 
+		LinkedPins.Append(InPinToResolve->GetLinkedTargetPins(false));
+
+		for(URigVMPin* LinkedPin : LinkedPins)
+		{
+			if(LinkedPin->IsUnknownType())
+			{
+				ResolveUnknownTypePin(LinkedPin, InTemplatePin, bSetupUndoRedo);
+			}
+		}
+	}
+}
+
+bool URigVMController::ChangePinType(const FString& InPinPath, const FString& InCPPType, const FName& InCPPTypeObjectPath, bool bSetupUndoRedo, bool bSetupOrphanPins, bool bBreakLinks, bool bRemoveSubPins)
 {
 	if (!IsValidGraph())
 	{
@@ -11212,13 +11311,13 @@ bool URigVMController::ChangePinType(const FString& InPinPath, const FString& In
 
 	if (URigVMPin* Pin = Graph->FindPin(InPinPath))
 	{
-		return ChangePinType(Pin, InCPPType, InCPPTypeObjectPath, bSetupUndoRedo);
+		return ChangePinType(Pin, InCPPType, InCPPTypeObjectPath, bSetupUndoRedo, bSetupOrphanPins, bBreakLinks, bRemoveSubPins);
 	}
 
 	return false;
 }
 
-bool URigVMController::ChangePinType(URigVMPin* InPin, const FString& InCPPType, const FName& InCPPTypeObjectPath, bool bSetupUndoRedo, bool bSetupOrphanPins)
+bool URigVMController::ChangePinType(URigVMPin* InPin, const FString& InCPPType, const FName& InCPPTypeObjectPath, bool bSetupUndoRedo, bool bSetupOrphanPins, bool bBreakLinks, bool bRemoveSubPins)
 {
 	if (InCPPType == TEXT("None") || InCPPType.IsEmpty())
 	{
@@ -11231,19 +11330,25 @@ bool URigVMController::ChangePinType(URigVMPin* InPin, const FString& InCPPType,
 		return false;
 	}
 
-	if (FRigVMPropertyDescription::RequiresCPPTypeObject(InCPPType) && !CPPTypeObject)
+	return ChangePinType(InPin, InCPPType, CPPTypeObject, bSetupUndoRedo, bSetupOrphanPins, bBreakLinks, bRemoveSubPins);
+}
+
+bool URigVMController::ChangePinType(URigVMPin* InPin, const FString& InCPPType, UObject* InCPPTypeObject, bool bSetupUndoRedo, bool bSetupOrphanPins, bool bBreakLinks, bool bRemoveSubPins)
+{
+	if (InCPPType == TEXT("None") || InCPPType.IsEmpty())
 	{
 		return false;
 	}
 
-	if (CPPTypeObject)
+	FName CPPTypeObjectPath(NAME_None);
+	if(InCPPTypeObject)
 	{
-		// for now we don't allow UObjects
-		if (!CPPTypeObject->IsA<UEnum>() &&
-			!CPPTypeObject->IsA<UStruct>())
-		{
-			return false;
-		}
+		CPPTypeObjectPath = *InCPPTypeObject->GetPathName();
+	}
+
+	if (FRigVMPropertyDescription::RequiresCPPTypeObject(InCPPType) && !InCPPTypeObject)
+	{
+		return false;
 	}
 
 	FRigVMBaseAction Action;
@@ -11257,14 +11362,14 @@ bool URigVMController::ChangePinType(URigVMPin* InPin, const FString& InCPPType,
 
 	if (bSetupUndoRedo)
 	{
-		if(!bSetupOrphanPins)
+		if(!bSetupOrphanPins && bBreakLinks)
 		{
 			BreakAllLinks(InPin, true, true);
 			BreakAllLinks(InPin, false, true);
 			BreakAllLinksRecursive(InPin, true, false, true);
 			BreakAllLinksRecursive(InPin, false, false, true);
 		}
-		ActionStack->AddAction(FRigVMChangePinTypeAction(InPin, InCPPType, InCPPTypeObjectPath));
+		ActionStack->AddAction(FRigVMChangePinTypeAction(InPin, InCPPType, CPPTypeObjectPath, bSetupOrphanPins, bBreakLinks, bRemoveSubPins));
 	}
 	
 	if(bSetupOrphanPins)
@@ -11289,17 +11394,20 @@ bool URigVMController::ChangePinType(URigVMPin* InPin, const FString& InCPPType,
 		}
 	}
 
-	TArray<URigVMPin*> Pins = InPin->SubPins;
-	for (URigVMPin* Pin : Pins)
+	if(bRemoveSubPins || !InPin->IsArray())
 	{
-		RemovePin(Pin, false, true);
+		TArray<URigVMPin*> Pins = InPin->SubPins;
+		for (URigVMPin* Pin : Pins)
+		{
+			RemovePin(Pin, false, true);
+		}
+		
+		InPin->SubPins.Reset();
 	}
 	
-	InPin->SubPins.Reset();
-	
 	InPin->CPPType = InCPPType;
-	InPin->CPPTypeObjectPath = InCPPTypeObjectPath;
-	InPin->CPPTypeObject = CPPTypeObject;
+	InPin->CPPTypeObjectPath = CPPTypeObjectPath;
+	InPin->CPPTypeObject = InCPPTypeObject;
 	InPin->DefaultValue = FString();
 
 	if (InPin->IsStruct())
@@ -11312,6 +11420,18 @@ bool URigVMController::ChangePinType(URigVMPin* InPin, const FString& InCPPType,
 	Notify(ERigVMGraphNotifType::PinTypeChanged, InPin);
 	Notify(ERigVMGraphNotifType::PinDefaultValueChanged, InPin);
 
+	// in cases were we are just changing the type we have to let the
+	// clients know that the links are still there
+	if(!bSetupOrphanPins && !bBreakLinks && !bRemoveSubPins)
+	{
+		const TArray<URigVMLink*> CurrentLinks = InPin->GetLinks();
+		for(URigVMLink* CurrentLink : CurrentLinks)
+		{
+			Notify(ERigVMGraphNotifType::LinkRemoved, CurrentLink);
+			Notify(ERigVMGraphNotifType::LinkAdded, CurrentLink);
+		}
+	}
+
 	if (bSetupUndoRedo)
 	{
 		ActionStack->EndAction(Action);
diff --git a/Engine/Source/Developer/RigVMDeveloper/Private/RigVMModel/RigVMControllerActions.cpp b/Engine/Source/Developer/RigVMDeveloper/Private/RigVMModel/RigVMControllerActions.cpp
--- a/Engine/Source/Developer/RigVMDeveloper/Private/RigVMModel/RigVMControllerActions.cpp
+++ b/Engine/Source/Developer/RigVMDeveloper/Private/RigVMModel/RigVMControllerActions.cpp
@@ -1357,12 +1357,26 @@ bool FRigVMBreakLinkAction::Redo(URigVMController* InController)
 	return FRigVMBaseAction::Redo(InController);
 }
 
-FRigVMChangePinTypeAction::FRigVMChangePinTypeAction(URigVMPin* InPin, const FString& InCppType, const FName& InCppTypeObjectPath)
+FRigVMChangePinTypeAction::FRigVMChangePinTypeAction()
+: PinPath()
+, OldCPPType()
+, OldCPPTypeObjectPath(NAME_None)
+, NewCPPType()
+, NewCPPTypeObjectPath(NAME_None)
+, bSetupOrphanPins(true)
+, bBreakLinks(true)
+, bRemoveSubPins(true)
+{}
+
+FRigVMChangePinTypeAction::FRigVMChangePinTypeAction(URigVMPin* InPin, const FString& InCppType, const FName& InCppTypeObjectPath, bool InSetupOrphanPins, bool InBreakLinks, bool InRemoveSubPins)
 : PinPath(InPin->GetPinPath())
 , OldCPPType(InPin->GetCPPType())
 , OldCPPTypeObjectPath(NAME_None)
 , NewCPPType(InCppType)
 , NewCPPTypeObjectPath(InCppTypeObjectPath)
+, bSetupOrphanPins(InSetupOrphanPins)
+, bBreakLinks(InBreakLinks)
+, bRemoveSubPins(InRemoveSubPins)
 {
 	if (UObject* CPPTypeObject = InPin->GetCPPTypeObject())
 	{
@@ -1376,12 +1390,12 @@ bool FRigVMChangePinTypeAction::Undo(URigVMController* InController)
 	{
 		return false;
 	}
-	return InController->ChangePinType(PinPath, OldCPPType, OldCPPTypeObjectPath, false);
+	return InController->ChangePinType(PinPath, OldCPPType, OldCPPTypeObjectPath, false, bSetupOrphanPins, bBreakLinks, bRemoveSubPins);
 }
 
 bool FRigVMChangePinTypeAction::Redo(URigVMController* InController)
 {
-	if(!InController->ChangePinType(PinPath, NewCPPType, NewCPPTypeObjectPath, false))
+	if(!InController->ChangePinType(PinPath, NewCPPType, NewCPPTypeObjectPath, false, bSetupOrphanPins, bBreakLinks, bRemoveSubPins))
 	{
 		return false;
 	}
diff --git a/Engine/Source/Developer/RigVMDeveloper/Private/RigVMModel/RigVMNode.cpp b/Engine/Source/Developer/RigVMDeveloper/Private/RigVMModel/RigVMNode.cpp
--- a/Engine/Source/Developer/RigVMDeveloper/Private/RigVMModel/RigVMNode.cpp
+++ b/Engine/Source/Developer/RigVMDeveloper/Private/RigVMModel/RigVMNode.cpp
@@ -292,6 +292,18 @@ bool URigVMNode::IsMutable() const
 	return false;
 }
 
+bool URigVMNode::HasUnknownTypePin() const
+{
+	for (const URigVMPin* Pin : GetPins())
+	{
+		if (Pin->IsUnknownType())
+		{
+			return true;
+		}
+	}
+	return false;
+}
+
 bool URigVMNode::IsEvent() const
 {
 	return IsMutable() && !HasInputPin(true /* include io */) && !GetEventName().IsNone();
diff --git a/Engine/Source/Developer/RigVMDeveloper/Private/RigVMModel/RigVMPin.cpp b/Engine/Source/Developer/RigVMDeveloper/Private/RigVMModel/RigVMPin.cpp
--- a/Engine/Source/Developer/RigVMDeveloper/Private/RigVMModel/RigVMPin.cpp
+++ b/Engine/Source/Developer/RigVMDeveloper/Private/RigVMModel/RigVMPin.cpp
@@ -7,6 +7,7 @@
 #include "RigVMModel/Nodes/RigVMPrototypeNode.h"
 #include "RigVMCompiler/RigVMCompiler.h"
 #include "RigVMCore/RigVMExecuteContext.h"
+#include "RigVMCore/RigVMUnknownType.h"
 #include "UObject/Package.h"
 #include "Misc/DefaultValueHelper.h"
 #include "Misc/PackageName.h"
@@ -527,7 +528,7 @@ bool URigVMPin::IsStringType() const
 
 bool URigVMPin::IsExecuteContext() const
 {
-	if (UScriptStruct* ScriptStruct = GetScriptStruct())
+	if (const UScriptStruct* ScriptStruct = GetScriptStruct())
 	{
 		if (ScriptStruct->IsChildOf(FRigVMExecuteContext::StaticStruct()))
 		{
@@ -537,6 +538,18 @@ bool URigVMPin::IsExecuteContext() const
 	return false;
 }
 
+bool URigVMPin::IsUnknownType() const
+{
+	if (const UScriptStruct* ScriptStruct = GetScriptStruct())
+	{
+		if (ScriptStruct->IsChildOf(FRigVMUnknownType::StaticStruct()))
+		{
+			return true;
+		}
+	}
+	return false;
+}
+
 FString URigVMPin::GetDefaultValue() const
 {
 	return GetDefaultValue(EmptyPinOverride);
@@ -1320,6 +1333,22 @@ bool URigVMPin::CanLink(URigVMPin* InSourcePin, URigVMPin* InTargetPin, FString*
 			bCPPTypesDiffer = false;
 		}
 
+		if (bCPPTypesDiffer)
+		{
+			if(InSourcePin->IsArray() == InTargetPin->IsArray())
+			{
+				if(InSourcePin->IsUnknownType() && !InTargetPin->IsExecuteContext())
+				{
+					
+					return true;
+				}
+				else if(InTargetPin->IsUnknownType() && !InSourcePin->IsExecuteContext())
+				{
+					return true;
+				}
+			}
+		}
+
 		if (bCPPTypesDiffer)
 		{
 			if (OutFailureReason)
diff --git a/Engine/Source/Developer/RigVMDeveloper/Public/RigVMCompiler/RigVMCompiler.h b/Engine/Source/Developer/RigVMDeveloper/Public/RigVMCompiler/RigVMCompiler.h
--- a/Engine/Source/Developer/RigVMDeveloper/Public/RigVMCompiler/RigVMCompiler.h
+++ b/Engine/Source/Developer/RigVMDeveloper/Public/RigVMCompiler/RigVMCompiler.h
@@ -174,6 +174,8 @@ private:
 
 	FRigVMOperand FindOrAddRegister(const FRigVMVarExprAST* InVarExpr, FRigVMCompilerWorkData& WorkData, bool bIsDebugValue = false);
 
+	bool ValidateNode(URigVMNode* InNode);
+	
 	void ReportInfo(const FString& InMessage);
 	void ReportWarning(const FString& InMessage);
 	void ReportError(const FString& InMessage);
diff --git a/Engine/Source/Developer/RigVMDeveloper/Public/RigVMModel/RigVMController.h b/Engine/Source/Developer/RigVMDeveloper/Public/RigVMModel/RigVMController.h
--- a/Engine/Source/Developer/RigVMDeveloper/Public/RigVMModel/RigVMController.h
+++ b/Engine/Source/Developer/RigVMDeveloper/Public/RigVMModel/RigVMController.h
@@ -831,10 +831,14 @@ private:
 	static void CreateDefaultValueForStructIfRequired(UScriptStruct* InStruct, FString& InOutDefaultValue);
 	static void PostProcessDefaultValue(URigVMPin* Pin, FString& OutDefaultValue);
 
+	/*
 	void PotentiallyResolvePrototypeNode(URigVMPrototypeNode* InNode, bool bSetupUndoRedo);
 	void PotentiallyResolvePrototypeNode(URigVMPrototypeNode* InNode, bool bSetupUndoRedo, TArray<URigVMNode*>& NodesVisited);
-	bool ChangePinType(const FString& InPinPath, const FString& InCPPType, const FName& InCPPTypeObjectPath, bool bSetupUndoRedo);
-	bool ChangePinType(URigVMPin* InPin, const FString& InCPPType, const FName& InCPPTypeObjectPath, bool bSetupUndoRedo, bool bSetupOrphanPins = true);
+	*/
+	void ResolveUnknownTypePin(URigVMPin* InPinToResolve, const URigVMPin* InTemplatePin, bool bSetupUndoRedo, bool bTraverseNode = true, bool bTraverseParentPins = true, bool bTraverseLinks = true);
+	bool ChangePinType(const FString& InPinPath, const FString& InCPPType, const FName& InCPPTypeObjectPath, bool bSetupUndoRedo, bool bSetupOrphanPins = true, bool bBreakLinks = true, bool bRemoveSubPins = true);
+	bool ChangePinType(URigVMPin* InPin, const FString& InCPPType, const FName& InCPPTypeObjectPath, bool bSetupUndoRedo, bool bSetupOrphanPins = true, bool bBreakLinks = true, bool bRemoveSubPins = true);
+	bool ChangePinType(URigVMPin* InPin, const FString& InCPPType,UObject* InCPPTypeObject, bool bSetupUndoRedo, bool bSetupOrphanPins = true, bool bBreakLinks = true, bool bRemoveSubPins = true);
 
 #if WITH_EDITOR
 	void RewireLinks(URigVMPin* OldPin, URigVMPin* NewPin, bool bAsInput, bool bSetupUndoRedo, TArray<URigVMLink*> InLinks = TArray<URigVMLink*>());
diff --git a/Engine/Source/Developer/RigVMDeveloper/Public/RigVMModel/RigVMControllerActions.h b/Engine/Source/Developer/RigVMDeveloper/Public/RigVMModel/RigVMControllerActions.h
--- a/Engine/Source/Developer/RigVMDeveloper/Public/RigVMModel/RigVMControllerActions.h
+++ b/Engine/Source/Developer/RigVMDeveloper/Public/RigVMModel/RigVMControllerActions.h
@@ -1170,8 +1170,8 @@ struct FRigVMChangePinTypeAction : public FRigVMBaseAction
 
 public:
 
-	FRigVMChangePinTypeAction() {}
-	FRigVMChangePinTypeAction(URigVMPin* InPin, const FString& InCppType, const FName& InCppTypeObjectPath);
+	FRigVMChangePinTypeAction();
+	FRigVMChangePinTypeAction(URigVMPin* InPin, const FString& InCppType, const FName& InCppTypeObjectPath, bool InSetupOrphanPins, bool InBreakLinks, bool InRemoveSubPins);
 	virtual ~FRigVMChangePinTypeAction() {};
 	virtual bool Undo(URigVMController* InController) override;
 	virtual bool Redo(URigVMController* InController) override;
@@ -1190,6 +1190,15 @@ public:
 
 	UPROPERTY()
 	FName NewCPPTypeObjectPath;
+
+	UPROPERTY()
+	bool bSetupOrphanPins;
+
+	UPROPERTY()
+	bool bBreakLinks;
+
+	UPROPERTY()
+	bool bRemoveSubPins;
 };
 
 /**
diff --git a/Engine/Source/Developer/RigVMDeveloper/Public/RigVMModel/RigVMNode.h b/Engine/Source/Developer/RigVMDeveloper/Public/RigVMModel/RigVMNode.h
--- a/Engine/Source/Developer/RigVMDeveloper/Public/RigVMModel/RigVMNode.h
+++ b/Engine/Source/Developer/RigVMDeveloper/Public/RigVMModel/RigVMNode.h
@@ -139,6 +139,9 @@ public:
 	UFUNCTION(BlueprintCallable, Category = RigVMNode)
 	virtual bool IsMutable() const;
 
+	// Returns true if this node has an unknown type pin
+	bool HasUnknownTypePin() const;
+
 	virtual bool ContributesToResult() const { return IsMutable(); }
 
 	// Returns true if this Node is the beginning of a scope
diff --git a/Engine/Source/Developer/RigVMDeveloper/Public/RigVMModel/RigVMPin.h b/Engine/Source/Developer/RigVMDeveloper/Public/RigVMModel/RigVMPin.h
--- a/Engine/Source/Developer/RigVMDeveloper/Public/RigVMModel/RigVMPin.h
+++ b/Engine/Source/Developer/RigVMDeveloper/Public/RigVMModel/RigVMPin.h
@@ -206,6 +206,10 @@ public:
 	UFUNCTION(BlueprintCallable, Category = RigVMPin)
 	bool IsExecuteContext() const;
 
+	// Returns true if the C++ data type is unknown
+	UFUNCTION(BlueprintCallable, Category = RigVMPin)
+	bool IsUnknownType() const;
+
 	// Returns the default value of the Pin as a string.
 	// Note that this value is computed based on the Pin's
 	// SubPins - so for example for a FVector typed Pin
diff --git a/Engine/Source/Runtime/RigVM/Private/RigVMCore/RigVMUnknownType.cpp b/Engine/Source/Runtime/RigVM/Private/RigVMCore/RigVMUnknownType.cpp
new file mode 100644
--- /dev/null
+++ b/Engine/Source/Runtime/RigVM/Private/RigVMCore/RigVMUnknownType.cpp
@@ -0,0 +1,3 @@
+// Copyright Epic Games, Inc. All Rights Reserved.
+
+#include "RigVMCore/RigVMUnknownType.h"
diff --git a/Engine/Source/Runtime/RigVM/Public/RigVMCore/RigVMUnknownType.h b/Engine/Source/Runtime/RigVM/Public/RigVMCore/RigVMUnknownType.h
new file mode 100644
--- /dev/null
+++ b/Engine/Source/Runtime/RigVM/Public/RigVMCore/RigVMUnknownType.h
@@ -0,0 +1,25 @@
+// Copyright Epic Games, Inc. All Rights Reserved.
+
+#pragma once
+
+#include "CoreMinimal.h"
+#include "RigVMDefines.h"
+#include "RigVMUnknownType.generated.h"
+
+/**
+ * The unknown type is used to identify untyped nodes
+ */
+USTRUCT()
+struct RIGVM_API FRigVMUnknownType
+{
+	GENERATED_BODY()
+
+	FRigVMUnknownType()
+		: Hash(0)
+	{}
+
+private:
+	
+	UPROPERTY()
+	uint32 Hash;
+};
