Control Rig: Reroute node title based on value type

diff --git a/Engine/Source/Developer/RigVMDeveloper/Private/RigVMModel/Nodes/RigVMRerouteNode.cpp b/Engine/Source/Developer/RigVMDeveloper/Private/RigVMModel/Nodes/RigVMRerouteNode.cpp
--- a/Engine/Source/Developer/RigVMDeveloper/Private/RigVMModel/Nodes/RigVMRerouteNode.cpp
+++ b/Engine/Source/Developer/RigVMDeveloper/Private/RigVMModel/Nodes/RigVMRerouteNode.cpp
@@ -10,6 +10,49 @@ URigVMRerouteNode::URigVMRerouteNode()
 {
 }
 
+FString URigVMRerouteNode::GetNodeTitle() const
+{
+	if(const URigVMPin* ValuePin = FindPin(ValueName))
+	{
+		FString TypeDisplayName;
+		if(const UScriptStruct* ScriptStruct = Cast<UScriptStruct>(ValuePin->GetCPPTypeObject()))
+		{
+			TypeDisplayName = ScriptStruct->GetDisplayNameText().ToString();
+		}
+		else if(const UEnum* Enum = Cast<UEnum>(ValuePin->GetCPPTypeObject()))
+		{
+			TypeDisplayName = Enum->GetName();
+		}
+		else if(const UClass* Class = Cast<UClass>(ValuePin->GetCPPTypeObject()))
+		{
+			TypeDisplayName = Class->GetDisplayNameText().ToString();
+		}
+		else if(ValuePin->IsArray())
+		{
+			TypeDisplayName = ValuePin->GetArrayElementCppType();
+		}
+		else
+		{
+			TypeDisplayName = ValuePin->GetCPPType();
+		}
+
+		if(TypeDisplayName.IsEmpty())
+		{
+			return RerouteName;
+		}
+
+		TypeDisplayName = TypeDisplayName.Left(1).ToUpper() + TypeDisplayName.Mid(1);
+
+		if(ValuePin->IsArray())
+		{
+			TypeDisplayName += TEXT(" Array");
+		}
+
+		return TypeDisplayName;
+	}
+	return RerouteName;
+}
+
 bool URigVMRerouteNode::GetShowsAsFullNode() const
 {
 	return bShowAsFullNode;
diff --git a/Engine/Source/Developer/RigVMDeveloper/Public/RigVMModel/Nodes/RigVMRerouteNode.h b/Engine/Source/Developer/RigVMDeveloper/Public/RigVMModel/Nodes/RigVMRerouteNode.h
--- a/Engine/Source/Developer/RigVMDeveloper/Public/RigVMModel/Nodes/RigVMRerouteNode.h
+++ b/Engine/Source/Developer/RigVMDeveloper/Public/RigVMModel/Nodes/RigVMRerouteNode.h
@@ -23,7 +23,7 @@ public:
 	URigVMRerouteNode();
 
 	// Override of node title
-	virtual FString GetNodeTitle() const override { return RerouteName; }
+	virtual FString GetNodeTitle() const override;
 
 	// Returns true if this node should be shown as a full node,
 	// or false if this node should be shown as a small circle.
