Control Rig: Add string sanitization to all processes inside of the controller

diff --git a/Engine/Source/Developer/RigVMDeveloper/Private/RigVMModel/RigVMController.cpp b/Engine/Source/Developer/RigVMDeveloper/Private/RigVMModel/RigVMController.cpp
--- a/Engine/Source/Developer/RigVMDeveloper/Private/RigVMModel/RigVMController.cpp
+++ b/Engine/Source/Developer/RigVMDeveloper/Private/RigVMModel/RigVMController.cpp
@@ -273,15 +273,19 @@ TArray<FString> URigVMController::GeneratePythonCommands()
 {
 	TArray<FString> Commands;
 
+	const FString GraphName = GetSanitizedGraphName(GetGraph()->GetGraphName());
+
 	// Add local variables
 	for (const FRigVMGraphVariableDescription& Variable : GetGraph()->LocalVariables)
 	{
+		const FString VariableName = GetSanitizedVariableName(Variable.Name.ToString());
+
 		if (Variable.CPPTypeObject)
 		{
 			// FRigVMGraphVariableDescription AddLocalVariable(const FName& InVariableName, const FString& InCPPType, UObject* InCPPTypeObject, const FString& InDefaultValue, bool bSetupUndoRedo = true);
 			Commands.Add(FString::Printf(TEXT("blueprint.get_controller_by_name('%s').add_local_variable_from_object_path('%s', '%s', '%s', '%s')"),
-						*GetGraph()->GetGraphName(),
-						*Variable.Name.ToString(),
+						*GraphName,
+						*VariableName,
 						*Variable.CPPType,
 						Variable.CPPTypeObject ? *Variable.CPPTypeObject->GetPathName() : TEXT(""),
 						*Variable.DefaultValue));
@@ -290,8 +294,8 @@ TArray<FString> URigVMController::GeneratePythonCommands()
 		{
 			// FRigVMGraphVariableDescription AddLocalVariable(const FName& InVariableName, const FString& InCPPType, UObject* InCPPTypeObject, const FString& InDefaultValue, bool bSetupUndoRedo = true);
 			Commands.Add(FString::Printf(TEXT("blueprint.get_controller_by_name('%s').add_local_variable('%s', '%s', None, '%s')"),
-						*GetGraph()->GetGraphName(),
-						*Variable.Name.ToString(),
+						*GraphName,
+						*VariableName,
 						*Variable.CPPType,
 						*Variable.DefaultValue));
 		}
@@ -343,11 +347,14 @@ TArray<FString> URigVMController::GeneratePythonCommands()
 			continue;
 		}
 
+		const FString SourcePinPath = GetSanitizedPinPath(SourcePin->GetPinPath());
+		const FString TargetPinPath = GetSanitizedPinPath(TargetPin->GetPinPath());
+
 		//bool AddLink(const FString& InOutputPinPath, const FString& InInputPinPath, bool bSetupUndoRedo = true);
 		Commands.Add(FString::Printf(TEXT("blueprint.get_controller_by_name('%s').add_link('%s', '%s')"),
-					*GetGraph()->GetGraphName(),
-					*SourcePin->GetPinPath(),
-					*TargetPin->GetPinPath()));
+					*GraphName,
+					*SourcePinPath,
+					*TargetPinPath));
 	}
 
 	// Reroutes
@@ -371,7 +378,9 @@ TArray<FString> URigVMController::GeneratePythonCommands()
 			bNodesAdded = false;
 			for (URigVMRerouteNode* Reroute : Reroutes)
 			{
-				if (ReroutesAdded.Contains(Reroute->GetName()))
+				const FString RerouteName = GetSanitizedNodeName(Reroute->GetName());
+
+				if (ReroutesAdded.Contains(RerouteName))
 				{
 					continue;
 				}
@@ -396,15 +405,18 @@ TArray<FString> URigVMController::GeneratePythonCommands()
 				if (!TargetPins.IsEmpty())
 				{
 					URigVMPin* TargetPin = TargetPins[0];
+					
+					const FString TargetPinPath = GetSanitizedPinPath(TargetPin->GetPinPath());
+
 					// AddRerouteNodeOnPin(const FString& InPinPath, bool bAsInput, bool bShowAsFullNode, const FVector2D& InPosition, const FString& InNodeName, bool bSetupUndoRedo, bool bPrintPythonCommand)
 					Commands.Add(FString::Printf(TEXT("blueprint.get_controller_by_name('%s').add_reroute_node_on_pin('%s', True, %s, %s, '%s')"),
-							*GetGraph()->GetGraphName(),
-							*TargetPin->GetPinPath(),
+							*GraphName,
+							*TargetPinPath,
 							Reroute->GetShowsAsFullNode() ? TEXT("True") : TEXT("False"),
 							*RigVMPythonUtils::Vector2DToPythonString(Reroute->GetPosition()),
-							*Reroute->GetName()));
+							*RerouteName));
 
-					ReroutesAdded.Add(Reroute->GetName());
+					ReroutesAdded.Add(RerouteName);
 					bNodesAdded = true;
 
 					// Add the rest of target links
@@ -422,11 +434,14 @@ TArray<FString> URigVMController::GeneratePythonCommands()
 							continue;
 						}
 
+						const FString OtherTargetPinPath = GetSanitizedPinPath(OtherTargetPin->GetPinPath());
+						const FString FirstReroutePinPath = GetSanitizedPinPath(Reroute->Pins[0]->GetPinPath());
+
 						//bool AddLink(const FString& InOutputPinPath, const FString& InInputPinPath, bool bSetupUndoRedo = true);
 						Commands.Add(FString::Printf(TEXT("blueprint.get_controller_by_name('%s').add_link('%s', '%s')"),
-									*GetGraph()->GetGraphName(),
-									*Reroute->Pins[0]->GetPinPath(),
-									*OtherTargetPin->GetPinPath()));
+									*GraphName,
+									*FirstReroutePinPath,
+									*OtherTargetPinPath));
 					}
 
 					// Set default value in input if necessary
@@ -434,9 +449,11 @@ TArray<FString> URigVMController::GeneratePythonCommands()
 					const FString DefaultValue = Pin->GetDefaultValue();
 					if (!DefaultValue.IsEmpty() && DefaultValue != TEXT("()"))
 					{
+						const FString PinPath = GetSanitizedPinPath(Pin->GetPinPath());
+
 						Commands.Add(FString::Printf(TEXT("blueprint.get_controller_by_name('%s').set_pin_default_value('%s', '%s')"),
-									*GetGraph()->GetGraphName(),
-									*Pin->GetPinPath(),
+									*GraphName,
+									*PinPath,
 									*Pin->GetDefaultValue()));
 					}
 					
@@ -450,11 +467,14 @@ TArray<FString> URigVMController::GeneratePythonCommands()
 							continue;
 						}
 
+						const FString SourcePinPath = GetSanitizedPinPath(SourcePin->GetPinPath());
+						const FString LinkTargetPinPath = GetSanitizedPinPath(Link->GetTargetPin()->GetPinPath());
+
 						//bool AddLink(const FString& InOutputPinPath, const FString& InInputPinPath, bool bSetupUndoRedo = true);
 						Commands.Add(FString::Printf(TEXT("blueprint.get_controller_by_name('%s').add_link('%s', '%s')"),
-									*GetGraph()->GetGraphName(),
-									*SourcePin->GetPinPath(),
-									*Link->GetTargetPin()->GetPinPath()));
+									*GraphName,
+									*SourcePinPath,
+									*LinkTargetPinPath));
 					}				
 				}				
 			}
@@ -467,178 +487,193 @@ TArray<FString> URigVMController::GeneratePythonCommands()
 TArray<FString> URigVMController::GetAddNodePythonCommands(URigVMNode* Node) const
 {
 	TArray<FString> Commands;
-	if (URigVMUnitNode* UnitNode = Cast<URigVMUnitNode>(Node))
+
+	const FString GraphName = GetSanitizedGraphName(GetGraph()->GetGraphName());
+	const FString NodeName = GetSanitizedNodeName(Node->GetName());
+
+	if (const URigVMUnitNode* UnitNode = Cast<URigVMUnitNode>(Node))
 	{
-		if (URigVMInjectionInfo* InjectionInfo = Cast<URigVMInjectionInfo>(UnitNode->GetOuter()))
+		if (const URigVMInjectionInfo* InjectionInfo = Cast<URigVMInjectionInfo>(UnitNode->GetOuter()))
 		{
+			const URigVMPin* InjectionInfoPin = InjectionInfo->GetPin();
+			const FString InjectionInfoPinPath = GetSanitizedPinPath(InjectionInfoPin->GetPinPath());
+			const FString InjectionInfoInputPinName = GetSanitizedPinName(InjectionInfo->InputPin->GetName());
+			const FString InjectionInfoOutputPinName = GetSanitizedPinName(InjectionInfo->OutputPin->GetName());
+
 			//URigVMInjectionInfo* AddInjectedNodeFromStructPath(const FString& InPinPath, bool bAsInput, const FString& InScriptStructPath, const FName& InMethodName, const FName& InInputPinName, const FName& InOutputPinName, const FString& InNodeName = TEXT(""), bool bSetupUndoRedo = true);
 			Commands.Add(FString::Printf(TEXT("%s_info = blueprint.get_controller_by_name('%s').add_injected_node_from_struct_path('%s', %s, '%s', '%s', '%s', '%s', '%s')"),
-					*UnitNode->GetName(), 
-					*GetGraph()->GetGraphName(), 
-					*InjectionInfo->GetPin()->GetPinPath(), 
-					InjectionInfo->GetPin()->GetDirection() == ERigVMPinDirection::Input ? TEXT("True") : TEXT("False"), 
+					*NodeName, 
+					*GraphName, 
+					*InjectionInfoPinPath, 
+					InjectionInfoPin->GetDirection() == ERigVMPinDirection::Input ? TEXT("True") : TEXT("False"), 
 					*UnitNode->GetScriptStruct()->GetPathName(), 
 					*UnitNode->GetMethodName().ToString(), 
-					*InjectionInfo->InputPin->GetName(), 
-					*InjectionInfo->OutputPin->GetName(), 
+					*InjectionInfoInputPinName, 
+					*InjectionInfoOutputPinName, 
 					*UnitNode->GetName()));
 		}
 		else
 		{
 			// add_struct_node_from_struct_path(script_struct_path, method_name, position=[0.0, 0.0], node_name='', undo=True)
 			Commands.Add(FString::Printf(TEXT("blueprint.get_controller_by_name('%s').add_unit_node_from_struct_path('%s', 'Execute', %s, '%s')"),
-					*GetGraph()->GetGraphName(),
+					*GraphName,
 					*UnitNode->GetScriptStruct()->GetPathName(),
 					*RigVMPythonUtils::Vector2DToPythonString(UnitNode->GetPosition()),
-					*UnitNode->GetName()));
+					*NodeName));
 		}
 	}
-	else if (URigVMVariableNode* VariableNode = Cast<URigVMVariableNode>(Node))
+	else if (const URigVMVariableNode* VariableNode = Cast<URigVMVariableNode>(Node))
 	{
+		const FString VariableName = GetSanitizedVariableName(VariableNode->GetVariableName().ToString());
+
 		// add_variable_node(variable_name, cpp_type, cpp_type_object, is_getter, default_value, position=[0.0, 0.0], node_name='', undo=True)
 		if (VariableNode->GetVariableDescription().CPPTypeObject)
 		{
 			Commands.Add(FString::Printf(TEXT("blueprint.get_controller_by_name('%s').add_variable_node_from_object_path('%s', '%s', '%s', %s, '%s', %s, '%s')"),
-					*GetGraph()->GetGraphName(),
-					*VariableNode->GetVariableName().ToString(),
+					*GraphName,
+					*VariableName,
 					*VariableNode->GetVariableDescription().CPPType,
 					*VariableNode->GetVariableDescription().CPPTypeObject->GetPathName(),
 					VariableNode->IsGetter() ? TEXT("True") : TEXT("False"),
 					*VariableNode->GetVariableDescription().DefaultValue,
 					*RigVMPythonUtils::Vector2DToPythonString(VariableNode->GetPosition()),
-					*VariableNode->GetName()));	
+					*NodeName));	
 		}
 		else
 		{
 			Commands.Add(FString::Printf(TEXT("blueprint.get_controller_by_name('%s').add_variable_node('%s', '%s', None, %s, '%s', %s, '%s')"),
-					*GetGraph()->GetGraphName(),
-					*VariableNode->GetVariableName().ToString(),
+					*GraphName,
+					*VariableName,
 					*VariableNode->GetVariableDescription().CPPType,
 					VariableNode->IsGetter() ? TEXT("True") : TEXT("False"),
 					*VariableNode->GetVariableDescription().DefaultValue,
 					*RigVMPythonUtils::Vector2DToPythonString(VariableNode->GetPosition()),
-					*VariableNode->GetName()));	
+					*NodeName));	
 		}
 	}
-	else if (URigVMParameterNode* ParameterNode = Cast<URigVMParameterNode>(Node))
+	else if (const URigVMParameterNode* ParameterNode = Cast<URigVMParameterNode>(Node))
 	{
+		const FString ParameterName = GetSanitizedVariableName(ParameterNode->GetParameterName().ToString());
+
 		// add_parameter_node_from_object_path(parameter_name, cpp_type, cpp_type_object_path, is_input, default_value, position=[0.0, 0.0], node_name='', undo=True)
 		if (ParameterNode->GetParameterDescription().CPPTypeObject)
 		{
 			Commands.Add(FString::Printf(TEXT("blueprint.get_controller_by_name('%s').add_parameter_node_from_object_path('%s', '%s', '%s', %s, '%s', %s, '%s')"),
-					*GetGraph()->GetGraphName(),
-					*ParameterNode->GetParameterName().ToString(),
+					*GraphName,
+					*ParameterName,
 					*ParameterNode->GetParameterDescription().CPPType,
 					*ParameterNode->GetParameterDescription().CPPTypeObject->GetPathName(),
 					ParameterNode->IsInput() ? TEXT("True") : TEXT("False"),
 					*ParameterNode->GetParameterDescription().DefaultValue,
 					*RigVMPythonUtils::Vector2DToPythonString(ParameterNode->GetPosition()),
-					*ParameterNode->GetName()));	
+					*NodeName));	
 		}
 		else
 		{
 			Commands.Add(FString::Printf(TEXT("blueprint.get_controller_by_name('%s').add_parameter_node('%s', '%s', None, %s, str(%s), %s, '%s')"),
-					*GetGraph()->GetGraphName(),
-					*ParameterNode->GetParameterName().ToString(),
+					*GraphName,
+					*ParameterName,
 					*ParameterNode->GetParameterDescription().CPPType,
 					ParameterNode->IsInput() ? TEXT("True") : TEXT("False"),
 					*ParameterNode->GetParameterDescription().DefaultValue,
 					*RigVMPythonUtils::Vector2DToPythonString(ParameterNode->GetPosition()),
-					*ParameterNode->GetName()));	
+					*NodeName));	
 		}
 	}
-	else if (URigVMCommentNode* CommentNode = Cast<URigVMCommentNode>(Node))
+	else if (const URigVMCommentNode* CommentNode = Cast<URigVMCommentNode>(Node))
 	{
 		// add_comment_node(comment_text, position=[0.0, 0.0], size=[400.0, 300.0], color=[0.0, 0.0, 0.0, 0.0], node_name='', undo=True)
 		Commands.Add(FString::Printf(TEXT("blueprint.get_controller_by_name('%s').add_comment_node('%s', %s, %s, %s, '%s')"),
-					*GetGraph()->GetGraphName(),
+					*GraphName,
 					*CommentNode->GetCommentText(),
 					*RigVMPythonUtils::Vector2DToPythonString(CommentNode->GetPosition()),
 					*RigVMPythonUtils::Vector2DToPythonString(CommentNode->GetSize()),
 					*RigVMPythonUtils::LinearColorToPythonString(CommentNode->GetNodeColor()),
-					*CommentNode->GetName()));	
+					*NodeName));	
 	}
-	else if (URigVMBranchNode* BranchNode = Cast<URigVMBranchNode>(Node))
+	else if (const URigVMBranchNode* BranchNode = Cast<URigVMBranchNode>(Node))
 	{
 		// add_branch_node(position=[0.0, 0.0], node_name='', undo=True)
 		Commands.Add(FString::Printf(TEXT("blueprint.get_controller_by_name('%s').add_branch_node(%s, '%s')"),
-						*GetGraph()->GetGraphName(),
+						*GraphName,
 						*RigVMPythonUtils::Vector2DToPythonString(BranchNode->GetPosition()),
-						*BranchNode->GetName()));	
+						*NodeName));	
 	}
-	else if (URigVMIfNode* IfNode = Cast<URigVMIfNode>(Node))
+	else if (const URigVMIfNode* IfNode = Cast<URigVMIfNode>(Node))
 	{
 		// add_if_node(cpp_type, cpp_type_object_path, position=[0.0, 0.0], node_name='', undo=True)
 		URigVMPin* ResultPin = IfNode->FindPin(URigVMIfNode::ResultName);
 		Commands.Add(FString::Printf(TEXT("blueprint.get_controller_by_name('%s').add_if_node('%s', '%s', %s, '%s')"),
-						*GetGraph()->GetGraphName(),
+						*GraphName,
 						*ResultPin->GetCPPType(),
 						*ResultPin->CPPTypeObject->GetPathName(),
 						*RigVMPythonUtils::Vector2DToPythonString(IfNode->GetPosition()),
-						*IfNode->GetName()));
+						*NodeName));
 	}
-	else if (URigVMSelectNode* SelectNode = Cast<URigVMSelectNode>(Node))
+	else if (const URigVMSelectNode* SelectNode = Cast<URigVMSelectNode>(Node))
 	{
 		// add_select_node(cpp_type, cpp_type_object_path, position=[0.0, 0.0], node_name='', undo=True)
 		URigVMPin* ResultPin = SelectNode->FindPin(URigVMSelectNode::ResultName);
 		Commands.Add(FString::Printf(TEXT("blueprint.get_controller_by_name('%s').add_select_node('%s', '%s', %s, '%s')"),
-						*GetGraph()->GetGraphName(),
+						*GraphName,
 						*ResultPin->GetCPPType(),
 						*ResultPin->CPPTypeObject->GetPathName(),
 						*RigVMPythonUtils::Vector2DToPythonString(SelectNode->GetPosition()),
-						*SelectNode->GetName()));
+						*NodeName));
 	}
-	else if (URigVMPrototypeNode* PrototypeNode = Cast<URigVMPrototypeNode>(Node))
+	else if (const URigVMPrototypeNode* PrototypeNode = Cast<URigVMPrototypeNode>(Node))
 	{
 		// add_prototype_node(notation, position=[0.0, 0.0], node_name='', undo=True)
 		Commands.Add(FString::Printf(TEXT("blueprint.get_controller_by_name('%s').add_prototype_node('%s', %s, '%s')"),
-						*GetGraph()->GetGraphName(),
+						*GraphName,
 						*PrototypeNode->GetNotation().ToString(),
 						*RigVMPythonUtils::Vector2DToPythonString(PrototypeNode->GetPosition()),
-						*PrototypeNode->GetName()));
+						*NodeName));
 	}
-	else if (URigVMEnumNode* EnumNode = Cast<URigVMEnumNode>(Node))
+	else if (const URigVMEnumNode* EnumNode = Cast<URigVMEnumNode>(Node))
 	{
 		// add_enum_node(cpp_type_object_path, position=[0.0, 0.0], node_name='', undo=True)
 		Commands.Add(FString::Printf(TEXT("blueprint.get_controller_by_name('%s').add_enum_node('%s', %s, '%s')"),
-						*GetGraph()->GetGraphName(),
+						*GraphName,
 						*EnumNode->GetCPPTypeObject()->GetPathName(),
 						*RigVMPythonUtils::Vector2DToPythonString(EnumNode->GetPosition()),
-						*EnumNode->GetName()));
+						*NodeName));
 	}
-	else if (URigVMLibraryNode* LibraryNode = Cast<URigVMLibraryNode>(Node))
+	else if (const URigVMLibraryNode* LibraryNode = Cast<URigVMLibraryNode>(Node))
 	{
+		const FString ContainedGraphName = GetSanitizedGraphName(LibraryNode->GetContainedGraph()->GetGraphName());
+
 		// AddFunctionReferenceNode(URigVMLibraryNode* InFunctionDefinition, const FVector2D& InNodePosition = FVector2D::ZeroVector, const FString& InNodeName = TEXT(""), bool bSetupUndoRedo = true);
 		URigVMFunctionLibrary* Library = LibraryNode->GetLibrary();
 		if (!Library || Library == GetGraph()->GetDefaultFunctionLibrary())
 		{
 			Commands.Add(FString::Printf(TEXT("blueprint.get_controller_by_name('%s').add_function_reference_node(function_%s, %s, '%s')"),
-						*GetGraph()->GetGraphName(),
-						*RigVMPythonUtils::NameToPep8(LibraryNode->GetContainedGraph()->GetGraphName()),
+						*GraphName,
+						*RigVMPythonUtils::NameToPep8(ContainedGraphName),
 						*RigVMPythonUtils::Vector2DToPythonString(LibraryNode->GetPosition()), 
-						*LibraryNode->GetName()));
+						*NodeName));
 		}
 		else
 		{
 			Commands.Add(FString::Printf(TEXT("function_blueprint = unreal.load_object(name = '%s', outer = None)"),
 				*Library->GetOuter()->GetPathName()));
 			Commands.Add(FString::Printf(TEXT("blueprint.get_controller_by_name('%s').add_function_reference_node(function_blueprint.get_local_function_library().find_function('%s'), %s, '%s')"),
-						*GetGraph()->GetGraphName(),
-						*LibraryNode->GetName(),
+						*GraphName,
+						*NodeName,
 						*RigVMPythonUtils::Vector2DToPythonString(LibraryNode->GetPosition()), 
-						*LibraryNode->GetName()));
+						*NodeName));
 		}
 
 		if (Node->IsA<URigVMCollapseNode>())
 		{
 			Commands.Add(FString::Printf(TEXT("blueprint.get_controller_by_name('%s').promote_function_reference_node_to_collapse_node('%s')"),
-					*GetGraph()->GetGraphName(),
-					*Node->GetName()));
+					*GraphName,
+					*NodeName));
 			Commands.Add(FString::Printf(TEXT("library_controller.remove_function_from_library('%s')"),
-					*LibraryNode->GetContainedGraph()->GetGraphName()));
+					*ContainedGraphName));
 		}
 	}
-	else if (URigVMRerouteNode* RerouteNode = Cast<URigVMRerouteNode>(Node))
+	else if (const URigVMRerouteNode* RerouteNode = Cast<URigVMRerouteNode>(Node))
 	{
 		// Do nothing, we need to create the links first
 	}
@@ -647,7 +682,7 @@ TArray<FString> URigVMController::GetAddNodePythonCommands(URigVMNode* Node) con
 		
 		
 	}
-	else if (URigVMArrayNode* ArrayNode = Cast<URigVMArrayNode>(Node))
+	else if (const URigVMArrayNode* ArrayNode = Cast<URigVMArrayNode>(Node))
 	{
 		FString OpCodeString = StaticEnum<ERigVMOpCode>()->GetNameStringByValue((int64)ArrayNode->GetOpCode());
 		OpCodeString = RigVMPythonUtils::NameToPep8(OpCodeString);
@@ -657,21 +692,21 @@ TArray<FString> URigVMController::GetAddNodePythonCommands(URigVMNode* Node) con
 		if (ArrayNode->GetCPPTypeObject())
 		{
 			Commands.Add(FString::Printf(TEXT("blueprint.get_controller_by_name('%s').add_array_node_from_object_path(unreal.RigVMOpCode.%s, '%s', '%s', %s, '%s')"),
-					*GetGraph()->GetGraphName(),
+					*GraphName,
 					*OpCodeString,
 					*ArrayNode->GetCPPType(),
 					*ArrayNode->GetCPPTypeObject()->GetPathName(),
 					*RigVMPythonUtils::Vector2DToPythonString(ArrayNode->GetPosition()),
-					*ArrayNode->GetName()));	
+					*NodeName));	
 		}
 		else
 		{
 			Commands.Add(FString::Printf(TEXT("blueprint.get_controller_by_name('%s').add_array_node(unreal.RigVMOpCode.%s, '%s', None, %s, '%s')"),
-					*GetGraph()->GetGraphName(),
+					*GraphName,
 					*OpCodeString,
 					*ArrayNode->GetCPPType(),
 					*RigVMPythonUtils::Vector2DToPythonString(ArrayNode->GetPosition()),
-					*ArrayNode->GetName()));	
+					*NodeName));	
 		}
 	}
 	else
@@ -691,9 +726,11 @@ TArray<FString> URigVMController::GetAddNodePythonCommands(URigVMNode* Node) con
 			const FString DefaultValue = Pin->GetDefaultValue();
 			if (!DefaultValue.IsEmpty() && DefaultValue != TEXT("()"))
 			{
+				const FString PinPath = GetSanitizedPinPath(Pin->GetPinPath());
+
 				Commands.Add(FString::Printf(TEXT("blueprint.get_controller_by_name('%s').set_pin_default_value('%s', '%s')"),
-							*GetGraph()->GetGraphName(),
-							*Pin->GetPinPath(),
+							*GraphName,
+							*PinPath,
 							*Pin->GetDefaultValue()));
 			}
 		}
@@ -817,7 +854,7 @@ URigVMUnitNode* URigVMController::AddUnitNode(UScriptStruct* InScriptStruct, con
 		TArray<FString> Commands = GetAddNodePythonCommands(Node);
 		for (const FString& Command : Commands)
 		{
-			RigVMPythonUtils::Print(GetGraph()->GetRootGraph()->GetOuter()->GetFName().ToString(),
+			RigVMPythonUtils::Print(GetGraphOuterName(),
 				FString::Printf(TEXT("%s"), *Command));
 		}
 	}
@@ -1020,7 +1057,7 @@ URigVMVariableNode* URigVMController::AddVariableNode(const FName& InVariableNam
 		TArray<FString> Commands = GetAddNodePythonCommands(Node);
 		for (const FString& Command : Commands)
 		{
-			RigVMPythonUtils::Print(GetGraph()->GetRootGraph()->GetOuter()->GetFName().ToString(), 
+			RigVMPythonUtils::Print(GetGraphOuterName(), 
 								FString::Printf(TEXT("%s"), *Command));
 		}
 	}
@@ -1615,7 +1652,7 @@ URigVMParameterNode* URigVMController::AddParameterNode(const FName& InParameter
 		TArray<FString> Commands = GetAddNodePythonCommands(Node);
 		for (const FString& Command : Commands)
 		{
-			RigVMPythonUtils::Print(GetGraph()->GetRootGraph()->GetOuter()->GetFName().ToString(), 
+			RigVMPythonUtils::Print(GetGraphOuterName(), 
 								FString::Printf(TEXT("%s"), *Command));
 		}
 	}
@@ -1693,7 +1730,7 @@ URigVMCommentNode* URigVMController::AddCommentNode(const FString& InCommentText
 		TArray<FString> Commands = GetAddNodePythonCommands(Node);
 		for (const FString& Command : Commands)
 		{
-			RigVMPythonUtils::Print(GetGraph()->GetRootGraph()->GetOuter()->GetFName().ToString(), 
+			RigVMPythonUtils::Print(GetGraphOuterName(), 
 								FString::Printf(TEXT("%s"), *Command));
 		}
 	}
@@ -1715,7 +1752,7 @@ URigVMRerouteNode* URigVMController::AddRerouteNodeOnLink(URigVMLink* InLink, bo
 	}
 
 	URigVMPin* SourcePin = InLink->GetSourcePin();
-	URigVMPin* TargetPin = InLink->GetTargetPin();
+	const URigVMPin* TargetPin = InLink->GetTargetPin();
 
 	TGuardValue<bool> GuardCompactness(bIgnoreRerouteCompactnessChanges, true);
 
@@ -1746,13 +1783,16 @@ URigVMRerouteNode* URigVMController::AddRerouteNodeOnLink(URigVMLink* InLink, bo
 
 	if (bPrintPythonCommand)
 	{
-		RigVMPythonUtils::Print(GetGraph()->GetRootGraph()->GetOuter()->GetFName().ToString(), 
+		const FString GraphName = GetSanitizedGraphName(GetGraph()->GetGraphName());
+		const FString NodeName = GetSanitizedNodeName(Node->GetName());
+
+		RigVMPythonUtils::Print(GetGraphOuterName(), 
 							FString::Printf(TEXT("blueprint.get_controller_by_name('%s').add_reroute_node_on_link_path('%s', %s, %s, '%s')"),
-											*GetGraph()->GetGraphName(),
+											*GraphName,
 											*InLink->GetPinPathRepresentation(),
 											(bShowAsFullNode) ? TEXT("True") : TEXT("False"),
 											*RigVMPythonUtils::Vector2DToPythonString(Node->GetPosition()),
-											*Node->GetName()));
+											*NodeName));
 	}
 
 	return Node;
@@ -1862,15 +1902,18 @@ URigVMRerouteNode* URigVMController::AddRerouteNodeOnPin(const FString& InPinPat
 
 	if (bPrintPythonCommand)
 	{
+		const FString GraphName = GetSanitizedGraphName(GetGraph()->GetGraphName());
+		const FString NodeName = GetSanitizedNodeName(Node->GetName());
 		// AddRerouteNodeOnPin(const FString& InPinPath, bool bAsInput, bool bShowAsFullNode, const FVector2D& InPosition, const FString& InNodeName, bool bSetupUndoRedo, bool bPrintPythonCommand)
-		RigVMPythonUtils::Print(GetGraph()->GetRootGraph()->GetOuter()->GetFName().ToString(), 
+		
+		RigVMPythonUtils::Print(GetGraphOuterName(), 
 							FString::Printf(TEXT("blueprint.get_controller_by_name('%s').add_reroute_node_on_pin('%s', %s, %s, %s '%s')"),
-											*GetGraph()->GetGraphName(),
-											*InPinPath,
+											*GraphName,
+											*GetSanitizedPinPath(InPinPath),
 											(bAsInput) ? TEXT("True") : TEXT("False"),
 											(bShowAsFullNode) ? TEXT("True") : TEXT("False"),
 											*RigVMPythonUtils::Vector2DToPythonString(Node->GetPosition()),
-											*Node->GetName()));
+											*NodeName));
 	}
 
 	return Node;
@@ -2070,16 +2113,18 @@ URigVMInjectionInfo* URigVMController::AddInjectedNode(const FString& InPinPath,
 
 	if (bPrintPythonCommand)
 	{
-		RigVMPythonUtils::Print(GetGraph()->GetRootGraph()->GetOuter()->GetFName().ToString(), 
+		const FString GraphName = GetSanitizedGraphName(GetGraph()->GetGraphName());
+		
+		RigVMPythonUtils::Print(GetGraphOuterName(), 
 							FString::Printf(TEXT("blueprint.get_controller_by_name('%s').add_injected_node_from_struct_path('%s', %s, '%s', '%s', '%s', '%s', '%s')"),
-											*GetGraph()->GetGraphName(),
-											*InPinPath,
+											*GraphName,
+											*GetSanitizedPinPath(InPinPath),
 											(bAsInput) ? TEXT("True") : TEXT("False"),
 											*InScriptStruct->GetPathName(),
 											*InMethodName.ToString(),
-											*InInputPinName.ToString(),
-											*InOutputPinName.ToString(),
-											*InNodeName));
+											*GetSanitizedPinName(InInputPinName.ToString()),
+											*GetSanitizedPinName(InOutputPinName.ToString()),
+											*GetSanitizedNodeName(InNodeName)));
 	}
 
 	return InjectionInfo;
@@ -2226,10 +2271,12 @@ URigVMNode* URigVMController::EjectNodeFromPin(const FString& InPinPath, bool bS
 
 	if (bPrintPythonCommand)
 	{
-		RigVMPythonUtils::Print(GetGraph()->GetRootGraph()->GetOuter()->GetFName().ToString(), 
+		const FString GraphName = GetSanitizedGraphName(GetGraph()->GetGraphName());
+
+		RigVMPythonUtils::Print(GetGraphOuterName(), 
 							FString::Printf(TEXT("blueprint.get_controller_by_name('%s').eject_node_from_pin('%s')"),
-											*GetGraph()->GetGraphName(),
-											*InPinPath));
+											*GraphName,
+											*GetSanitizedPinPath(InPinPath)));
 	}
 
 	return EjectedNode;
@@ -2626,10 +2673,12 @@ TArray<FName> URigVMController::ImportNodesFromText(const FString& InText, bool
 		FString PythonContent = InText.Replace(TEXT("\\\""), TEXT("\\\\\""));
 		PythonContent = InText.Replace(TEXT("'"), TEXT("\\'"));
 		PythonContent = PythonContent.Replace(TEXT("\r\n"), TEXT("\\r\\n'\r\n'"));
-		
-		RigVMPythonUtils::Print(GetGraph()->GetRootGraph()->GetOuter()->GetFName().ToString(),
+
+		const FString GraphName = GetSanitizedGraphName(GetGraph()->GetGraphName());
+
+		RigVMPythonUtils::Print(GetGraphOuterName(),
 			FString::Printf(TEXT("blueprint.get_controller_by_name('%s').import_nodes_from_text('%s')"),
-			*GetGraph()->GetGraphName(),
+			*GraphName,
 			*PythonContent));
 	}
 #endif
@@ -2857,10 +2906,12 @@ TMap<URigVMLibraryNode*, URigVMLibraryNode*> URigVMController::LocalizeFunctions
 			}
 		}
 		FunctionNames += TEXT("]");
-		
-		RigVMPythonUtils::Print(GetGraph()->GetRootGraph()->GetOuter()->GetFName().ToString(), 
+
+		const FString GraphName = GetSanitizedGraphName(GetGraph()->GetGraphName());
+
+		RigVMPythonUtils::Print(GetGraphOuterName(), 
 							FString::Printf(TEXT("blueprint.get_controller_by_name('%s').localize_functions(%s, %s)"),
-											*GetGraph()->GetGraphName(),
+											*GraphName,
 											*FunctionNames,
 											(bLocalizeDependentPrivateFunctions) ? TEXT("True") : TEXT("False")));
 	}
@@ -2868,15 +2919,17 @@ TMap<URigVMLibraryNode*, URigVMLibraryNode*> URigVMController::LocalizeFunctions
 	return LocalizedFunctions;
 }
 
-FName URigVMController::GetUniqueName(const FName& InName, TFunction<bool(const FName&)> IsNameAvailableFunction)
+FName URigVMController::GetUniqueName(const FName& InName, TFunction<bool(const FName&)> IsNameAvailableFunction, bool bAllowPeriod, bool bAllowSpace)
 {
-	FString NamePrefix = InName.ToString();
+	FString SanitizedPrefix = InName.ToString();
+	SanitizeName(SanitizedPrefix, bAllowPeriod, bAllowSpace);
+
 	int32 NameSuffix = 0;
-	FString Name = NamePrefix;
+	FString Name = SanitizedPrefix;
 	while (!IsNameAvailableFunction(*Name))
 	{
 		NameSuffix++;
-		Name = FString::Printf(TEXT("%s_%d"), *NamePrefix, NameSuffix);
+		Name = FString::Printf(TEXT("%s_%d"), *SanitizedPrefix, NameSuffix);
 	}
 	return *Name;
 }
@@ -2916,10 +2969,12 @@ URigVMCollapseNode* URigVMController::CollapseNodes(const TArray<FName>& InNodeN
 			}
 		}
 		ArrayStr += TEXT("]");
-		
-		RigVMPythonUtils::Print(GetGraph()->GetRootGraph()->GetOuter()->GetFName().ToString(), 
+
+		const FString GraphName = GetSanitizedGraphName(GetGraph()->GetGraphName());
+
+		RigVMPythonUtils::Print(GetGraphOuterName(), 
 							FString::Printf(TEXT("blueprint.get_controller_by_name('%s').collapse_nodes(%s, '%s')"),
-											*GetGraph()->GetGraphName(),
+											*GraphName,
 											*ArrayStr,
 											*InCollapseNodeName));
 	}
@@ -2955,10 +3010,13 @@ TArray<URigVMNode*> URigVMController::ExpandLibraryNode(const FName& InNodeName,
 
 	if (!Nodes.IsEmpty() && bPrintPythonCommand)
 	{
-		RigVMPythonUtils::Print(GetGraph()->GetRootGraph()->GetOuter()->GetFName().ToString(), 
+		const FString GraphName = GetSanitizedGraphName(GetGraph()->GetGraphName());
+		const FString NodeName = GetSanitizedNodeName(Node->GetName());
+		
+		RigVMPythonUtils::Print(GetGraphOuterName(), 
 							FString::Printf(TEXT("blueprint.get_controller_by_name('%s').expand_library_node('%s')"),
-											*GetGraph()->GetGraphName(),
-											*InNodeName.ToString()));
+											*GraphName,
+											*NodeName));
 	}
 
 	return Nodes;
@@ -3240,7 +3298,7 @@ URigVMCollapseNode* URigVMController::CollapseNodes(const TArray<URigVMNode*>& I
 
 		FName PinName = GetUniqueName(PinToCollapse->GetFName(), [CollapseNode](const FName& InName) {
 			return CollapseNode->FindPin(InName.ToString()) == nullptr;
-		});
+		}, false, true);
 
 		URigVMPin* CollapsedPin = NewObject<URigVMPin>(CollapseNode, PinName);
 		ConfigurePinFromPin(CollapsedPin, PinToCollapse);
@@ -4084,10 +4142,12 @@ FName URigVMController::PromoteCollapseNodeToFunctionReferenceNode(const FName&
 	{
 		if (bPrintPythonCommand)
 		{
-			RigVMPythonUtils::Print(GetGraph()->GetRootGraph()->GetOuter()->GetFName().ToString(), 
+			const FString GraphName = GetSanitizedGraphName(GetGraph()->GetGraphName());
+			
+			RigVMPythonUtils::Print(GetGraphOuterName(), 
 								FString::Printf(TEXT("blueprint.get_controller_by_name('%s').promote_collapse_node_to_function_reference_node('%s')"),
-													*GetGraph()->GetGraphName(),
-													*InNodeName.ToString()));
+													*GraphName,
+													*GetSanitizedNodeName(InNodeName.ToString())));
 		}
 		
 		return Result->GetFName();
@@ -4554,18 +4614,24 @@ bool URigVMController::RemoveNode(URigVMNode* InNode, bool bSetupUndoRedo, bool
 
 	if (bPrintPythonCommand)
 	{
+		const FString GraphName = GetSanitizedGraphName(GetGraph()->GetGraphName());
+
 		if (Graph->IsA<URigVMFunctionLibrary>())
 		{
-			RigVMPythonUtils::Print(GetGraph()->GetRootGraph()->GetOuter()->GetFName().ToString(), 
+			const FString NodeName = GetSanitizedNodeName(InNode->GetName());
+			
+			RigVMPythonUtils::Print(GetGraphOuterName(), 
 								FString::Printf(TEXT("library_controller.remove_function_from_library('%s')"),
-												*InNode->GetFName().ToString()));
+												*NodeName));
 		}
 		else
 		{
-			RigVMPythonUtils::Print(GetGraph()->GetRootGraph()->GetOuter()->GetFName().ToString(), 
+			const FString NodePath = GetSanitizedPinPath(InNode->GetNodePath());
+			
+			RigVMPythonUtils::Print(GetGraphOuterName(), 
 								FString::Printf(TEXT("blueprint.get_controller_by_name('%s').remove_node_by_name('%s')"),
-												*GetGraph()->GetGraphName(),
-												*InNode->GetNodePath()));
+												*GraphName,
+												*NodePath));
 		}
 	}
 
@@ -4670,7 +4736,7 @@ bool URigVMController::RenameNode(URigVMNode* InNode, const FName& InNewName, bo
 
 	if (GetGraph()->IsA<URigVMFunctionLibrary>())
 	{
-		RigVMPythonUtils::Print(GetGraph()->GetRootGraph()->GetOuter()->GetFName().ToString(), 
+		RigVMPythonUtils::Print(GetGraphOuterName(), 
 							FString::Printf(TEXT("library_controller.rename_function('%s', '%s')"),
 										*OldName,
 										*InNewName.ToString()));
@@ -4814,17 +4880,19 @@ bool URigVMController::SetNodeSelection(const TArray<FName>& InNodeNames, bool b
 		FString ArrayStr = TEXT("[");
 		for (auto It = InNodeNames.CreateConstIterator(); It; ++It)
 		{
-			ArrayStr += TEXT("'") + It->ToString() + TEXT("'");
+			ArrayStr += TEXT("'") + GetSanitizedNodeName(It->ToString()) + TEXT("'");
 			if (It.GetIndex() < InNodeNames.Num() - 1)
 			{
 				ArrayStr += TEXT(", ");
 			}
 		}
 		ArrayStr += TEXT("]");
-		
-		RigVMPythonUtils::Print(GetGraph()->GetRootGraph()->GetOuter()->GetFName().ToString(), 
+
+		const FString GraphName = GetSanitizedGraphName(GetGraph()->GetGraphName());
+
+		RigVMPythonUtils::Print(GetGraphOuterName(), 
 							FString::Printf(TEXT("blueprint.get_controller_by_name('%s').set_node_selection(%s)"),
-											*GetGraph()->GetGraphName(),
+											*GraphName,
 											*ArrayStr));
 	}
 
@@ -4861,10 +4929,13 @@ bool URigVMController::SetNodePosition(URigVMNode* InNode, const FVector2D& InPo
 
 	if (bPrintPythonCommand)
 	{
-		RigVMPythonUtils::Print(GetGraph()->GetRootGraph()->GetOuter()->GetFName().ToString(), 
+		const FString GraphName = GetSanitizedGraphName(GetGraph()->GetGraphName());
+		const FString NodePath = GetSanitizedPinPath(InNode->GetNodePath());
+
+		RigVMPythonUtils::Print(GetGraphOuterName(), 
 							FString::Printf(TEXT("blueprint.get_controller_by_name('%s').set_node_position_by_name('%s', %s)"),
-											*GetGraph()->GetGraphName(),
-											*InNode->GetNodePath(),
+											*GraphName,
+											*NodePath,
 											*RigVMPythonUtils::Vector2DToPythonString(InPosition)));
 	}
 
@@ -4916,10 +4987,13 @@ bool URigVMController::SetNodeSize(URigVMNode* InNode, const FVector2D& InSize,
 
 	if (bPrintPythonCommand)
 	{
-		RigVMPythonUtils::Print(GetGraph()->GetRootGraph()->GetOuter()->GetFName().ToString(), 
+		const FString GraphName = GetSanitizedGraphName(GetGraph()->GetGraphName());
+		const FString NodePath = GetSanitizedPinPath(InNode->GetNodePath());
+
+		RigVMPythonUtils::Print(GetGraphOuterName(), 
 							FString::Printf(TEXT("blueprint.get_controller_by_name('%s').set_node_size_by_name('%s', %s)"),
-											*GetGraph()->GetGraphName(),
-											*InNode->GetNodePath(),
+											*GraphName,
+											*NodePath,
 											*RigVMPythonUtils::Vector2DToPythonString(InSize)));
 	}
 
@@ -4982,10 +5056,13 @@ bool URigVMController::SetNodeColor(URigVMNode* InNode, const FLinearColor& InCo
 
 	if (bPrintPythonCommand)
 	{
-		RigVMPythonUtils::Print(GetGraph()->GetRootGraph()->GetOuter()->GetFName().ToString(), 
+		const FString GraphName = GetSanitizedGraphName(GetGraph()->GetGraphName());
+		const FString NodePath = GetSanitizedPinPath(InNode->GetNodePath());
+
+		RigVMPythonUtils::Print(GetGraphOuterName(), 
 							FString::Printf(TEXT("blueprint.get_controller_by_name('%s').set_node_color_by_name('%s', %s)"),
-											*GetGraph()->GetGraphName(),
-											*InNode->GetNodePath(),
+											*GraphName,
+											*NodePath,
 											*RigVMPythonUtils::LinearColorToPythonString(InColor)));
 	}
 
@@ -5036,10 +5113,13 @@ bool URigVMController::SetNodeCategory(URigVMCollapseNode* InNode, const FString
 
 	if (bPrintPythonCommand)
 	{
-		RigVMPythonUtils::Print(GetGraph()->GetRootGraph()->GetOuter()->GetFName().ToString(),
+		const FString GraphName = GetSanitizedGraphName(GetGraph()->GetGraphName());
+		const FString NodePath = GetSanitizedPinPath(InNode->GetNodePath());
+
+		RigVMPythonUtils::Print(GetGraphOuterName(),
 			FString::Printf(TEXT("blueprint.get_controller_by_name('%s').set_node_category_by_name('%s', '%s')"),
-				*GetGraph()->GetGraphName(),
-				*InNode->GetNodePath(),
+				*GraphName,
+				*NodePath,
 				*InCategory));
 	}
 
@@ -5090,10 +5170,13 @@ bool URigVMController::SetNodeKeywords(URigVMCollapseNode* InNode, const FString
 
 	if (bPrintPythonCommand)
 	{
-		RigVMPythonUtils::Print(GetGraph()->GetRootGraph()->GetOuter()->GetFName().ToString(),
+		const FString GraphName = GetSanitizedGraphName(GetGraph()->GetGraphName());
+		const FString NodePath = GetSanitizedPinPath(InNode->GetNodePath());
+
+		RigVMPythonUtils::Print(GetGraphOuterName(),
 			FString::Printf(TEXT("blueprint.get_controller_by_name('%s').set_node_keywords_by_name('%s', '%s')"),
-				*GetGraph()->GetGraphName(),
-				*InNode->GetNodePath(),
+				*GraphName,
+				*NodePath,
 				*InKeywords));
 	}
 
@@ -5144,10 +5227,13 @@ bool URigVMController::SetNodeDescription(URigVMCollapseNode* InNode, const FStr
 
 	if (bPrintPythonCommand)
 	{
-		RigVMPythonUtils::Print(GetGraph()->GetRootGraph()->GetOuter()->GetFName().ToString(),
+		const FString GraphName = GetSanitizedGraphName(GetGraph()->GetGraphName());
+		const FString NodePath = GetSanitizedPinPath(InNode->GetNodePath());
+
+		RigVMPythonUtils::Print(GetGraphOuterName(),
 			FString::Printf(TEXT("blueprint.get_controller_by_name('%s').set_node_description_by_name('%s', '%s')"),
-				*GetGraph()->GetGraphName(),
-				*InNode->GetNodePath(),
+				*GraphName,
+				*NodePath,
 				*InDescription));
 	}
 
@@ -5200,10 +5286,13 @@ bool URigVMController::SetCommentText(URigVMNode* InNode, const FString& InComme
 
 		if (bPrintPythonCommand)
 		{
-			RigVMPythonUtils::Print(GetGraph()->GetRootGraph()->GetOuter()->GetFName().ToString(),
+			const FString GraphName = GetSanitizedGraphName(GetGraph()->GetGraphName());
+			const FString NodePath = GetSanitizedPinPath(CommentNode->GetNodePath());
+
+			RigVMPythonUtils::Print(GetGraphOuterName(),
 				FString::Printf(TEXT("blueprint.get_controller_by_name('%s').set_comment_text_by_name('%s', '%s')"),
-				*GetGraph()->GetGraphName(),
-				*CommentNode->GetNodePath(),
+				*GraphName,
+				*NodePath,
 				*InCommentText));
 		}
 
@@ -5469,10 +5558,12 @@ bool URigVMController::SetPinExpansion(const FString& InPinPath, bool bIsExpande
 	const bool bSuccess = SetPinExpansion(Pin, bIsExpanded, bSetupUndoRedo);
 	if (bSuccess && bPrintPythonCommand)
 	{
-		RigVMPythonUtils::Print(GetGraph()->GetRootGraph()->GetOuter()->GetFName().ToString(),
+		const FString GraphName = GetSanitizedGraphName(GetGraph()->GetGraphName());
+
+		RigVMPythonUtils::Print(GetGraphOuterName(),
 			FString::Printf(TEXT("blueprint.get_controller_by_name('%s').set_pin_expansion('%s', %s)"),
-			*GetGraph()->GetGraphName(),
-			*InPinPath,
+			*GraphName,
+			*GetSanitizedPinPath(InPinPath),
 			(bIsExpanded) ? TEXT("True") : TEXT("False")));
 	}
 
@@ -5657,10 +5748,12 @@ bool URigVMController::SetPinDefaultValue(const FString& InPinPath, const FStrin
 
 	if (bPrintPythonCommand)
 	{
-		RigVMPythonUtils::Print(GetGraph()->GetRootGraph()->GetOuter()->GetFName().ToString(),
+		const FString GraphName = GetSanitizedGraphName(GetGraph()->GetGraphName());
+
+		RigVMPythonUtils::Print(GetGraphOuterName(),
 			FString::Printf(TEXT("blueprint.get_controller_by_name('%s').set_pin_default_value('%s', '%s', %s)"),
-			*GetGraph()->GetGraphName(),
-			*InPinPath,
+			*GraphName,
+			*GetSanitizedPinPath(InPinPath),
 			*InDefaultValue,
 			(bResizeArrays) ? TEXT("True") : TEXT("False")));
 	}
@@ -5804,10 +5897,12 @@ bool URigVMController::ResetPinDefaultValue(const FString& InPinPath, bool bSetu
 	const bool bSuccess = ResetPinDefaultValue(Pin, bSetupUndoRedo);
 	if (bSuccess && bPrintPythonCommand)
 	{
-		RigVMPythonUtils::Print(GetGraph()->GetRootGraph()->GetOuter()->GetFName().ToString(),
+		const FString GraphName = GetSanitizedGraphName(GetGraph()->GetGraphName());
+
+		RigVMPythonUtils::Print(GetGraphOuterName(),
 			FString::Printf(TEXT("blueprint.get_controller_by_name('%s').reset_pin_default_value('%s')"),
-			*GetGraph()->GetGraphName(),
-			*InPinPath));
+			*GraphName,
+			*GetSanitizedPinPath(InPinPath)));
 	}
 
 	return bSuccess;
@@ -5943,10 +6038,12 @@ FString URigVMController::InsertArrayPin(const FString& InArrayPinPath, int32 In
 	{
 		if (bPrintPythonCommand)
 		{
-			RigVMPythonUtils::Print(GetGraph()->GetRootGraph()->GetOuter()->GetFName().ToString(),
+			const FString GraphName = GetSanitizedGraphName(GetGraph()->GetGraphName());
+			
+			RigVMPythonUtils::Print(GetGraphOuterName(),
 				FString::Printf(TEXT("blueprint.get_controller_by_name('%s').insert_array_pin('%s', %d, '%s')"),
-				*GetGraph()->GetGraphName(),
-				*InArrayPinPath,
+				*GraphName,
+				*GetSanitizedPinPath(InArrayPinPath),
 				InIndex,
 				*InDefaultValue));
 		}
@@ -6098,10 +6195,12 @@ bool URigVMController::RemoveArrayPin(const FString& InArrayElementPinPath, bool
 
 	if (bPrintPythonCommand)
 	{
-		RigVMPythonUtils::Print(GetGraph()->GetRootGraph()->GetOuter()->GetFName().ToString(),
+		const FString GraphName = GetSanitizedGraphName(GetGraph()->GetGraphName());
+
+		RigVMPythonUtils::Print(GetGraphOuterName(),
 			FString::Printf(TEXT("blueprint.get_controller_by_name('%s').remove_array_pin('%s')"),
-			*GetGraph()->GetGraphName(),
-			*InArrayElementPinPath));
+			*GraphName,
+			*GetSanitizedPinPath(InArrayElementPinPath)));
 	}
 
 	return true;
@@ -6255,10 +6354,12 @@ bool URigVMController::BindPinToVariable(const FString& InPinPath, const FString
 	const bool bSuccess = BindPinToVariable(Pin, InNewBoundVariablePath, bSetupUndoRedo);
 	if (bSuccess && bPrintPythonCommand)
 	{
-		RigVMPythonUtils::Print(GetGraph()->GetRootGraph()->GetOuter()->GetFName().ToString(),
+		const FString GraphName = GetSanitizedGraphName(GetGraph()->GetGraphName());
+		
+		RigVMPythonUtils::Print(GetGraphOuterName(),
 			FString::Printf(TEXT("blueprint.get_controller_by_name('%s').bind_pin_to_variable('%s', '%s')"),
-			*GetGraph()->GetGraphName(),
-			*InPinPath,
+			*GraphName,
+			*GetSanitizedPinPath(InPinPath),
 			*InNewBoundVariablePath));
 	}
 	
@@ -6482,10 +6583,12 @@ bool URigVMController::PromotePinToVariable(const FString& InPinPath, bool bCrea
 	const bool bSuccess = PromotePinToVariable(Pin, bCreateVariableNode, InNodePosition, bSetupUndoRedo);
 	if (bSuccess && bPrintPythonCommand)
 	{
-		RigVMPythonUtils::Print(GetGraph()->GetRootGraph()->GetOuter()->GetFName().ToString(), 
+		const FString GraphName = GetSanitizedGraphName(GetGraph()->GetGraphName());
+		
+		RigVMPythonUtils::Print(GetGraphOuterName(), 
 			FString::Printf(TEXT("blueprint.get_controller_by_name('%s').promote_pin_to_variable('%s', %s, %s)"),
-			*GetGraph()->GetGraphName(),
-			*InPinPath,
+			*GraphName,
+			*GetSanitizedPinPath(InPinPath),
 			(bCreateVariableNode) ? TEXT("True") : TEXT("False"),
 			*RigVMPythonUtils::Vector2DToPythonString(InNodePosition)));
 	}
@@ -6611,12 +6714,17 @@ bool URigVMController::AddLink(const FString& InOutputPinPath, const FString& In
 
 	const bool bSuccess = AddLink(OutputPin, InputPin, bSetupUndoRedo);
 	if (bSuccess && bPrintPythonCommand)
-	{		
-		RigVMPythonUtils::Print(GetGraph()->GetRootGraph()->GetOuter()->GetFName().ToString(), 
+	{
+		const FString GraphName = GetSanitizedGraphName(GetGraph()->GetGraphName());
+
+		const FString SanitizedInputPinPath = GetSanitizedPinPath(InputPin->GetPinPath());
+		const FString SanitizedOutputPinPath = GetSanitizedPinPath(OutputPin->GetPinPath());
+
+		RigVMPythonUtils::Print(GetGraphOuterName(), 
 			FString::Printf(TEXT("blueprint.get_controller_by_name('%s').add_link('%s', '%s')"),
-			*GetGraph()->GetGraphName(),
-			*OutputPin->GetPinPath(),
-			*InputPin->GetPinPath()));
+			*GraphName,
+			*SanitizedOutputPinPath,
+			*SanitizedInputPinPath));
 	}
 	
 	return bSuccess;
@@ -6770,11 +6878,13 @@ bool URigVMController::BreakLink(const FString& InOutputPinPath, const FString&
 	const bool bSuccess = BreakLink(OutputPin, InputPin, bSetupUndoRedo);
 	if (bSuccess && bPrintPythonCommand)
 	{
-		RigVMPythonUtils::Print(GetGraph()->GetRootGraph()->GetOuter()->GetFName().ToString(), 
+		const FString GraphName = GetSanitizedGraphName(GetGraph()->GetGraphName());
+
+		RigVMPythonUtils::Print(GetGraphOuterName(), 
 			FString::Printf(TEXT("blueprint.get_controller_by_name('%s').break_link('%s', '%s')"),
-			*GetGraph()->GetGraphName(),
-			*OutputPin->GetPinPath(),
-			*InputPin->GetPinPath()));
+			*GraphName,
+			*GetSanitizedPinPath(OutputPin->GetPinPath()),
+			*GetSanitizedPinPath(InputPin->GetPinPath())));
 	}
 	return bSuccess;
 }
@@ -6866,10 +6976,12 @@ bool URigVMController::BreakAllLinks(const FString& InPinPath, bool bAsInput, bo
 	const bool bSuccess = BreakAllLinks(Pin, bAsInput, bSetupUndoRedo);
 	if (bSuccess && bPrintPythonCommand)
 	{
-		RigVMPythonUtils::Print(GetGraph()->GetRootGraph()->GetOuter()->GetFName().ToString(), 
+		const FString GraphName = GetSanitizedGraphName(GetGraph()->GetGraphName());
+
+		RigVMPythonUtils::Print(GetGraphOuterName(), 
 			FString::Printf(TEXT("blueprint.get_controller_by_name('%s').break_all_links('%s', %s)"),
-			*GetGraph()->GetGraphName(),
-			*Pin->GetPinPath(),
+			*GraphName,
+			*GetSanitizedPinPath(Pin->GetPinPath()),
 			bAsInput ? TEXT("True") : TEXT("False")));
 	}
 	return bSuccess;
@@ -7002,7 +7114,7 @@ FName URigVMController::AddExposedPin(const FName& InPinName, ERigVMPinDirection
 
 		return LibraryNode->FindPin(InName.ToString()) == nullptr;
 
-	});
+	}, false, true);
 
 	URigVMPin* Pin = NewObject<URigVMPin>(LibraryNode, PinName);
 	Pin->CPPType = InCPPType;
@@ -7059,10 +7171,13 @@ FName URigVMController::AddExposedPin(const FName& InPinName, ERigVMPinDirection
 		}
 		
 		//AddExposedPin(const FName& InPinName, ERigVMPinDirection InDirection, const FString& InCPPType, const FName& InCPPTypeObjectPath, const FString& InDefaultValue, bool bSetupUndoRedo, bool bPrintPythonCommand)
-		RigVMPythonUtils::Print(GetGraph()->GetRootGraph()->GetOuter()->GetFName().ToString(), 
+
+		const FString GraphName = GetSanitizedGraphName(GetGraph()->GetGraphName());
+		
+		RigVMPythonUtils::Print(GetGraphOuterName(), 
 			FString::Printf(TEXT("blueprint.get_controller_by_name('%s').add_exposed_pin('%s', %s, '%s', '%s', '%s')"),
-				*GetGraph()->GetGraphName(),
-				*InPinName.ToString(),
+				*GraphName,
+				*GetSanitizedPinName(InPinName.ToString()),
 				*DirectionStr,
 				*InCPPType,
 				*InCPPTypeObjectPath.ToString(),
@@ -7146,10 +7261,12 @@ bool URigVMController::RemoveExposedPin(const FName& InPinName, bool bSetupUndoR
 
 	if (bSuccessfullyRemovedPin && bPrintPythonCommand)
 	{
-		RigVMPythonUtils::Print(GetGraph()->GetRootGraph()->GetOuter()->GetFName().ToString(), 
+		const FString GraphName = GetSanitizedGraphName(GetGraph()->GetGraphName());
+
+		RigVMPythonUtils::Print(GetGraphOuterName(), 
 			FString::Printf(TEXT("blueprint.get_controller_by_name('%s').remove_exposed_pin('%s')"),
-				*GetGraph()->GetGraphName(),
-				*InPinName.ToString()));
+				*GraphName,
+				*GetSanitizedPinName(InPinName.ToString())));
 	}
 
 	return bSuccessfullyRemovedPin;
@@ -7206,7 +7323,7 @@ bool URigVMController::RenameExposedPin(const FName& InOldPinName, const FName&
 
 	FName PinName = GetUniqueName(InNewPinName, [LibraryNode](const FName& InName) {
 		return LibraryNode->FindPin(InName.ToString()) == nullptr;
-		});
+	}, false, true);
 
 	FRigVMRenameExposedPinAction Action;
 	if (bSetupUndoRedo)
@@ -7295,11 +7412,13 @@ bool URigVMController::RenameExposedPin(const FName& InOldPinName, const FName&
 
 	if (bPrintPythonCommand)
 	{
-		RigVMPythonUtils::Print(GetGraph()->GetRootGraph()->GetOuter()->GetFName().ToString(), 
+		const FString GraphName = GetSanitizedGraphName(GetGraph()->GetGraphName());
+
+		RigVMPythonUtils::Print(GetGraphOuterName(), 
 			FString::Printf(TEXT("blueprint.get_controller_by_name('%s').rename_exposed_pin('%s', '%s')"),
-				*GetGraph()->GetGraphName(),
-				*InOldPinName.ToString(),
-				*InNewPinName.ToString()));
+				*GraphName,
+				*GetSanitizedPinName(InOldPinName.ToString()),
+				*GetSanitizedPinName(InNewPinName.ToString())));
 	}
 
 	return true;
@@ -7407,10 +7526,12 @@ bool URigVMController::ChangeExposedPinType(const FName& InPinName, const FStrin
 
 	if (bPrintPythonCommand)
 	{
-		RigVMPythonUtils::Print(GetGraph()->GetRootGraph()->GetOuter()->GetFName().ToString(), 
+		const FString GraphName = GetSanitizedGraphName(GetGraph()->GetGraphName());
+		
+		RigVMPythonUtils::Print(GetGraphOuterName(), 
 			FString::Printf(TEXT("blueprint.get_controller_by_name('%s').change_exposed_pin_type('%s', '%s', '%s')"),
-				*GetGraph()->GetGraphName(),
-				*InPinName.ToString(),
+				*GraphName,
+				*GetSanitizedPinName(InPinName.ToString()),
 				*InCPPType,
 				*InCPPTypeObjectPath.ToString()));
 	}
@@ -7484,10 +7605,12 @@ bool URigVMController::SetExposedPinIndex(const FName& InPinName, int32 InNewInd
 
 	if (bPrintPythonCommand)
 	{
-		RigVMPythonUtils::Print(GetGraph()->GetRootGraph()->GetOuter()->GetFName().ToString(), 
+		const FString GraphName = GetSanitizedGraphName(GetGraph()->GetGraphName());
+
+		RigVMPythonUtils::Print(GetGraphOuterName(), 
 			FString::Printf(TEXT("blueprint.get_controller_by_name('%s').set_exposed_pin_index('%s', %d)"),
-				*GetGraph()->GetGraphName(),
-				*InPinName.ToString(),
+				*GraphName,
+				*GetSanitizedPinName(InPinName.ToString()),
 				InNewIndex));
 	}
 
@@ -7573,26 +7696,30 @@ URigVMFunctionReferenceNode* URigVMController::AddFunctionReferenceNode(URigVMLi
 
 	if (bPrintPythonCommand)
 	{
+		const FString GraphName = GetSanitizedGraphName(GetGraph()->GetGraphName());
+		const FString FunctionDefinitionName = GetSanitizedNodeName(InFunctionDefinition->GetName());
+
 		if (InFunctionDefinition->GetLibrary() == GetGraph()->GetDefaultFunctionLibrary())
 		{
-			RigVMPythonUtils::Print(GetGraph()->GetRootGraph()->GetOuter()->GetFName().ToString(), 
+
+			RigVMPythonUtils::Print(GetGraphOuterName(), 
 				FString::Printf(TEXT("blueprint.get_controller_by_name('%s').add_function_reference_node(library.find_function('%s'), %s, '%s')"),
-						*GetGraph()->GetGraphName(),
-						*InFunctionDefinition->GetName(),
+						*GraphName,
+						*FunctionDefinitionName,
 						*RigVMPythonUtils::Vector2DToPythonString(InNodePosition),
-						*InNodeName));
+						*NodeName));
 		}
 		else
 		{
-			RigVMPythonUtils::Print(GetGraph()->GetRootGraph()->GetOuter()->GetFName().ToString(), 
+			RigVMPythonUtils::Print(GetGraphOuterName(), 
 				FString::Printf(TEXT("function_blueprint = unreal.load_object(name = '%s', outer = None)"),
 				*InFunctionDefinition->GetLibrary()->GetOuter()->GetPathName()));
-			RigVMPythonUtils::Print(GetGraph()->GetRootGraph()->GetOuter()->GetFName().ToString(), 
+			RigVMPythonUtils::Print(GetGraphOuterName(), 
 				FString::Printf(TEXT("blueprint.get_controller_by_name('%s').add_function_reference_node(function_blueprint.get_local_function_library().find_function('%s'), %s, '%s')"),
-						*GetGraph()->GetGraphName(),
-						*InFunctionDefinition->GetName(),
+						*GraphName,
+						*FunctionDefinitionName,
 						*RigVMPythonUtils::Vector2DToPythonString(InFunctionDefinition->GetPosition()), 
-						*InFunctionDefinition->GetName()));
+						*FunctionDefinitionName));
 		}
 		
 	}
@@ -7754,10 +7881,12 @@ URigVMLibraryNode* URigVMController::AddFunctionToLibrary(const FName& InFunctio
 
 	if (bPrintPythonCommand)
 	{
+		const FString GraphName = GetSanitizedGraphName(GetGraph()->GetGraphName());
+
 		//AddFunctionToLibrary(const FName& InFunctionName, bool bMutable, const FVector2D& InNodePosition, bool bSetupUndoRedo, bool bPrintPythonCommand)
-		RigVMPythonUtils::Print(GetGraph()->GetRootGraph()->GetOuter()->GetFName().ToString(), 
+		RigVMPythonUtils::Print(GetGraphOuterName(), 
 			FString::Printf(TEXT("library_controller.add_function_to_library('%s', %s, %s)"),
-				*InFunctionName.ToString(),
+				*GetSanitizedNodeName(InFunctionName.ToString()),
 				(bMutable) ? TEXT("True") : TEXT("False"),
 				*RigVMPythonUtils::Vector2DToPythonString(InNodePosition)));
 	}
@@ -7847,7 +7976,7 @@ FRigVMGraphVariableDescription URigVMController::AddLocalVariable(const FName& I
 			}
 		}
 		return true;
-	});
+	}, false, true);
 
 	NewVariable.Name = VariableName;
 	NewVariable.CPPType = InCPPType;
@@ -7873,9 +8002,11 @@ FRigVMGraphVariableDescription URigVMController::AddLocalVariable(const FName& I
 
 	if (bPrintPythonCommand)
 	{
-		RigVMPythonUtils::Print(GetGraph()->GetRootGraph()->GetOuter()->GetFName().ToString(), 
+		const FString GraphName = GetSanitizedGraphName(GetGraph()->GetGraphName());
+
+		RigVMPythonUtils::Print(GetGraphOuterName(), 
 			FString::Printf(TEXT("blueprint.get_controller_by_name('%s').add_local_variable_from_object_path('%s', '%s', '%s', '%s')"),
-				*GetGraph()->GetGraphName(),
+				*GraphName,
 				*NewVariable.Name.ToString(),
 				*NewVariable.CPPType,
 				(NewVariable.CPPTypeObject) ? *NewVariable.CPPTypeObject->GetPathName() : *FString(),
@@ -7975,10 +8106,12 @@ bool URigVMController::RemoveLocalVariable(const FName& InVariableName, bool bSe
 
 		if (bPrintPythonCommand)
 		{
-			RigVMPythonUtils::Print(GetGraph()->GetRootGraph()->GetOuter()->GetFName().ToString(), 
+			const FString GraphName = GetSanitizedGraphName(GetGraph()->GetGraphName());
+
+			RigVMPythonUtils::Print(GetGraphOuterName(), 
 				FString::Printf(TEXT("blueprint.get_controller_by_name('%s').remove_local_variable('%s')"),
-					*GetGraph()->GetGraphName(),
-					*InVariableName.ToString()));
+					*GraphName,
+					*GetSanitizedVariableName(InVariableName.ToString())));
 		}
 		return true;
 	}
@@ -8072,11 +8205,13 @@ bool URigVMController::RenameLocalVariable(const FName& InVariableName, const FN
 
 	if (bPrintPythonCommand)
 	{
-		RigVMPythonUtils::Print(GetGraph()->GetRootGraph()->GetOuter()->GetFName().ToString(), 
+		const FString GraphName = GetSanitizedGraphName(GetGraph()->GetGraphName());
+
+		RigVMPythonUtils::Print(GetGraphOuterName(), 
 			FString::Printf(TEXT("blueprint.get_controller_by_name('%s').rename_local_variable('%s', '%s')"),
-				*GetGraph()->GetGraphName(),
-				*InVariableName.ToString(),
-				*InNewVariableName.ToString()));
+				*GraphName,
+				*GetSanitizedVariableName(InVariableName.ToString()),
+				*GetSanitizedVariableName(InNewVariableName.ToString())));
 	}
 
 	return true;
@@ -8165,11 +8300,13 @@ bool URigVMController::SetLocalVariableType(const FName& InVariableName, const F
 
 	if (bPrintPythonCommand)
 	{
+		const FString GraphName = GetSanitizedGraphName(GetGraph()->GetGraphName());
+
 		//bool URigVMController::SetLocalVariableType(const FName& InVariableName, const FString& InCPPType, UObject* InCPPTypeObject, bool bSetupUndoRedo, bool bPrintPythonCommand)
-		RigVMPythonUtils::Print(GetGraph()->GetRootGraph()->GetOuter()->GetFName().ToString(), 
+		RigVMPythonUtils::Print(GetGraphOuterName(), 
 			FString::Printf(TEXT("blueprint.get_controller_by_name('%s').set_local_variable_type_from_object_path('%s', '%s', '%s')"),
-				*GetGraph()->GetGraphName(),
-				*InVariableName.ToString(),
+				*GraphName,
+				*GetSanitizedVariableName(InVariableName.ToString()),
 				*InCPPType,
 				(InCPPTypeObject) ? *InCPPTypeObject->GetPathName() : *FString()));
 	}
@@ -8255,10 +8392,12 @@ bool URigVMController::SetLocalVariableDefaultValue(const FName& InVariableName,
 
 	if (bPrintPythonCommand)
 	{
-		RigVMPythonUtils::Print(GetGraph()->GetRootGraph()->GetOuter()->GetFName().ToString(), 
+		const FString GraphName = GetSanitizedGraphName(GetGraph()->GetGraphName());
+
+		RigVMPythonUtils::Print(GetGraphOuterName(), 
 			FString::Printf(TEXT("blueprint.get_controller_by_name('%s').set_local_variable_default_value('%s', '%s')"),
-				*GetGraph()->GetGraphName(),
-				*InVariableName.ToString(),
+				*GraphName,
+				*GetSanitizedVariableName(InVariableName.ToString()),
 				*InDefaultValue));
 	}
 	
@@ -8476,7 +8615,7 @@ bool URigVMController::SetVariableName(URigVMVariableNode* InVariableNode, const
 			return true;
 		}
 		return InVariableNode->GetCPPType() == Descriptions[*FoundIndex].CPPType;
-	});
+	}, false, true);
 
 	int32 NodesSharingName = 0;
 	for (URigVMNode* Node : Graph->Nodes)
@@ -8536,7 +8675,7 @@ bool URigVMController::SetParameterName(URigVMParameterNode* InParameterNode, co
 			return true;
 		}
 		return InParameterNode->GetCPPType() == Descriptions[*FoundIndex].CPPType && InParameterNode->IsInput() == Descriptions[*FoundIndex].bIsInput;
-	});
+	}, false, true);
 
 	int32 NodesSharingName = 0;
 	for (URigVMNode* Node : Graph->Nodes)
@@ -8680,7 +8819,7 @@ URigVMBranchNode* URigVMController::AddBranchNode(const FVector2D& InPosition, c
 		TArray<FString> Commands = GetAddNodePythonCommands(Node);
 		for (const FString& Command : Commands)
 		{
-			RigVMPythonUtils::Print(GetGraph()->GetRootGraph()->GetOuter()->GetFName().ToString(), 
+			RigVMPythonUtils::Print(GetGraphOuterName(), 
 				FString::Printf(TEXT("%s"), *Command));
 		}
 	}
@@ -8786,7 +8925,7 @@ URigVMIfNode* URigVMController::AddIfNode(const FString& InCPPType, const FName&
 		TArray<FString> Commands = GetAddNodePythonCommands(Node);
 		for (const FString& Command : Commands)
 		{
-			RigVMPythonUtils::Print(GetGraph()->GetRootGraph()->GetOuter()->GetFName().ToString(), 
+			RigVMPythonUtils::Print(GetGraphOuterName(), 
 				FString::Printf(TEXT("%s"), *Command));
 		}
 	}
@@ -8871,7 +9010,7 @@ URigVMSelectNode* URigVMController::AddSelectNode(const FString& InCPPType, cons
 		TArray<FString> Commands = GetAddNodePythonCommands(Node);
 		for (const FString& Command : Commands)
 		{
-			RigVMPythonUtils::Print(GetGraph()->GetRootGraph()->GetOuter()->GetFName().ToString(), 
+			RigVMPythonUtils::Print(GetGraphOuterName(), 
 				FString::Printf(TEXT("%s"), *Command));
 		}
 	}
@@ -8948,7 +9087,7 @@ URigVMPrototypeNode* URigVMController::AddPrototypeNode(const FName& InNotation,
 		TArray<FString> Commands = GetAddNodePythonCommands(Node);
 		for (const FString& Command : Commands)
 		{
-			RigVMPythonUtils::Print(GetGraph()->GetRootGraph()->GetOuter()->GetFName().ToString(), 
+			RigVMPythonUtils::Print(GetGraphOuterName(), 
 				FString::Printf(TEXT("%s"), *Command));
 		}
 	}
@@ -9013,7 +9152,7 @@ URigVMEnumNode* URigVMController::AddEnumNode(const FName& InCPPTypeObjectPath,
 		TArray<FString> Commands = GetAddNodePythonCommands(Node);
 		for (const FString& Command : Commands)
 		{
-			RigVMPythonUtils::Print(GetGraph()->GetRootGraph()->GetOuter()->GetFName().ToString(), 
+			RigVMPythonUtils::Print(GetGraphOuterName(), 
 				FString::Printf(TEXT("%s"), *Command));
 		}
 	}
@@ -9313,7 +9452,7 @@ URigVMArrayNode* URigVMController::AddArrayNode(ERigVMOpCode InOpCode, const FSt
 		TArray<FString> Commands = GetAddNodePythonCommands(Node);
 		for (const FString& Command : Commands)
 		{
-			RigVMPythonUtils::Print(GetGraph()->GetRootGraph()->GetOuter()->GetFName().ToString(), 
+			RigVMPythonUtils::Print(GetGraphOuterName(), 
 				FString::Printf(TEXT("%s"), *Command));
 		}
 	}
@@ -9373,12 +9512,9 @@ FString URigVMController::GetValidNodeName(const FString& InPrefix)
 	URigVMGraph* Graph = GetGraph();
 	check(Graph);
 
-	FString SanitizedPrefix = InPrefix;
-	SanitizeName(SanitizedPrefix);
-
-	return GetUniqueName(*SanitizedPrefix, [&](const FName& InName) {
+	return GetUniqueName(*InPrefix, [&](const FName& InName) {
 		return Graph->IsNameAvailable(InName.ToString());
-	}).ToString();
+	}, false, true).ToString();
 }
 
 bool URigVMController::IsValidGraph()
@@ -11655,7 +11791,45 @@ void URigVMController::RefreshFunctionReferences(URigVMLibraryNode* InFunctionDe
 	}
 }
 
-void URigVMController::SanitizeName(FString& InOutName)
+FString URigVMController::GetGraphOuterName() const
+{
+	check(GetGraph() != nullptr);
+	return GetSanitizedName(GetGraph()->GetRootGraph()->GetOuter()->GetFName().ToString(), true, false);
+}
+
+FString URigVMController::GetSanitizedName(const FString& InName, bool bAllowPeriod, bool bAllowSpace)
+{
+	FString CopiedName = InName;
+	SanitizeName(CopiedName, bAllowPeriod, bAllowSpace);
+	return CopiedName;
+}
+
+FString URigVMController::GetSanitizedGraphName(const FString& InName)
+{
+	return GetSanitizedName(InName, true, true);
+}
+
+FString URigVMController::GetSanitizedNodeName(const FString& InName)
+{
+	return GetSanitizedName(InName, false, true);
+}
+
+FString URigVMController::GetSanitizedVariableName(const FString& InName)
+{
+	return GetSanitizedName(InName, false, true);
+}
+
+FString URigVMController::GetSanitizedPinName(const FString& InName)
+{
+	return GetSanitizedName(InName, false, true);
+}
+
+FString URigVMController::GetSanitizedPinPath(const FString& InName)
+{
+	return GetSanitizedName(InName, true, true);
+}
+
+void URigVMController::SanitizeName(FString& InOutName, bool bAllowPeriod, bool bAllowSpace)
 {
 	// Sanitize the name
 	for (int32 i = 0; i < InOutName.Len(); ++i)
@@ -11664,7 +11838,9 @@ void URigVMController::SanitizeName(FString& InOutName)
 
 		const bool bGoodChar =
 			((C >= 'A') && (C <= 'Z')) || ((C >= 'a') && (C <= 'z')) ||		// A-Z (upper and lowercase) anytime
-			(C == '_') || (C == '-') || (C == '.') || (C == ' ') ||			// _  - . and space anytime
+			(C == '_') || (C == '-') || 									// _  and - anytime
+			(bAllowPeriod && (C == '.')) ||
+			(bAllowSpace && (C == ' ')) ||
 			((i > 0) && (C >= '0') && (C <= '9'));							// 0-9 after the first character
 
 		if (!bGoodChar)
diff --git a/Engine/Source/Developer/RigVMDeveloper/Public/RigVMModel/RigVMController.h b/Engine/Source/Developer/RigVMDeveloper/Public/RigVMModel/RigVMController.h
--- a/Engine/Source/Developer/RigVMDeveloper/Public/RigVMModel/RigVMController.h
+++ b/Engine/Source/Developer/RigVMDeveloper/Public/RigVMModel/RigVMController.h
@@ -356,7 +356,7 @@ public:
         bool bPrintPythonCommand = false);
 
 	// Returns a unique name
-	static FName GetUniqueName(const FName& InName, TFunction<bool(const FName&)> IsNameAvailableFunction);
+	static FName GetUniqueName(const FName& InName, TFunction<bool(const FName&)> IsNameAvailableFunction, bool bAllowPeriod, bool bAllowSpace);
 
 	// Turns a series of nodes into a Collapse node
 	UFUNCTION(BlueprintCallable, Category = RigVMController)
@@ -856,7 +856,14 @@ private:
 	void RefreshFunctionReferences(URigVMLibraryNode* InFunctionDefinition, bool bSetupUndoRedo);
 
 	static int32 GetMaxNameLength() { return 100; }
-	static void SanitizeName(FString& InOutName);
+	FString GetGraphOuterName() const;
+	static FString GetSanitizedName(const FString& InName, bool bAllowPeriod, bool bAllowSpace);
+	static FString GetSanitizedGraphName(const FString& InName);
+	static FString GetSanitizedNodeName(const FString& InName);
+	static FString GetSanitizedVariableName(const FString& InName);
+	static FString GetSanitizedPinName(const FString& InName);
+	static FString GetSanitizedPinPath(const FString& InName);
+	static void SanitizeName(FString& InOutName, bool bAllowPeriod, bool bAllowSpace);
 	
 	UPROPERTY(transient)
 	TArray<TObjectPtr<URigVMGraph>> Graphs;
