RigVMController: Break links for target pins when importing graph content

diff --git a/Engine/Source/Developer/RigVMDeveloper/Private/RigVMModel/RigVMController.cpp b/Engine/Source/Developer/RigVMDeveloper/Private/RigVMModel/RigVMController.cpp
--- a/Engine/Source/Developer/RigVMDeveloper/Private/RigVMModel/RigVMController.cpp
+++ b/Engine/Source/Developer/RigVMDeveloper/Private/RigVMModel/RigVMController.cpp
@@ -2662,6 +2662,10 @@ TArray<FName> URigVMController::ImportNodesFromText(const FString& InText, bool
 					
 					if (SourcePin && TargetPin)
 					{
+						BreakAllLinksRecursive(TargetPin, true, true, bSetupUndoRedo);
+						BreakAllLinks(TargetPin, true, bSetupUndoRedo);
+						BreakAllLinksRecursive(TargetPin, true, false, bSetupUndoRedo);
+
 						Graph->Links.Add(CreatedLink);
 						SourcePin->Links.Add(CreatedLink);
 						TargetPin->Links.Add(CreatedLink);
