Control Rig: Fix crash occurs if undoing deleting function with function node in rig graph

diff --git a/Engine/Source/Developer/AutomationController/Private/AutomationControllerModule.cpp b/Engine/Source/Developer/AutomationController/Private/AutomationControllerModule.cpp
--- a/Engine/Source/Developer/AutomationController/Private/AutomationControllerModule.cpp
+++ b/Engine/Source/Developer/AutomationController/Private/AutomationControllerModule.cpp
@@ -60,6 +60,7 @@ void FAutomationControllerModule::StartupModule()
 	FMessageLogInitializationOptions InitOptions;
 	InitOptions.bShowFilters = true;
 	InitOptions.bShowPages = true;
+	InitOptions.bScrollToBottom = true;
 	MessageLogModule.RegisterLogListing("AutomationTestingLog", LOCTEXT("AutomationTesting", "Automation Testing Log"), InitOptions);
 #endif
 }
diff --git a/Engine/Source/Developer/RigVMDeveloper/Private/RigVMModel/RigVMController.cpp b/Engine/Source/Developer/RigVMDeveloper/Private/RigVMModel/RigVMController.cpp
--- a/Engine/Source/Developer/RigVMDeveloper/Private/RigVMModel/RigVMController.cpp
+++ b/Engine/Source/Developer/RigVMDeveloper/Private/RigVMModel/RigVMController.cpp
@@ -4681,57 +4681,10 @@
 		}
 	}
 
-	if (bSetupUndoRedo || bRecursive)
-	{
-		SelectNode(InNode, false, bSetupUndoRedo);
-
-		for (URigVMPin* Pin : InNode->GetPins())
-		{
-			TArray<URigVMInjectionInfo*> InjectedNodes = Pin->GetInjectedNodes();
-			for (URigVMInjectionInfo* InjectedNode : InjectedNodes)
-			{
-				RemoveNode(InjectedNode->UnitNode, bSetupUndoRedo, bRecursive);
-			}
-
-			BreakAllLinks(Pin, true, bSetupUndoRedo);
-			BreakAllLinks(Pin, false, bSetupUndoRedo);
-			BreakAllLinksRecursive(Pin, true, false, bSetupUndoRedo);
-			BreakAllLinksRecursive(Pin, false, false, bSetupUndoRedo);
-		}
-
-		if (bSetupUndoRedo)
-		{
-			ActionStack->AddAction(FRigVMRemoveNodeAction(InNode, this));
-		}
-
-		if (URigVMCollapseNode* CollapseNode = Cast<URigVMCollapseNode>(InNode))
-		{
-			URigVMGraph* SubGraph = CollapseNode->GetContainedGraph();
-			FRigVMControllerGraphGuard GraphGuard(this, SubGraph, false);
-
-			TArray<URigVMNode*> ContainedNodes = SubGraph->GetNodes();
-			for (URigVMNode* ContainedNode : ContainedNodes)
-			{
-				if(Cast<URigVMFunctionEntryNode>(ContainedNode) != nullptr ||
-					Cast<URigVMFunctionReturnNode>(ContainedNode) != nullptr)
-				{
-					continue;
-				}
-				RemoveNode(ContainedNode, false, true);
-			}
-		}
-	}
-
-	Graph->Nodes.Remove(InNode);
-	if (!bSuspendNotifications)
-	{
-		Graph->MarkPackageDirty();
-	}
-
-	Notify(ERigVMGraphNotifType::NodeRemoved, InNode);
-
+	
 	if (URigVMLibraryNode* LibraryNode = Cast<URigVMLibraryNode>(InNode))
 	{
+		// If we are removing a reference, remove the function references to this node in the function library
 		if(URigVMFunctionReferenceNode* FunctionReferenceNode = Cast<URigVMFunctionReferenceNode>(LibraryNode))
 		{
 			if (URigVMFunctionLibrary* FunctionLibrary = Cast<URigVMFunctionLibrary>(FunctionReferenceNode->GetLibrary()))
@@ -4755,22 +4708,29 @@
 				}
 			}
 		}
+		// If we are removing a function, remove all the references first
 		else if (URigVMFunctionLibrary* FunctionLibrary = Cast<URigVMFunctionLibrary>(LibraryNode->GetGraph()))
 		{
 			const FRigVMFunctionReferenceArray* FunctionReferencesPtr = FunctionLibrary->FunctionReferences.Find(LibraryNode);
 			if (FunctionReferencesPtr)
 			{
-				for (const TSoftObjectPtr<URigVMFunctionReferenceNode>& FunctionReferencePtr : FunctionReferencesPtr->FunctionReferences)
+				TArray< TSoftObjectPtr<URigVMFunctionReferenceNode> > FunctionReferences = FunctionReferencesPtr->FunctionReferences;
+				for (const TSoftObjectPtr<URigVMFunctionReferenceNode>& FunctionReferencePtr : FunctionReferences)
 				{
+					if (!FunctionReferencesPtr->FunctionReferences.Contains(FunctionReferencePtr))
+					{
+						continue;
+					}
+					
 					if (FunctionReferencePtr.IsValid())
 					{
 						{
+							FRigVMControllerGraphGuard GraphGuard(this, FunctionReferencePtr->GetGraph(), bSetupUndoRedo);
+							RemoveNode(FunctionReferencePtr.Get());
+
 							TGuardValue<TSoftObjectPtr<URigVMLibraryNode>> ClearReferencedNodePtr(
 								FunctionReferencePtr->ReferencedNodePtr,
 								TSoftObjectPtr<URigVMLibraryNode>());
-
-							FRigVMControllerGraphGuard GraphGuard(this, FunctionReferencePtr->GetGraph(), false);
-							RemoveNode(FunctionReferencePtr.Get());
 						}
 						FunctionReferencePtr->ReferencedNodePtr.ResetWeakPtr();
 					}
@@ -4789,6 +4749,57 @@
 		}
 	}
 
+	if (bSetupUndoRedo || bRecursive)
+	{
+		SelectNode(InNode, false, bSetupUndoRedo);
+
+		for (URigVMPin* Pin : InNode->GetPins())
+		{
+			TArray<URigVMInjectionInfo*> InjectedNodes = Pin->GetInjectedNodes();
+			for (URigVMInjectionInfo* InjectedNode : InjectedNodes)
+			{
+				RemoveNode(InjectedNode->UnitNode, bSetupUndoRedo, bRecursive);
+			}
+
+			BreakAllLinks(Pin, true, bSetupUndoRedo);
+			BreakAllLinks(Pin, false, bSetupUndoRedo);
+			BreakAllLinksRecursive(Pin, true, false, bSetupUndoRedo);
+			BreakAllLinksRecursive(Pin, false, false, bSetupUndoRedo);
+		}
+
+		if (URigVMCollapseNode* CollapseNode = Cast<URigVMCollapseNode>(InNode))
+		{
+			URigVMGraph* SubGraph = CollapseNode->GetContainedGraph();
+			FRigVMControllerGraphGuard GraphGuard(this, SubGraph, bSetupUndoRedo);
+		
+			TArray<URigVMNode*> ContainedNodes = SubGraph->GetNodes();
+			for (URigVMNode* ContainedNode : ContainedNodes)
+			{
+				if(Cast<URigVMFunctionEntryNode>(ContainedNode) != nullptr ||
+					Cast<URigVMFunctionReturnNode>(ContainedNode) != nullptr)
+				{
+					continue;
+				}
+				RemoveNode(ContainedNode, bSetupUndoRedo, bRecursive);
+			}
+		}
+		
+		if (bSetupUndoRedo)
+		{
+			ActionStack->AddAction(FRigVMRemoveNodeAction(InNode, this));
+		}
+	}
+
+	Graph->Nodes.Remove(InNode);
+	if (!bSuspendNotifications)
+	{
+		Graph->MarkPackageDirty();
+	}
+
+	Notify(ERigVMGraphNotifType::NodeRemoved, InNode);
+
+	
+
 	if (bPrintPythonCommand)
 	{
 		const FString GraphName = GetSanitizedGraphName(GetGraph()->GetGraphName());
