Control Rig: Name clashing between local/global variables handled inconsistently

diff --git a/Engine/Plugins/Animation/ControlRig/Source/ControlRigDeveloper/Private/Graph/ControlRigGraphSchema.cpp b/Engine/Plugins/Animation/ControlRig/Source/ControlRigDeveloper/Private/Graph/ControlRigGraphSchema.cpp
--- a/Engine/Plugins/Animation/ControlRig/Source/ControlRigDeveloper/Private/Graph/ControlRigGraphSchema.cpp
+++ b/Engine/Plugins/Animation/ControlRig/Source/ControlRigDeveloper/Private/Graph/ControlRigGraphSchema.cpp
@@ -37,6 +37,34 @@
 
 const FName UControlRigGraphSchema::GraphName_ControlRig(TEXT("Rig Graph"));
 
+FControlRigLocalVariableNameValidator::FControlRigLocalVariableNameValidator(const UBlueprint* Blueprint, const URigVMGraph* Graph, FName InExistingName)
+	: FStringSetNameValidator(InExistingName.ToString())
+{
+	if (Blueprint)
+	{
+		TSet<FName> NamesTemp;
+		// We allow local variables with same name as blueprint variable
+		
+		FBlueprintEditorUtils::GetFunctionNameList(Blueprint, NamesTemp);
+		FBlueprintEditorUtils::GetAllGraphNames(Blueprint, NamesTemp);
+		FBlueprintEditorUtils::GetSCSVariableNameList(Blueprint, NamesTemp);
+		FBlueprintEditorUtils::GetImplementingBlueprintsFunctionNameList(Blueprint, NamesTemp);
+
+		for (FName & Name : NamesTemp)
+		{
+			Names.Add(Name.ToString());
+		}
+	}
+
+	if (Graph)
+	{
+		for (const FRigVMGraphVariableDescription& LocalVariable : Graph->GetLocalVariables())
+		{
+			Names.Add(LocalVariable.Name.ToString());
+		}
+	}
+}
+
 FEdGraphPinType FControlRigGraphSchemaAction_LocalVar::GetPinType() const
 {
 	if (UControlRigGraph* Graph = Cast<UControlRigGraph>(GetVariableScope()))
@@ -67,6 +95,10 @@ void FControlRigGraphSchemaAction_LocalVar::ChangeVariableType(const FEdGraphPin
 void FControlRigGraphSchemaAction_LocalVar::RenameVariable(const FName& NewName)
 {
 	const FName OldName = GetVariableName();
+	if (OldName == NewName)
+	{
+		return;
+	}
 	
 	if (UControlRigGraph* Graph = Cast<UControlRigGraph>(GetVariableScope()))
 	{
@@ -81,22 +113,13 @@ bool FControlRigGraphSchemaAction_LocalVar::IsValidName(const FName& NewName, FT
 {
 	if (UControlRigGraph* ControlRigGraph = Cast<UControlRigGraph>(GetVariableScope()))
 	{
-		FKismetNameValidator NameValidator(ControlRigGraph->GetBlueprint());
-		EValidatorResult Result = NameValidator.IsValid(NewName.ToString());
-		if (Result != EValidatorResult::Ok)
+		FControlRigLocalVariableNameValidator NameValidator(ControlRigGraph->GetBlueprint(), ControlRigGraph->GetModel(), GetVariableName());
+		EValidatorResult Result = NameValidator.IsValid(NewName.ToString(), false);
+		if (Result != EValidatorResult::Ok && Result != EValidatorResult::ExistingName)
 		{
 			OutErrorMessage = FText::FromString(TEXT("Name with invalid format"));
 			return false;
 		}
-		
-		for (const FRigVMGraphVariableDescription& LocalVar : ControlRigGraph->GetModel()->GetLocalVariables())
-		{
-			if (LocalVar.Name != GetVariableName() && LocalVar.Name == NewName)
-			{
-				OutErrorMessage = FText::FromString(TEXT("Name already in use"));
-				return false;
-			}
-		}
 	}
 	return true;
 }
diff --git a/Engine/Plugins/Animation/ControlRig/Source/ControlRigDeveloper/Public/Graph/ControlRigGraphSchema.h b/Engine/Plugins/Animation/ControlRig/Source/ControlRigDeveloper/Public/Graph/ControlRigGraphSchema.h
--- a/Engine/Plugins/Animation/ControlRig/Source/ControlRigDeveloper/Public/Graph/ControlRigGraphSchema.h
+++ b/Engine/Plugins/Animation/ControlRig/Source/ControlRigDeveloper/Public/Graph/ControlRigGraphSchema.h
@@ -6,6 +6,7 @@
 #include "GraphEditorDragDropAction.h"
 #include "RigVMModel/RigVMGraph.h"
 #include "EdGraphSchema_K2_Actions.h"
+#include "Kismet2/Kismet2NameValidators.h"
 
 #include "ControlRigGraphSchema.generated.h"
 
@@ -43,6 +44,15 @@ struct FControlRigPinConnectionResponse
 	ECanCreateConnectionResponse_Extended ExtendedResponse;
 };
 
+/////////////////////////////////////////////////////
+// FControlRigLocalVariableNameValidator
+class CONTROLRIGDEVELOPER_API FControlRigLocalVariableNameValidator : public FStringSetNameValidator
+{
+
+public:
+	FControlRigLocalVariableNameValidator(const class UBlueprint* Blueprint, const URigVMGraph* Graph, FName InExistingName = NAME_None);
+};
+
 USTRUCT()
 struct CONTROLRIGDEVELOPER_API FControlRigGraphSchemaAction_LocalVar : public FEdGraphSchemaAction_BlueprintVariableBase
 {
diff --git a/Engine/Plugins/Animation/ControlRig/Source/ControlRigEditor/Private/ControlRigLocalVariableDetails.cpp b/Engine/Plugins/Animation/ControlRig/Source/ControlRigEditor/Private/ControlRigLocalVariableDetails.cpp
--- a/Engine/Plugins/Animation/ControlRig/Source/ControlRigEditor/Private/ControlRigLocalVariableDetails.cpp
+++ b/Engine/Plugins/Animation/ControlRig/Source/ControlRigEditor/Private/ControlRigLocalVariableDetails.cpp
@@ -32,6 +32,8 @@ void FRigVMLocalVariableDetails::CustomizeHeader(TSharedRef<IPropertyHandle> InS
 		VariableDescription = *ObjectsBeingCustomized[0]->GetContent<FRigVMGraphVariableDescription>();
 		GraphBeingCustomized = ObjectsBeingCustomized[0]->GetTypedOuter<URigVMGraph>();
 		BlueprintBeingCustomized = GraphBeingCustomized->GetTypedOuter<UControlRigBlueprint>();
+
+		NameValidator = FControlRigLocalVariableNameValidator(BlueprintBeingCustomized, GraphBeingCustomized, VariableDescription.Name);
 	}
 }
 
@@ -58,7 +60,11 @@ void FRigVMLocalVariableDetails::CustomizeChildren(TSharedRef<IPropertyHandle> I
 	.ValueContent()
 	.MaxDesiredWidth(250.0f)
 	[
-		NameHandle->CreatePropertyValueWidget()
+		SNew(SEditableTextBox)
+		.Font(IDetailLayoutBuilder::GetDetailFont())
+		.Text(this, &FRigVMLocalVariableDetails::GetName)
+		.OnTextCommitted(this, &FRigVMLocalVariableDetails::SetName)
+		.OnVerifyTextChanged(this, &FRigVMLocalVariableDetails::OnVerifyNameChanged)
 	];
 
 	TSharedPtr<IPinTypeSelectorFilter> CustomPinTypeFilter;
@@ -119,6 +125,35 @@ void FRigVMLocalVariableDetails::CustomizeChildren(TSharedRef<IPropertyHandle> I
 	}
 }
 
+FText FRigVMLocalVariableDetails::GetName() const
+{
+	return FText::FromName(VariableDescription.Name);
+}
+
+void FRigVMLocalVariableDetails::SetName(const FText& InNewText, ETextCommit::Type InCommitType)
+{
+	if(InCommitType == ETextCommit::OnCleared)
+	{
+		return;
+	}
+
+	if (InNewText.ToString() == VariableDescription.Name.ToString())
+	{
+		return;
+	}
+
+	VariableDescription.Name =  *InNewText.ToString();
+	NameHandle->SetValue(VariableDescription.Name);
+}
+
+bool FRigVMLocalVariableDetails::OnVerifyNameChanged(const FText& InText, FText& OutErrorMessage)
+{
+	EValidatorResult Result = NameValidator.IsValid(InText.ToString(), false);
+	OutErrorMessage = INameValidatorInterface::GetErrorText(InText.ToString(), Result);	
+
+	return Result == EValidatorResult::Ok || Result == EValidatorResult::ExistingName;
+}
+
 FEdGraphPinType FRigVMLocalVariableDetails::OnGetPinInfo() const
 {
 	if (!VariableDescription.Name.IsNone())
diff --git a/Engine/Plugins/Animation/ControlRig/Source/ControlRigEditor/Private/ControlRigLocalVariableDetails.h b/Engine/Plugins/Animation/ControlRig/Source/ControlRigEditor/Private/ControlRigLocalVariableDetails.h
--- a/Engine/Plugins/Animation/ControlRig/Source/ControlRigEditor/Private/ControlRigLocalVariableDetails.h
+++ b/Engine/Plugins/Animation/ControlRig/Source/ControlRigEditor/Private/ControlRigLocalVariableDetails.h
@@ -13,6 +13,7 @@
 #include "Styling/SlateTypes.h"
 #include "IPropertyUtilities.h"
 #include "DetailsViewWrapperObject.h"
+#include "Graph/ControlRigGraphSchema.h"
 
 class IPropertyHandle;
 
@@ -21,6 +22,7 @@ class FRigVMLocalVariableDetails : public IPropertyTypeCustomization
 	FRigVMLocalVariableDetails()
 	: GraphBeingCustomized(nullptr)
 	, BlueprintBeingCustomized(nullptr)
+	, NameValidator(nullptr, nullptr, NAME_None)
 	{}
 
 	
@@ -49,7 +51,12 @@ private:
 	TSharedPtr<IPropertyHandle> TypeObjectHandle;
 	TSharedPtr<IPropertyHandle> DefaultValueHandle;
 
+	FControlRigLocalVariableNameValidator NameValidator;
 	TArray<TSharedPtr<FString>> EnumOptions;
+
+	FText GetName() const;
+	void SetName(const FText& InNewText, ETextCommit::Type InCommitType);
+	bool OnVerifyNameChanged(const FText& InText, FText& OutErrorMessage);
 	
 	FEdGraphPinType OnGetPinInfo() const;
 	void HandlePinInfoChanged(const FEdGraphPinType& PinType);
diff --git a/Engine/Source/Developer/RigVMDeveloper/Private/RigVMCompiler/RigVMCompiler.cpp b/Engine/Source/Developer/RigVMDeveloper/Private/RigVMCompiler/RigVMCompiler.cpp
--- a/Engine/Source/Developer/RigVMDeveloper/Private/RigVMCompiler/RigVMCompiler.cpp
+++ b/Engine/Source/Developer/RigVMDeveloper/Private/RigVMCompiler/RigVMCompiler.cpp
@@ -273,6 +273,14 @@ bool URigVMCompiler::Compile(URigVMGraph* InGraph, URigVMController* InControlle
 				FString Path = FString::Printf(TEXT("LocalVariableDefault::%s|%s::Const"), *LibraryNode->GetFName().ToString(), *Variable.Name.ToString());
 				FRigVMOperand Operand = WorkData.AddProperty(ERigVMMemoryType::Literal, *Path, Variable.CPPType, Variable.CPPTypeObject, Variable.DefaultValue);
 				WorkData.PinPathToOperand->Add(Path, Operand);
+
+				for (const FRigVMExternalVariable& ExternalVariable : InExternalVariables)
+				{
+					if (ExternalVariable.Name == Variable.Name)
+					{
+						ReportWarningf(TEXT("Blueprint variable %s is being shadowed by a local variable in function %s"), *ExternalVariable.Name.ToString(), *LibraryNode->GetName());
+					}
+				}
 			}
 		}
 	}
diff --git a/Engine/Source/Developer/RigVMDeveloper/Private/RigVMModel/RigVMController.cpp b/Engine/Source/Developer/RigVMDeveloper/Private/RigVMModel/RigVMController.cpp
--- a/Engine/Source/Developer/RigVMDeveloper/Private/RigVMModel/RigVMController.cpp
+++ b/Engine/Source/Developer/RigVMDeveloper/Private/RigVMModel/RigVMController.cpp
@@ -1191,6 +1191,14 @@ void URigVMController::OnExternalVariableRemoved(const FName& InVarName, bool bS
 	URigVMGraph* Graph = GetGraph();
 	check(Graph);
 
+	for (const FRigVMGraphVariableDescription LocalVariable : Graph->GetLocalVariables())
+	{
+		if (InVarName == LocalVariable.Name)
+		{
+			return;
+		}
+	}
+	
 	const FString VarNameStr = InVarName.ToString();
 
 	if (bSetupUndoRedo)
@@ -1278,6 +1286,14 @@ void URigVMController::OnExternalVariableRenamed(const FName& InOldVarName, cons
 	URigVMGraph* Graph = GetGraph();
 	check(Graph);
 
+	for (const FRigVMGraphVariableDescription LocalVariable : Graph->GetLocalVariables())
+	{
+		if (InOldVarName == LocalVariable.Name)
+		{
+			return;
+		}
+	}
+
 	const FString VarNameStr = InOldVarName.ToString();
 
 	if (bSetupUndoRedo)
@@ -1368,6 +1384,14 @@ void URigVMController::OnExternalVariableTypeChanged(const FName& InVarName, con
 	URigVMGraph* Graph = GetGraph();
 	check(Graph);
 
+	for (const FRigVMGraphVariableDescription LocalVariable : Graph->GetLocalVariables())
+	{
+		if (InVarName == LocalVariable.Name)
+		{
+			return;
+		}
+	}
+
 	const FString VarNameStr = InVarName.ToString();
 
 	if (bSetupUndoRedo)
@@ -5580,6 +5604,15 @@ bool URigVMController::RenameVariable(const FName& InOldName, const FName& InNew
 		}
 	}
 
+	// If there is a local variable with the old name, a rename of the blueprint member variable does not affect this graph
+	for (FRigVMGraphVariableDescription& LocalVariable : Graph->GetLocalVariables())
+	{
+		if (LocalVariable.Name == InOldName)
+		{
+			return false;
+		}
+	}
+
 	FRigVMRenameVariableAction Action;
 	if (bSetupUndoRedo)
 	{
@@ -6583,8 +6616,8 @@ bool URigVMController::BindPinToVariable(URigVMPin* InPin, const FString& InNewB
 		FString VariableName = InNewBoundVariablePath, SegmentPath;
 		InNewBoundVariablePath.Split(TEXT("."), &VariableName, &SegmentPath);
 
-		FRigVMExternalVariable ExternalVariable = GetVariableByName(*VariableName);
-		if (ExternalVariable.IsValid(true))
+		FRigVMExternalVariable Variable = GetVariableByName(*VariableName);
+		if (Variable.IsValid(true))
 		{
 #if UE_RIGVM_UCLASS_BASED_STORAGE_DISABLED
 			FRigVMRegisterOffset RegisterOffset;
@@ -6597,7 +6630,7 @@ bool URigVMController::BindPinToVariable(URigVMPin* InPin, const FString& InNewB
 				return false;
 			}
 #else
-			if (!InPin->CanBeBoundToVariable(ExternalVariable, SegmentPath))
+			if (!InPin->CanBeBoundToVariable(Variable, SegmentPath))
 			{
 				return false;
 			}
@@ -6700,8 +6733,8 @@ bool URigVMController::MakeBindingsFromVariableNode(URigVMVariableNode* InNode,
 	}
 
 	FName VariableName = InNode->GetVariableName();
-	FRigVMExternalVariable ExternalVariable = GetVariableByName(VariableName);
-	if (!ExternalVariable.IsValid(true /* allow nullptr */))
+	FRigVMExternalVariable Variable = GetVariableByName(VariableName);
+	if (!Variable.IsValid(true /* allow nullptr */))
 	{
 		return false;
 	}
@@ -8171,6 +8204,17 @@ FRigVMGraphVariableDescription URigVMController::AddLocalVariable(const FName& I
 
 	Graph->LocalVariables.Add(NewVariable);
 
+	for (URigVMNode* Node : Graph->GetNodes())
+	{
+		if (URigVMVariableNode* VariableNode = Cast<URigVMVariableNode>(Node))
+		{
+			if (VariableName == VariableNode->GetVariableName())
+			{
+				RefreshVariableNode(VariableNode->GetFName(), VariableName, InCPPType, InCPPTypeObject, bSetupUndoRedo);
+			}
+		}
+	}
+
 	if (bSetupUndoRedo)
 	{
 		FRigVMInverseAction InverseAction;
@@ -8247,49 +8291,109 @@ bool URigVMController::RemoveLocalVariable(const FName& InVariableName, bool bSe
 
 	if (FoundIndex != INDEX_NONE)
 	{
+		FRigVMInverseAction InverseAction;
+		if (bSetupUndoRedo)
+		{
+			InverseAction.Title = FString::Printf(TEXT("Remove Local Variable %s"), *InVariableName.ToString());
+			ActionStack->BeginAction(InverseAction);			
+		}	
+		
 		const FString VarNameStr = InVariableName.ToString();
-		TArray<URigVMNode*> Nodes = Graph->GetNodes();
-		for (URigVMNode* Node : Nodes)
+
+		bool bSwitchToMemberVariable = false;
+		FRigVMExternalVariable ExternalVariableToSwitch;
+		{
+			TArray<FRigVMExternalVariable> ExternalVariables;
+			if (GetExternalVariablesDelegate.IsBound())
+			{
+				ExternalVariables.Append(GetExternalVariablesDelegate.Execute(GetGraph()));
+			}
+
+			for (const FRigVMExternalVariable& ExternalVariable : ExternalVariables)
+			{
+				if (ExternalVariable.Name == InVariableName)
+				{
+					bSwitchToMemberVariable = true;
+					ExternalVariableToSwitch = ExternalVariable;
+					break;
+				}	
+			}
+		}
+
+		if (!bSwitchToMemberVariable)
 		{
-			if (URigVMVariableNode* VariableNode = Cast<URigVMVariableNode>(Node))
+			TArray<URigVMNode*> Nodes = Graph->GetNodes();
+			for (URigVMNode* Node : Nodes)
 			{
-				if (URigVMPin* VariablePin = VariableNode->FindPin(URigVMVariableNode::VariableName))
+				if (URigVMVariableNode* VariableNode = Cast<URigVMVariableNode>(Node))
 				{
-					if (VariablePin->GetDefaultValue() == VarNameStr)
+					if (URigVMPin* VariablePin = VariableNode->FindPin(URigVMVariableNode::VariableName))
 					{
-						RemoveNode(Node, bSetupUndoRedo, true);
-						continue;
+						if (VariablePin->GetDefaultValue() == VarNameStr)
+						{
+							RemoveNode(Node, bSetupUndoRedo, true);
+							continue;
+						}
 					}
 				}
-			}
 
-			TArray<URigVMPin*> AllPins = Node->GetAllPinsRecursively();
-			for (URigVMPin* Pin : AllPins)
-			{
-				if (Pin->GetBoundVariableName() == InVariableName.ToString())
+				TArray<URigVMPin*> AllPins = Node->GetAllPinsRecursively();
+				for (URigVMPin* Pin : AllPins)
 				{
-					BindPinToVariable(Pin, FString(), bSetupUndoRedo);
+					if (Pin->GetBoundVariableName() == InVariableName.ToString())
+					{
+						BindPinToVariable(Pin, FString(), bSetupUndoRedo);
+					}
 				}
 			}
 		}
+		else
+		{
+			TArray<URigVMNode*> Nodes = Graph->GetNodes();
+			for (URigVMNode* Node : Nodes)
+			{
+				if (URigVMVariableNode* VariableNode = Cast<URigVMVariableNode>(Node))
+				{
+					if (URigVMPin* VariablePin = VariableNode->FindPin(URigVMVariableNode::VariableName))
+					{
+						if (VariablePin->GetDefaultValue() == VarNameStr)
+						{
+							RefreshVariableNode(VariableNode->GetFName(), ExternalVariableToSwitch.Name, ExternalVariableToSwitch.TypeName.ToString(), ExternalVariableToSwitch.TypeObject, bSetupUndoRedo);
+							continue;
+						}
+					}
+				}
+
+				TArray<URigVMPin*> AllPins = Node->GetAllPinsRecursively();
+				for (URigVMPin* Pin : AllPins)
+				{
+					if (Pin->GetBoundVariableName() == InVariableName.ToString())
+					{
+						if (Pin->GetCPPType() != ExternalVariableToSwitch.TypeName.ToString() || Pin->GetCPPTypeObject() == ExternalVariableToSwitch.TypeObject)
+						{
+							BindPinToVariable(Pin, FString(), bSetupUndoRedo);
+						}
+					}
+				}
+			}		
+		}
 
 		if (!bSuspendNotifications)
 		{
 			Graph->MarkPackageDirty();
 		}
-		
+
 		if (bSetupUndoRedo)
 		{
-			FRigVMInverseAction InverseAction;
-			InverseAction.Title = FString::Printf(TEXT("Remove Local Variable %s"), *InVariableName.ToString());
-
-			ActionStack->BeginAction(InverseAction);
 			ActionStack->AddAction(FRigVMAddLocalVariableAction(LocalVariables[FoundIndex]));
-			ActionStack->EndAction(InverseAction);
-		}	
-		
+		}
 		LocalVariables.RemoveAt(FoundIndex);
 
+		if (bSetupUndoRedo)
+		{
+			ActionStack->EndAction(InverseAction);
+		}
+
 		if (bPrintPythonCommand)
 		{
 			const FString GraphName = GetSanitizedGraphName(GetGraph()->GetGraphName());
@@ -11903,25 +12007,14 @@ void URigVMController::ValidatePin(URigVMPin* InPin)
 
 FRigVMExternalVariable URigVMController::GetVariableByName(const FName& InExternalVariableName)
 {
-	TArray<FRigVMExternalVariable> ExternalVariables = GetAllVariables();
-	for (const FRigVMExternalVariable& ExternalVariable : ExternalVariables)
+	TArray<FRigVMExternalVariable> Variables = GetAllVariables();
+	for (const FRigVMExternalVariable& Variable : Variables)
 	{
-		if (ExternalVariable.Name == InExternalVariableName)
+		if (Variable.Name == InExternalVariableName)
 		{
-			return ExternalVariable;
+			return Variable;
 		}
-	}
-
-	if(URigVMGraph* Graph = GetGraph())
-	{
-		for (FRigVMGraphVariableDescription LocalVariable : Graph->GetLocalVariables())
-		{
-			if (LocalVariable.Name == InExternalVariableName)
-			{
-				return LocalVariable.ToExternalVariable();
-			}
-		}
-	}
+	}	
 	
 	return FRigVMExternalVariable();
 }
@@ -11929,11 +12022,6 @@ FRigVMExternalVariable URigVMController::GetVariableByName(const FName& InExtern
 TArray<FRigVMExternalVariable> URigVMController::GetAllVariables()
 {
 	TArray<FRigVMExternalVariable> ExternalVariables;
-	
-	if (GetExternalVariablesDelegate.IsBound())
-	{
-		ExternalVariables.Append(GetExternalVariablesDelegate.Execute(GetGraph()));
-	}
 
 	if(URigVMGraph* Graph = GetGraph())
 	{
@@ -11943,6 +12031,11 @@ TArray<FRigVMExternalVariable> URigVMController::GetAllVariables()
 		}
 	}
 	
+	if (GetExternalVariablesDelegate.IsBound())
+	{
+		ExternalVariables.Append(GetExternalVariablesDelegate.Execute(GetGraph()));
+	}
+
 	return ExternalVariables;
 }
 
diff --git a/Engine/Source/Developer/RigVMDeveloper/Public/RigVMModel/RigVMGraph.h b/Engine/Source/Developer/RigVMDeveloper/Public/RigVMModel/RigVMGraph.h
--- a/Engine/Source/Developer/RigVMDeveloper/Public/RigVMModel/RigVMGraph.h
+++ b/Engine/Source/Developer/RigVMDeveloper/Public/RigVMModel/RigVMGraph.h
@@ -121,7 +121,8 @@ public:
 
 	TArray<FRigVMExternalVariable> GetExternalVariables() const;
 
-	TArray<FRigVMGraphVariableDescription>& GetLocalVariables() { return LocalVariables; }
+	UFUNCTION(BlueprintCallable, Category = RigVMGraph)
+	TArray<FRigVMGraphVariableDescription> GetLocalVariables() const { return LocalVariables; }
 
 	// Returns the modified event, which can be used to 
 	// subscribe to changes happening within the Graph.
