Control Rig: Fix compiling errors in RigVMController

diff --git a/Engine/Source/Developer/RigVMDeveloper/Private/RigVMModel/RigVMController.cpp b/Engine/Source/Developer/RigVMDeveloper/Private/RigVMModel/RigVMController.cpp
--- a/Engine/Source/Developer/RigVMDeveloper/Private/RigVMModel/RigVMController.cpp
+++ b/Engine/Source/Developer/RigVMDeveloper/Private/RigVMModel/RigVMController.cpp
@@ -1191,7 +1191,7 @@ void URigVMController::OnExternalVariableRemoved(const FName& InVarName, bool bS
 	URigVMGraph* Graph = GetGraph();
 	check(Graph);
 
-	for (const FRigVMGraphVariableDescription LocalVariable : Graph->GetLocalVariables())
+	for (const FRigVMGraphVariableDescription& LocalVariable : Graph->GetLocalVariables())
 	{
 		if (InVarName == LocalVariable.Name)
 		{
@@ -1286,7 +1286,7 @@ void URigVMController::OnExternalVariableRenamed(const FName& InOldVarName, cons
 	URigVMGraph* Graph = GetGraph();
 	check(Graph);
 
-	for (const FRigVMGraphVariableDescription LocalVariable : Graph->GetLocalVariables())
+	for (const FRigVMGraphVariableDescription& LocalVariable : Graph->GetLocalVariables())
 	{
 		if (InOldVarName == LocalVariable.Name)
 		{
@@ -1384,7 +1384,7 @@ void URigVMController::OnExternalVariableTypeChanged(const FName& InVarName, con
 	URigVMGraph* Graph = GetGraph();
 	check(Graph);
 
-	for (const FRigVMGraphVariableDescription LocalVariable : Graph->GetLocalVariables())
+	for (const FRigVMGraphVariableDescription& LocalVariable : Graph->GetLocalVariables())
 	{
 		if (InVarName == LocalVariable.Name)
 		{
