Control Rig: Lookup resolved value during constant folding both from work + literal memory

diff --git a/Engine/Source/Developer/RigVMDeveloper/Private/RigVMCompiler/RigVMAST.cpp b/Engine/Source/Developer/RigVMDeveloper/Private/RigVMCompiler/RigVMAST.cpp
--- a/Engine/Source/Developer/RigVMDeveloper/Private/RigVMCompiler/RigVMAST.cpp
+++ b/Engine/Source/Developer/RigVMDeveloper/Private/RigVMCompiler/RigVMAST.cpp
@@ -1971,7 +1971,15 @@ bool FRigVMParserAST::FoldConstantValuesToLiterals(URigVMGraph* InGraph, URigVMC
 
 #else
 
-		FString DefaultValue = TempVM->GetWorkMemory()->GetDataAsString(Operand.GetRegisterIndex());
+		FString DefaultValue;
+		if(Operand.GetMemoryType() == ERigVMMemoryType::Work)
+		{
+			DefaultValue = TempVM->GetWorkMemory()->GetDataAsString(Operand.GetRegisterIndex());
+		}
+		else if(Operand.GetMemoryType() == ERigVMMemoryType::Literal)
+		{
+			DefaultValue = TempVM->GetLiteralMemory()->GetDataAsString(Operand.GetRegisterIndex());
+		}
 		if(DefaultValue.IsEmpty())
 		{
 			continue;
