Control Rig: Throw warnings for unconnected execute pins

diff --git a/Engine/Source/Developer/RigVMDeveloper/Private/RigVMCompiler/RigVMCompiler.cpp b/Engine/Source/Developer/RigVMDeveloper/Private/RigVMCompiler/RigVMCompiler.cpp
--- a/Engine/Source/Developer/RigVMDeveloper/Private/RigVMCompiler/RigVMCompiler.cpp
+++ b/Engine/Source/Developer/RigVMDeveloper/Private/RigVMCompiler/RigVMCompiler.cpp
@@ -167,8 +167,7 @@ bool URigVMCompiler::Compile(URigVMGraph* InGraph, URigVMController* InControlle
 	TArray<URigVMGraph*> VisitedGraphs;
 	VisitedGraphs.Add(InGraph);
 
-	int32 NodesWithOrphanPins = 0;
-	int32 NodesWithUnmappedVariables = 0;
+	bool bEncounteredGraphError = false;
 	for(int32 GraphIndex=0; GraphIndex<VisitedGraphs.Num(); GraphIndex++)
 	{
 		URigVMGraph* VisitedGraph = VisitedGraphs[GraphIndex];
@@ -179,7 +178,7 @@ bool URigVMCompiler::Compile(URigVMGraph* InGraph, URigVMController* InControlle
 			{
 				static const FString LinkedMessage = TEXT("Node @@ uses pins that no longer exist. Please rewire the links and re-compile.");
 				Settings.ASTSettings.Report(EMessageSeverity::Error, ModelNode, LinkedMessage);
-				NodesWithOrphanPins++;
+				bEncounteredGraphError = true;
 			}
 
 			if(URigVMFunctionReferenceNode* FunctionReferenceNode = Cast<URigVMFunctionReferenceNode>(ModelNode))
@@ -188,19 +187,22 @@ bool URigVMCompiler::Compile(URigVMGraph* InGraph, URigVMController* InControlle
 				{
 					static const FString UnmappedMessage = TEXT("Node @@ has unmapped variables. Please adjust the node and re-compile.");
 					Settings.ASTSettings.Report(EMessageSeverity::Error, ModelNode, UnmappedMessage);
-					NodesWithUnmappedVariables++;
+					bEncounteredGraphError = true;
 				}
 			}
 
 			if(ModelNode->IsA<URigVMFunctionEntryNode>() || ModelNode->IsA<URigVMFunctionReturnNode>())
 			{
-				static const FString ExecuteContextName = FRigVMStruct::ExecuteContextName.ToString();
-				if(URigVMPin* ExecutePin = ModelNode->FindPin(ExecuteContextName))
+				for(URigVMPin* ExecutePin : ModelNode->Pins)
 				{
-					if(ExecutePin->GetLinks().Num() == 0)
+					if(ExecutePin->IsExecuteContext())
 					{
-						static const FString UnlinkedExecuteMessage = TEXT("Node @@ has an unconnected Execute pin.\nThe function might cause unexpected behavior.");
-						Settings.ASTSettings.Report(EMessageSeverity::Warning, ModelNode, UnlinkedExecuteMessage);
+						if(ExecutePin->GetLinks().Num() == 0)
+						{
+							static const FString UnlinkedExecuteMessage = TEXT("Node @@ has an unconnected Execute pin.\nThe function might cause unexpected behavior.");
+							Settings.ASTSettings.Report(EMessageSeverity::Error, ModelNode, UnlinkedExecuteMessage);
+							bEncounteredGraphError = true;
+						}
 					}
 				}
 			}
@@ -223,7 +225,7 @@ bool URigVMCompiler::Compile(URigVMGraph* InGraph, URigVMController* InControlle
 		}
 	}
 
-	if((NodesWithOrphanPins + NodesWithUnmappedVariables) > 0)
+	if(bEncounteredGraphError)
 	{
 		return false;
 	}
