Control Rig: Fix expand node to work with direct links and local variables within functions

diff --git a/Engine/Source/Developer/RigVMDeveloper/Private/RigVMModel/RigVMController.cpp b/Engine/Source/Developer/RigVMDeveloper/Private/RigVMModel/RigVMController.cpp
--- a/Engine/Source/Developer/RigVMDeveloper/Private/RigVMModel/RigVMController.cpp
+++ b/Engine/Source/Developer/RigVMDeveloper/Private/RigVMModel/RigVMController.cpp
@@ -3714,8 +3714,10 @@ TArray<URigVMNode*> URigVMController::ExpandLibraryNode(URigVMLibraryNode* InNod
 	// exist, they will be reused. If a local variable is not used, it will not be created.
 	if (URigVMFunctionReferenceNode* FunctionReferenceNode = Cast<URigVMFunctionReferenceNode>(InNode))
 	{
+		const bool bIsWithinFunctionGraph = !Graph->IsRootGraph() && Graph->GetRootGraph()->IsA<URigVMFunctionLibrary>();
+		TArray<FRigVMExternalVariable> AllExistingVariables = GetAllVariables();
+
 		TArray<FRigVMGraphVariableDescription> LocalVariables = FunctionReferenceNode->GetContainedGraph()->LocalVariables;
-		TArray<FRigVMExternalVariable> CurrentVariables = GetAllVariables();
 		TArray<FRigVMGraphVariableDescription> VariablesToAdd;
 		for (const URigVMNode* Node : FunctionReferenceNode->GetContainedGraph()->GetNodes())
 		{
@@ -3728,13 +3730,13 @@ TArray<URigVMNode*> URigVMController::ExpandLibraryNode(URigVMLibraryNode* InNod
 						bool bVariableExists = false;
 						bool bVariableIncompatible = false;
 						FRigVMExternalVariable LocalVariableExternalType = LocalVariable.ToExternalVariable();
-						for (FRigVMExternalVariable& CurrentVariable : CurrentVariables)
+						for (FRigVMExternalVariable& ExistingVariable : AllExistingVariables)
 						{						
-							if (CurrentVariable.Name == LocalVariable.Name)
+							if (ExistingVariable.Name == LocalVariable.Name)
 							{
-								if (CurrentVariable.TypeName != LocalVariableExternalType.TypeName ||
-									CurrentVariable.TypeObject != LocalVariableExternalType.TypeObject ||
-									CurrentVariable.bIsArray != LocalVariableExternalType.bIsArray)
+								if (ExistingVariable.TypeName != LocalVariableExternalType.TypeName ||
+									ExistingVariable.TypeObject != LocalVariableExternalType.TypeObject ||
+									ExistingVariable.bIsArray != LocalVariableExternalType.bIsArray)
 								{
 									bVariableIncompatible = true;	
 								}
@@ -3749,7 +3751,7 @@ TArray<URigVMNode*> URigVMController::ExpandLibraryNode(URigVMLibraryNode* InNod
 						}
 						else if(bVariableIncompatible)
 						{
-							ReportErrorf(TEXT("Found variable %s of incompatible type with a local variable inside function %s"), *LocalVariable.Name.ToString(), *FunctionReferenceNode->GetReferencedNode()->GetName());
+							ReportAndNotifyErrorf(TEXT("Found variable %s of incompatible type with a local variable inside function %s"), *LocalVariable.Name.ToString(), *FunctionReferenceNode->GetReferencedNode()->GetName());
 							if (bSetupUndoRedo)
 							{
 								ActionStack->CancelAction(ExpandAction);
@@ -3762,11 +3764,21 @@ TArray<URigVMNode*> URigVMController::ExpandLibraryNode(URigVMLibraryNode* InNod
 			}
 		}
 
-		if (RequestNewExternalVariableDelegate.IsBound())
+		if(!bIsWithinFunctionGraph)
 		{
-			for (const FRigVMGraphVariableDescription& OldVariable : VariablesToAdd)
+			if (RequestNewExternalVariableDelegate.IsBound())
 			{
-				RequestNewExternalVariableDelegate.Execute(OldVariable, false, false);
+				for (const FRigVMGraphVariableDescription& VariableToAdd : VariablesToAdd)
+				{
+					RequestNewExternalVariableDelegate.Execute(VariableToAdd, false, false);
+				}
+			}
+		}
+		else
+		{
+			for (const FRigVMGraphVariableDescription& VariableToAdd : VariablesToAdd)
+			{
+				AddLocalVariable(VariableToAdd.Name, VariableToAdd.CPPType, VariableToAdd.CPPTypeObject, VariableToAdd.DefaultValue, bSetupUndoRedo, false);
 			}
 		}
 	}
@@ -4124,15 +4136,52 @@ TArray<URigVMNode*> URigVMController::ExpandLibraryNode(URigVMLibraryNode* InNod
 			RemappedSourcePinPath = TargetPin->GetPinPath();
 		}
 
+		
 		for (const FString& FromEntryNodeTargetPinPath : FromEntryNodePair.Value)
 		{
-			URigVMPin* SourcePin = GetGraph()->FindPin(*RemappedSourcePinPath);
+			TArray<URigVMPin*> TargetPins;
+
+			URigVMPin* SourcePin = GetGraph()->FindPin(RemappedSourcePinPath);
 			URigVMPin* TargetPin = GetGraph()->FindPin(FromEntryNodeTargetPinPath);
-			if (SourcePin && TargetPin)
+
+			// potentially the target pin was on the entry node,
+			// so there's no node been added for it. we'll have to look into the remapped
+			// pins for the "FromLibraryNode" map.
+			if(TargetPin == nullptr)
 			{
-				if (!SourcePin->IsLinkedTo(TargetPin))
+				FString RemappedTargetPinPath = FromEntryNodeTargetPinPath;
+				FString ReturnNodeName, ReturnPinPath;
+				if (URigVMPin::SplitPinPathAtStart(RemappedTargetPinPath, ReturnNodeName, ReturnPinPath))
 				{
-					AddLink(SourcePin, TargetPin, false);
+					if(Cast<URigVMFunctionReturnNode>(InNode->GetContainedGraph()->FindNode(ReturnNodeName)))
+					{
+						if(FromLibraryNode.Contains(ReturnPinPath))
+						{
+							const TArray<FString>& FromLibraryNodeTargetPins = FromLibraryNode.FindChecked(ReturnPinPath);
+							for(const FString& FromLibraryNodeTargetPin : FromLibraryNodeTargetPins)
+							{
+								if(URigVMPin* MappedTargetPin = GetGraph()->FindPin(FromLibraryNodeTargetPin))
+								{
+									TargetPins.Add(MappedTargetPin);
+								}
+							}
+						}
+					}
+				}
+			}
+			else
+			{
+				TargetPins.Add(TargetPin);
+			}
+			
+			if (SourcePin)
+			{
+				for(URigVMPin* EachTargetPin : TargetPins)
+				{
+					if (!SourcePin->IsLinkedTo(EachTargetPin))
+					{
+						AddLink(SourcePin, EachTargetPin, false);
+					}
 				}
 			}
 		}
