Control Rig: Fix up invalid graphs on load, guard against infinite recusion in AST

diff --git a/Engine/Source/Developer/RigVMDeveloper/Private/RigVMCompiler/RigVMAST.cpp b/Engine/Source/Developer/RigVMDeveloper/Private/RigVMCompiler/RigVMAST.cpp
--- a/Engine/Source/Developer/RigVMDeveloper/Private/RigVMCompiler/RigVMAST.cpp
+++ b/Engine/Source/Developer/RigVMDeveloper/Private/RigVMCompiler/RigVMAST.cpp
@@ -2937,13 +2937,16 @@ void FRigVMParserAST::Inline(URigVMGraph* InGraph, const TArray<FRigVMASTProxy>&
 					}
 					else if(URigVMLibraryNode* LibraryNode = Cast<URigVMLibraryNode>(ChildPin->GetNode()))
 					{
-						const FRigVMASTProxy ChildPinProxy = InPinProxy.GetSibling(ChildPin); 
-						if (ShouldRecursePin(ChildPinProxy))
+						if(ChildPin != InPinProxy.GetSubject<URigVMPin>())
 						{
-							FRigVMASTProxy SourceSourcePinProxy = FindSourcePin(ChildPinProxy, OutTraversalInfo);
-							if(SourceSourcePinProxy.IsValid())
+							const FRigVMASTProxy ChildPinProxy = InPinProxy.GetSibling(ChildPin); 
+							if (ShouldRecursePin(ChildPinProxy))
 							{
-								SourcePinProxy = SourceSourcePinProxy;
+								FRigVMASTProxy SourceSourcePinProxy = FindSourcePin(ChildPinProxy, OutTraversalInfo);
+								if(SourceSourcePinProxy.IsValid())
+								{
+									SourcePinProxy = SourceSourcePinProxy;
+								}
 							}
 						}
 					}
diff --git a/Engine/Source/Developer/RigVMDeveloper/Private/RigVMModel/RigVMController.cpp b/Engine/Source/Developer/RigVMDeveloper/Private/RigVMModel/RigVMController.cpp
--- a/Engine/Source/Developer/RigVMDeveloper/Private/RigVMModel/RigVMController.cpp
+++ b/Engine/Source/Developer/RigVMDeveloper/Private/RigVMModel/RigVMController.cpp
@@ -11223,6 +11223,7 @@ void URigVMController::RemovePinsDuringRepopulate(URigVMNode* InNode, TArray<URi
 			{
 				Pin->Rename(nullptr, OrphanedRootPin, REN_ForceNoResetLoaders | REN_DoNotDirty | REN_DontCreateRedirectors | REN_NonTransactional);
 				RootPin->SubPins.Remove(Pin);
+				EnsurePinValidity(Pin, false);
 				AddSubPin(OrphanedRootPin, Pin);
 			}
 		}
