Control Rig: Local variable deletion in the Control Rig cannot be undone

diff --git a/Engine/Source/Developer/RigVMDeveloper/Private/RigVMModel/RigVMControllerActions.cpp b/Engine/Source/Developer/RigVMDeveloper/Private/RigVMModel/RigVMControllerActions.cpp
--- a/Engine/Source/Developer/RigVMDeveloper/Private/RigVMModel/RigVMControllerActions.cpp
+++ b/Engine/Source/Developer/RigVMDeveloper/Private/RigVMModel/RigVMControllerActions.cpp
@@ -1817,14 +1817,14 @@ bool FRigVMAddLocalVariableAction::Undo(URigVMController* InController)
 	{
 		return false;
 	}
-	return InController->RemoveLocalVariable(LocalVariable.Name);
+	return InController->RemoveLocalVariable(LocalVariable.Name, false);
 }
 
 bool FRigVMAddLocalVariableAction::Redo(URigVMController* InController)
 {
 	if (!LocalVariable.Name.IsNone())
 	{
-		return InController->AddLocalVariable(LocalVariable.Name, LocalVariable.CPPType, LocalVariable.CPPTypeObject, LocalVariable.DefaultValue).Name.IsNone() == false;
+		return InController->AddLocalVariable(LocalVariable.Name, LocalVariable.CPPType, LocalVariable.CPPTypeObject, LocalVariable.DefaultValue, false).Name.IsNone() == false;
 	}
 	return FRigVMBaseAction::Redo(InController);
 }
