Control Rig: Refactor watch value

diff --git a/Engine/Source/Developer/RigVMDeveloper/Private/RigVMCompiler/RigVMAST.cpp b/Engine/Source/Developer/RigVMDeveloper/Private/RigVMCompiler/RigVMAST.cpp
--- a/Engine/Source/Developer/RigVMDeveloper/Private/RigVMCompiler/RigVMAST.cpp
+++ b/Engine/Source/Developer/RigVMDeveloper/Private/RigVMCompiler/RigVMAST.cpp
@@ -2803,7 +2803,7 @@ void FRigVMParserAST::Inline(URigVMGraph* InGraph, const TArray<FRigVMASTProxy>&
 	struct LocalPinTraversalInfo
 	{
 		URigVMPin::FPinOverrideMap* PinOverrides;
-		TMap<FRigVMASTProxy, FRigVMASTProxy> SourcePins;
+		TMap<FRigVMASTProxy, FRigVMASTProxy>* SourcePins;
 		TMap<FRigVMASTProxy, TArray<FRigVMASTProxy>>* TargetLinks;
 		TMap<FRigVMASTProxy, TArray<FRigVMASTProxy>>* SourceLinks;
 		TArray<FRigVMASTProxy> LibraryNodeCallstack;
@@ -2853,7 +2853,7 @@ void FRigVMParserAST::Inline(URigVMGraph* InGraph, const TArray<FRigVMASTProxy>&
 									FRigVMASTProxy ReturnPinProxy = LibraryNodeProxy.GetChild(ReturnPin);
 									FRigVMASTProxy SourcePinProxy = FindSourcePin(ReturnPinProxy, OutTraversalInfo);
 									SourcePinProxy = SourcePinProxy.IsValid() ? SourcePinProxy : ReturnPinProxy;
-									OutTraversalInfo.SourcePins.FindOrAdd(InPinProxy) = SourcePinProxy;
+									OutTraversalInfo.SourcePins->FindOrAdd(InPinProxy) = SourcePinProxy;
 									OutTraversalInfo.LibraryNodeCallstack.Pop();
 									return SourcePinProxy;
 
@@ -2879,7 +2879,7 @@ void FRigVMParserAST::Inline(URigVMGraph* InGraph, const TArray<FRigVMASTProxy>&
 									FRigVMASTProxy LibraryPinProxy = LibraryNodeProxy.GetSibling(LibraryPin);
 									FRigVMASTProxy SourcePinProxy = FindSourcePin(LibraryPinProxy, OutTraversalInfo);
 									SourcePinProxy = SourcePinProxy.IsValid() ? SourcePinProxy : LibraryPinProxy;
-									OutTraversalInfo.SourcePins.FindOrAdd(InPinProxy) = SourcePinProxy;
+									OutTraversalInfo.SourcePins->FindOrAdd(InPinProxy) = SourcePinProxy;
 									return SourcePinProxy;
 								}
 							}
@@ -2908,7 +2908,7 @@ void FRigVMParserAST::Inline(URigVMGraph* InGraph, const TArray<FRigVMASTProxy>&
 			{
 				// note: this map isn't going to work for functions which are referenced.
 				// (since the pin objects are shared between multiple invocation nodes)
-				if (const FRigVMASTProxy* SourcePinProxy = OutTraversalInfo.SourcePins.Find(InPinProxy))
+				if (const FRigVMASTProxy* SourcePinProxy = OutTraversalInfo.SourcePins->Find(InPinProxy))
 				{
 					return *SourcePinProxy;
 				}
@@ -2977,7 +2977,7 @@ void FRigVMParserAST::Inline(URigVMGraph* InGraph, const TArray<FRigVMASTProxy>&
 					FRigVMASTProxy ParentPinProxy = InPinProxy.GetSibling(ParentPin);
 
 					// if we found a parent pin which has a source that is not a reroute
-					if (FRigVMASTProxy* ParentSourcePinProxyPtr = OutTraversalInfo.SourcePins.Find(ParentPinProxy))
+					if (FRigVMASTProxy* ParentSourcePinProxyPtr = OutTraversalInfo.SourcePins->Find(ParentPinProxy))
 					{
 						FRigVMASTProxy& ParentSourcePinProxy = *ParentSourcePinProxyPtr;
 						if (ParentSourcePinProxy.IsValid())
@@ -3026,7 +3026,7 @@ void FRigVMParserAST::Inline(URigVMGraph* InGraph, const TArray<FRigVMASTProxy>&
 
 			if (!bIOPinOnLeftOfLibraryNode)
 			{
-				OutTraversalInfo.SourcePins.FindOrAdd(InPinProxy) = SourcePinProxy;
+				OutTraversalInfo.SourcePins->FindOrAdd(InPinProxy) = SourcePinProxy;
 			}
 			return SourcePinProxy;
 		}
@@ -3139,6 +3139,7 @@ void FRigVMParserAST::Inline(URigVMGraph* InGraph, const TArray<FRigVMASTProxy>&
 	//    also traverse links along reroutes and flatten them
 	LocalPinTraversalInfo TraversalInfo;
 	TraversalInfo.PinOverrides = &PinOverrides;
+	TraversalInfo.SourcePins = &SharedOperandPins;
 	TraversalInfo.TargetLinks = &TargetLinks;
 	TraversalInfo.SourceLinks = &SourceLinks;
 	TraversalInfo.Settings = &Settings;
diff --git a/Engine/Source/Developer/RigVMDeveloper/Private/RigVMCompiler/RigVMCompiler.cpp b/Engine/Source/Developer/RigVMDeveloper/Private/RigVMCompiler/RigVMCompiler.cpp
--- a/Engine/Source/Developer/RigVMDeveloper/Private/RigVMCompiler/RigVMCompiler.cpp
+++ b/Engine/Source/Developer/RigVMDeveloper/Private/RigVMCompiler/RigVMCompiler.cpp
@@ -2784,14 +2784,18 @@ FRigVMOperand URigVMCompiler::FindOrAddRegister(const FRigVMVarExprAST* InVarExp
 
 	// Get all possible pins that lead to the same operand
 
-	if(Settings.ASTSettings.bFoldAssignments && !bIsDebugValue)
+	if(Settings.ASTSettings.bFoldAssignments)
 	{
-		TArray<const URigVMPin*> VirtualPins;
-		Pin->GetExposedPinChain(VirtualPins);
-		for (const URigVMPin* VirtualPin : VirtualPins)
+		const TMap<FRigVMASTProxy, FRigVMASTProxy>& SharedOperandPins = InVarExpr->GetParser()->SharedOperandPins;
+		TArray<FRigVMASTProxy> PinProxies = FindProxiesWithSharedOperand(InVarExpr);
+
+		for (const FRigVMASTProxy& Proxy : PinProxies)
 		{
-			FString VirtualPinHash = GetPinHash(VirtualPin, InVarExpr, bIsDebugValue);
-			WorkData.PinPathToOperand->Add(VirtualPinHash, Operand);	
+			if (URigVMPin* VirtualPin = Cast<URigVMPin>(Proxy.GetSubject()))
+			{
+				FString VirtualPinHash = GetPinHash(VirtualPin, InVarExpr, bIsDebugValue);
+				WorkData.PinPathToOperand->Add(VirtualPinHash, Operand);
+			}	
 		}
 	}
 	else
@@ -2811,6 +2815,59 @@ FRigVMOperand URigVMCompiler::FindOrAddRegister(const FRigVMVarExprAST* InVarExp
 	return Operand;
 }
 
+TArray<FRigVMASTProxy> URigVMCompiler::FindProxiesWithSharedOperand(const FRigVMVarExprAST* InVarExpr)
+{
+	TArray<FRigVMASTProxy> PinProxies, PinProxiesToProcess;
+	const TMap<FRigVMASTProxy, FRigVMASTProxy>& SharedOperandPins = InVarExpr->GetParser()->SharedOperandPins;
+	PinProxiesToProcess.Add(InVarExpr->GetProxy());
+	while(!PinProxiesToProcess.IsEmpty())
+	{
+		FRigVMASTProxy Proxy = PinProxiesToProcess.Pop();
+
+		if (Proxy.IsValid())
+		{
+			if (URigVMPin* Pin = Cast<URigVMPin>(Proxy.GetSubject()))
+			{
+				if (Pin->GetNode()->IsA<URigVMVariableNode>() || Pin->GetNode()->IsA<URigVMParameterNode>())
+				{
+					if (Pin->GetDirection() == ERigVMPinDirection::Input)
+					{
+						continue;
+					}
+				}
+			}
+			
+			PinProxies.Add(Proxy);
+		}
+
+		for (TPair<FRigVMASTProxy, FRigVMASTProxy> Pair : SharedOperandPins)
+		{
+			FRigVMASTProxy* Other = nullptr;
+			if (PinProxies.Contains(Pair.Key))
+			{
+				Other = &Pair.Value;
+			}
+			else if (PinProxies.Contains(Pair.Value))
+			{
+				Other = &Pair.Key;
+			}
+
+			if (Other)
+			{
+				if (Other->IsValid())
+				{
+					if (!PinProxies.Contains(*Other) && !PinProxiesToProcess.Contains(*Other))
+					{
+						PinProxiesToProcess.Add(*Other);
+					}
+				}
+			}
+		}
+	}
+
+	return PinProxies;
+}
+
 bool URigVMCompiler::ValidateNode(URigVMNode* InNode)
 {
 	check(InNode);
diff --git a/Engine/Source/Developer/RigVMDeveloper/Public/RigVMCompiler/RigVMAST.h b/Engine/Source/Developer/RigVMDeveloper/Public/RigVMCompiler/RigVMAST.h
--- a/Engine/Source/Developer/RigVMDeveloper/Public/RigVMCompiler/RigVMAST.h
+++ b/Engine/Source/Developer/RigVMDeveloper/Public/RigVMCompiler/RigVMAST.h
@@ -1415,6 +1415,7 @@ private:
 	FRigVMBlockExprAST* ObsoleteBlock;
 
 	TArray<FRigVMASTProxy> NodeProxies;
+	TMap<FRigVMASTProxy, FRigVMASTProxy> SharedOperandPins;
 	TMap<FRigVMASTProxy, TArray<FRigVMASTProxy>> TargetLinks;
 	TMap<FRigVMASTProxy, TArray<FRigVMASTProxy>> SourceLinks;
 	static const TArray<FRigVMASTProxy> EmptyProxyArray;
diff --git a/Engine/Source/Developer/RigVMDeveloper/Public/RigVMCompiler/RigVMCompiler.h b/Engine/Source/Developer/RigVMDeveloper/Public/RigVMCompiler/RigVMCompiler.h
--- a/Engine/Source/Developer/RigVMDeveloper/Public/RigVMCompiler/RigVMCompiler.h
+++ b/Engine/Source/Developer/RigVMDeveloper/Public/RigVMCompiler/RigVMCompiler.h
@@ -173,6 +173,7 @@ private:
 	void InitializeLocalVariables(const FRigVMExprAST* InExpr, FRigVMCompilerWorkData& WorkData);
 
 	FRigVMOperand FindOrAddRegister(const FRigVMVarExprAST* InVarExpr, FRigVMCompilerWorkData& WorkData, bool bIsDebugValue = false);
+	TArray<FRigVMASTProxy> FindProxiesWithSharedOperand(const FRigVMVarExprAST* InVarExpr);
 
 	bool ValidateNode(URigVMNode* InNode);
 	
