Control Rig: Variable node can now reference an input parameter inside a function

diff --git a/Engine/Plugins/Animation/ControlRig/Source/ControlRigDeveloper/Private/Graph/ControlRigGraphNode.cpp b/Engine/Plugins/Animation/ControlRig/Source/ControlRigDeveloper/Private/Graph/ControlRigGraphNode.cpp
--- a/Engine/Plugins/Animation/ControlRig/Source/ControlRigDeveloper/Private/Graph/ControlRigGraphNode.cpp
+++ b/Engine/Plugins/Animation/ControlRig/Source/ControlRigDeveloper/Private/Graph/ControlRigGraphNode.cpp
@@ -106,6 +106,10 @@ FText UControlRigGraphNode::GetNodeTitle(ENodeTitleType::Type TitleType) const
 						SubTitle = FString::Printf(TEXT("%s\nDefault %s"), *LocalVariableString, *DefaultValue);
 					}
 				}
+				else if (VariableNode->IsInputArgument())
+				{
+					SubTitle = TEXT("Input parameter");
+				}
 				else
 				{
 					if(UBlueprint* Blueprint = GetBlueprint())
diff --git a/Engine/Plugins/Animation/ControlRig/Source/ControlRigDeveloper/Private/Graph/ControlRigGraphSchema.cpp b/Engine/Plugins/Animation/ControlRig/Source/ControlRigDeveloper/Private/Graph/ControlRigGraphSchema.cpp
--- a/Engine/Plugins/Animation/ControlRig/Source/ControlRigDeveloper/Private/Graph/ControlRigGraphSchema.cpp
+++ b/Engine/Plugins/Animation/ControlRig/Source/ControlRigDeveloper/Private/Graph/ControlRigGraphSchema.cpp
@@ -62,6 +62,11 @@ FControlRigLocalVariableNameValidator::FControlRigLocalVariableNameValidator(con
 		{
 			Names.Add(LocalVariable.Name.ToString());
 		}
+
+		for (const FRigVMGraphVariableDescription& InputArgument : Graph->GetInputArguments())
+		{
+			Names.Add(InputArgument.Name.ToString());
+		}
 	}
 }
 
@@ -76,6 +81,14 @@ FEdGraphPinType FControlRigGraphSchemaAction_LocalVar::GetPinType() const
 				return Variable.ToPinType();
 			}
 		}
+
+		for (FRigVMGraphVariableDescription Variable : Graph->GetModel()->GetInputArguments())
+		{
+			if (Variable.Name == GetVariableName())
+			{
+				return Variable.ToPinType();
+			}
+		}
 	}
 
 	return FEdGraphPinType();
diff --git a/Engine/Plugins/Animation/ControlRig/Source/ControlRigEditor/Private/ControlRigEditorModule.cpp b/Engine/Plugins/Animation/ControlRig/Source/ControlRigEditor/Private/ControlRigEditorModule.cpp
--- a/Engine/Plugins/Animation/ControlRig/Source/ControlRigEditor/Private/ControlRigEditorModule.cpp
+++ b/Engine/Plugins/Animation/ControlRig/Source/ControlRigEditor/Private/ControlRigEditorModule.cpp
@@ -1277,6 +1277,20 @@ void FControlRigEditorModule::GetInstanceActions(UControlRigBlueprint* CRB, FBlu
 				}
 			}
 		}
+
+		for (URigVMGraph* Graph : CRB->GetAllModels())
+		{
+			if (Graph->GetEntryNode())
+			{
+				FText NodeCategory = LOCTEXT("InputArguments", "Input Arguments");
+				for (const FRigVMGraphVariableDescription& InputArgument : Graph->GetInputArguments())
+				{
+					FText MenuDesc = FText::FromName(InputArgument.Name);
+					FText ToolTip = FText::FromString(FString::Printf(TEXT("Get the value of input %s"), *InputArgument.Name.ToString()));
+					ActionRegistrar.AddBlueprintAction(GeneratedClass, UControlRigVariableNodeSpawner::CreateFromLocalVariable(CRB, Graph, InputArgument, true, MenuDesc, NodeCategory, ToolTip));
+				}			
+			}
+		}
 	}
 }
 
diff --git a/Engine/Plugins/Animation/ControlRig/Source/ControlRigEditor/Private/ControlRigGraphDetails.cpp b/Engine/Plugins/Animation/ControlRig/Source/ControlRigEditor/Private/ControlRigGraphDetails.cpp
--- a/Engine/Plugins/Animation/ControlRig/Source/ControlRigEditor/Private/ControlRigGraphDetails.cpp
+++ b/Engine/Plugins/Animation/ControlRig/Source/ControlRigEditor/Private/ControlRigGraphDetails.cpp
@@ -159,8 +159,11 @@ void FControlRigArgumentLayout::GenerateHeaderRowContent(FDetailWidgetRow& NodeR
 						OutErrorMessage = LOCTEXT("ArgumentNameTooLong", "Name of argument is too long.");
 						return false;
 					}
-					
-					return true;
+
+					EValidatorResult Result = NameValidator.IsValid(InNewText.ToString(), false);
+					OutErrorMessage = INameValidatorInterface::GetErrorText(InNewText.ToString(), Result);	
+
+					return Result == EValidatorResult::Ok || Result == EValidatorResult::ExistingName;
 				})
 			]
 		]
@@ -422,7 +425,7 @@ void FControlRigArgumentLayout::PinInfoChanged(const FEdGraphPinType& PinType)
 				}
 
 				bool bSetupUndoRedo = true;
-				Controller->ChangeExposedPinType(Pin->GetFName(), CPPType, CPPTypeObjectName, bSetupUndoRedo, true, true);
+				Controller->ChangeExposedPinType(Pin->GetFName(), CPPType, CPPTypeObjectName, bSetupUndoRedo, false, true);
 
 				// If the controller has identified this as a bulk change, it has not added the actions to the action stack
 				// We need to disable the transaction from the UI as well to keep them synced
diff --git a/Engine/Plugins/Animation/ControlRig/Source/ControlRigEditor/Private/ControlRigGraphDetails.h b/Engine/Plugins/Animation/ControlRig/Source/ControlRigEditor/Private/ControlRigGraphDetails.h
--- a/Engine/Plugins/Animation/ControlRig/Source/ControlRigEditor/Private/ControlRigGraphDetails.h
+++ b/Engine/Plugins/Animation/ControlRig/Source/ControlRigEditor/Private/ControlRigGraphDetails.h
@@ -13,6 +13,7 @@
 #include "Graph/SControlRigGraphNode.h"
 #include "Widgets/Colors/SColorBlock.h"
 #include "DetailsViewWrapperObject.h"
+#include "Graph/ControlRigGraphSchema.h"
 
 class IDetailLayoutBuilder;
 
@@ -56,6 +57,7 @@ public:
 		: PinPtr(InPin)
 		, GraphPtr(InGraph)
 		, ControlRigBlueprintPtr(InBlueprint)
+		, NameValidator(InBlueprint, InGraph, InPin->GetFName())
 	{}
 
 private:
@@ -102,6 +104,9 @@ private:
 
 	/** Holds a weak pointer to the argument name widget, used for error notifications */
 	TWeakPtr<SEditableTextBox> ArgumentNameWidget;
+
+	/** The validator to check if a name for an argument is valid */
+	FControlRigLocalVariableNameValidator NameValidator;
 };
 
 class FControlRigArgumentDefaultNode : public IDetailCustomNodeBuilder, public TSharedFromThis<FControlRigArgumentDefaultNode>
diff --git a/Engine/Plugins/Animation/ControlRig/Source/ControlRigEditor/Private/Graph/SControlRigGraphPinVariableBinding.cpp b/Engine/Plugins/Animation/ControlRig/Source/ControlRigEditor/Private/Graph/SControlRigGraphPinVariableBinding.cpp
--- a/Engine/Plugins/Animation/ControlRig/Source/ControlRigEditor/Private/Graph/SControlRigGraphPinVariableBinding.cpp
+++ b/Engine/Plugins/Animation/ControlRig/Source/ControlRigEditor/Private/Graph/SControlRigGraphPinVariableBinding.cpp
@@ -117,7 +117,7 @@ FLinearColor SControlRigVariableBinding::GetBindingColor() const
 				return  FLinearColor::Red;
 			}
 
-			const TArray<FRigVMGraphVariableDescription>& LocalVariables =  Model->GetLocalVariables();
+			const TArray<FRigVMGraphVariableDescription>& LocalVariables =  Model->GetLocalVariables(true);
 			for(const FRigVMGraphVariableDescription& LocalVariable : LocalVariables)
 			{
 				const FRigVMExternalVariable ExternalVariable = LocalVariable.ToExternalVariable();
@@ -249,7 +249,7 @@ void SControlRigVariableBinding::FillLocalVariableMenu(FMenuBuilder& MenuBuilder
 	}
 
 	int32 ValidLocalVariables = 0;
-	const TArray<FRigVMGraphVariableDescription>& LocalVariables =  Model->GetLocalVariables();
+	const TArray<FRigVMGraphVariableDescription>& LocalVariables =  Model->GetLocalVariables(true);
 	for(const FRigVMGraphVariableDescription& LocalVariable : LocalVariables)
 	{
 		const FRigVMExternalVariable ExternalVariable = LocalVariable.ToExternalVariable();
diff --git a/Engine/Source/Developer/RigVMDeveloper/Private/RigVMCompiler/RigVMAST.cpp b/Engine/Source/Developer/RigVMDeveloper/Private/RigVMCompiler/RigVMAST.cpp
--- a/Engine/Source/Developer/RigVMDeveloper/Private/RigVMCompiler/RigVMAST.cpp
+++ b/Engine/Source/Developer/RigVMDeveloper/Private/RigVMCompiler/RigVMAST.cpp
@@ -2813,6 +2813,11 @@ void FRigVMParserAST::Inline(URigVMGraph* InGraph, const TArray<FRigVMASTProxy>&
 		{
 			URigVMPin* Pin = InPinProxy.GetSubjectChecked<URigVMPin>();
 			URigVMNode* Node = Pin->GetNode();
+			if (URigVMVariableNode* VarNode = Cast<URigVMVariableNode>(Node))
+			{
+				return VarNode->IsInputArgument();
+			}
+			
 			return Node->IsA<URigVMRerouteNode>() ||
 				Node->IsA<URigVMLibraryNode>() ||
 				Node->IsA<URigVMFunctionEntryNode>() ||
@@ -2861,6 +2866,23 @@ void FRigVMParserAST::Inline(URigVMGraph* InGraph, const TArray<FRigVMASTProxy>&
 							}
 						}
 					}
+					else if(URigVMVariableNode* VariableNode = Cast<URigVMVariableNode>(Node))
+					{
+						if (VariableNode->IsInputArgument())
+						{
+							if (URigVMFunctionEntryNode* EntryNode = VariableNode->GetGraph()->GetEntryNode())
+							{
+								if (URigVMPin* EntryPin = EntryNode->FindPin(VariableNode->GetVariableName().ToString()))
+								{
+									FRigVMASTProxy EntryPinProxy = InPinProxy.GetSibling(EntryPin);
+									FRigVMASTProxy SourcePinProxy = FindSourcePin(EntryPinProxy, OutTraversalInfo);
+									SourcePinProxy = SourcePinProxy.IsValid() ? SourcePinProxy : EntryPinProxy;
+									OutTraversalInfo.SourcePins.FindOrAdd(InPinProxy) = SourcePinProxy;
+									return SourcePinProxy;
+								} 
+							}
+						}
+					}
 					else if (URigVMFunctionEntryNode* EntryNode = Cast<URigVMFunctionEntryNode>(Node))
 					{
 						for (int32 LibraryNodeIndex = OutTraversalInfo.LibraryNodeCallstack.Num() - 1; LibraryNodeIndex >= 0; LibraryNodeIndex--)
diff --git a/Engine/Source/Developer/RigVMDeveloper/Private/RigVMModel/Nodes/RigVMVariableNode.cpp b/Engine/Source/Developer/RigVMDeveloper/Private/RigVMModel/Nodes/RigVMVariableNode.cpp
--- a/Engine/Source/Developer/RigVMDeveloper/Private/RigVMModel/Nodes/RigVMVariableNode.cpp
+++ b/Engine/Source/Developer/RigVMDeveloper/Private/RigVMModel/Nodes/RigVMVariableNode.cpp
@@ -58,6 +58,26 @@ bool URigVMVariableNode::IsLocalVariable() const
 	return false;
 }
 
+bool URigVMVariableNode::IsInputArgument() const
+{
+	const FName CurrentVariableName = GetVariableName();
+	
+	if(URigVMGraph* Graph = GetGraph())
+	{
+		if (URigVMFunctionEntryNode* EntryNode = Graph->GetEntryNode())
+		{
+			for (const URigVMPin* Pin : EntryNode->GetPins())
+			{
+				if(Pin->GetFName() == CurrentVariableName)
+				{
+					return true;
+				}
+			}
+		}
+	}
+	return false;
+}
+
 FString URigVMVariableNode::GetCPPType() const
 {
 	URigVMPin* ValuePin = FindPin(ValueName);
diff --git a/Engine/Source/Developer/RigVMDeveloper/Private/RigVMModel/RigVMController.cpp b/Engine/Source/Developer/RigVMDeveloper/Private/RigVMModel/RigVMController.cpp
--- a/Engine/Source/Developer/RigVMDeveloper/Private/RigVMModel/RigVMController.cpp
+++ b/Engine/Source/Developer/RigVMDeveloper/Private/RigVMModel/RigVMController.cpp
@@ -1286,7 +1286,7 @@ void URigVMController::OnExternalVariableRenamed(const FName& InOldVarName, cons
 	URigVMGraph* Graph = GetGraph();
 	check(Graph);
 
-	for (const FRigVMGraphVariableDescription& LocalVariable : Graph->GetLocalVariables())
+	for (const FRigVMGraphVariableDescription& LocalVariable : Graph->GetLocalVariables(true))
 	{
 		if (InOldVarName == LocalVariable.Name)
 		{
@@ -1384,7 +1384,7 @@ void URigVMController::OnExternalVariableTypeChanged(const FName& InVarName, con
 	URigVMGraph* Graph = GetGraph();
 	check(Graph);
 
-	for (const FRigVMGraphVariableDescription& LocalVariable : Graph->GetLocalVariables())
+	for (const FRigVMGraphVariableDescription& LocalVariable : Graph->GetLocalVariables(true))
 	{
 		if (InVarName == LocalVariable.Name)
 		{
@@ -3107,6 +3107,14 @@ URigVMCollapseNode* URigVMController::CollapseNodes(const TArray<URigVMNode*>& I
 			continue;
 		}
 
+		if (URigVMVariableNode* VariableNode = Cast<URigVMVariableNode>(Node))
+		{
+			if (VariableNode->IsInputArgument())
+			{
+				continue;
+			}
+		}
+
 		Nodes.Add(Node);
 	}
 
@@ -3689,6 +3697,14 @@ TArray<URigVMNode*> URigVMController::ExpandLibraryNode(URigVMLibraryNode* InNod
 				continue;
 			}
 
+			if (URigVMVariableNode* VariableNode = Cast<URigVMVariableNode>(Node))
+			{
+				if (VariableNode->IsInputArgument())
+				{
+					continue;
+				}
+			}
+
 			if(Node->IsInjected())
 			{
 				continue;
@@ -3714,15 +3730,18 @@ TArray<URigVMNode*> URigVMController::ExpandLibraryNode(URigVMLibraryNode* InNod
 	// exist, they will be reused. If a local variable is not used, it will not be created.
 	if (URigVMFunctionReferenceNode* FunctionReferenceNode = Cast<URigVMFunctionReferenceNode>(InNode))
 	{
-		const bool bIsWithinFunctionGraph = !Graph->IsRootGraph() && Graph->GetRootGraph()->IsA<URigVMFunctionLibrary>();
-		TArray<FRigVMExternalVariable> AllExistingVariables = GetAllVariables();
-
 		TArray<FRigVMGraphVariableDescription> LocalVariables = FunctionReferenceNode->GetContainedGraph()->LocalVariables;
+		TArray<FRigVMExternalVariable> CurrentVariables = GetAllVariables();
 		TArray<FRigVMGraphVariableDescription> VariablesToAdd;
 		for (const URigVMNode* Node : FunctionReferenceNode->GetContainedGraph()->GetNodes())
 		{
 			if (const URigVMVariableNode* VariableNode = Cast<URigVMVariableNode>(Node))
 			{
+				if (VariableNode->IsInputArgument())
+				{
+					continue;
+				}
+				
 				for (FRigVMGraphVariableDescription& LocalVariable : LocalVariables)
 				{
 					if (LocalVariable.Name == VariableNode->GetVariableName())
@@ -3730,13 +3749,13 @@ TArray<URigVMNode*> URigVMController::ExpandLibraryNode(URigVMLibraryNode* InNod
 						bool bVariableExists = false;
 						bool bVariableIncompatible = false;
 						FRigVMExternalVariable LocalVariableExternalType = LocalVariable.ToExternalVariable();
-						for (FRigVMExternalVariable& ExistingVariable : AllExistingVariables)
+						for (FRigVMExternalVariable& CurrentVariable : CurrentVariables)
 						{						
-							if (ExistingVariable.Name == LocalVariable.Name)
+							if (CurrentVariable.Name == LocalVariable.Name)
 							{
-								if (ExistingVariable.TypeName != LocalVariableExternalType.TypeName ||
-									ExistingVariable.TypeObject != LocalVariableExternalType.TypeObject ||
-									ExistingVariable.bIsArray != LocalVariableExternalType.bIsArray)
+								if (CurrentVariable.TypeName != LocalVariableExternalType.TypeName ||
+									CurrentVariable.TypeObject != LocalVariableExternalType.TypeObject ||
+									CurrentVariable.bIsArray != LocalVariableExternalType.bIsArray)
 								{
 									bVariableIncompatible = true;	
 								}
@@ -3751,7 +3770,7 @@ TArray<URigVMNode*> URigVMController::ExpandLibraryNode(URigVMLibraryNode* InNod
 						}
 						else if(bVariableIncompatible)
 						{
-							ReportAndNotifyErrorf(TEXT("Found variable %s of incompatible type with a local variable inside function %s"), *LocalVariable.Name.ToString(), *FunctionReferenceNode->GetReferencedNode()->GetName());
+							ReportErrorf(TEXT("Found variable %s of incompatible type with a local variable inside function %s"), *LocalVariable.Name.ToString(), *FunctionReferenceNode->GetReferencedNode()->GetName());
 							if (bSetupUndoRedo)
 							{
 								ActionStack->CancelAction(ExpandAction);
@@ -3764,21 +3783,11 @@ TArray<URigVMNode*> URigVMController::ExpandLibraryNode(URigVMLibraryNode* InNod
 			}
 		}
 
-		if(!bIsWithinFunctionGraph)
-		{
-			if (RequestNewExternalVariableDelegate.IsBound())
-			{
-				for (const FRigVMGraphVariableDescription& VariableToAdd : VariablesToAdd)
-				{
-					RequestNewExternalVariableDelegate.Execute(VariableToAdd, false, false);
-				}
-			}
-		}
-		else
+		if (RequestNewExternalVariableDelegate.IsBound())
 		{
-			for (const FRigVMGraphVariableDescription& VariableToAdd : VariablesToAdd)
+			for (const FRigVMGraphVariableDescription& OldVariable : VariablesToAdd)
 			{
-				AddLocalVariable(VariableToAdd.Name, VariableToAdd.CPPType, VariableToAdd.CPPTypeObject, VariableToAdd.DefaultValue, bSetupUndoRedo, false);
+				RequestNewExternalVariableDelegate.Execute(OldVariable, false, false);
 			}
 		}
 	}
@@ -3855,9 +3864,21 @@ TArray<URigVMNode*> URigVMController::ExpandLibraryNode(URigVMLibraryNode* InNod
 	if (URigVMFunctionEntryNode* EntryNode = InNode->GetEntryNode())
 	{
 		TArray<URigVMLink*> EntryLinks = EntryNode->GetLinks();
+
+		for (URigVMNode* Node : InNode->GetContainedGraph()->GetNodes())
+		{
+			if (URigVMVariableNode* VariableNode = Cast<URigVMVariableNode>(Node))
+			{
+				if (VariableNode->IsInputArgument())
+				{
+					EntryLinks.Append(VariableNode->GetLinks());
+				}
+			}
+		}
+		
 		for (URigVMLink* Link : EntryLinks)
 		{
-			if (Link->GetSourcePin()->GetNode() != EntryNode)
+			if (Link->GetSourcePin()->GetNode() != EntryNode && !Link->GetSourcePin()->GetNode()->IsA<URigVMVariableNode>())
 			{
 				continue;
 			}
@@ -3870,6 +3891,11 @@ TArray<URigVMNode*> URigVMController::ExpandLibraryNode(URigVMLibraryNode* InNod
 			FString NodeName, PinPath;
 			URigVMPin::SplitPinPathAtStart(Link->GetSourcePin()->GetPinPath(), NodeName, PinPath);
 
+			if (URigVMVariableNode* VariableNode = Cast<URigVMVariableNode>(Link->GetSourcePin()->GetNode()))
+			{
+				PinPath = VariableNode->GetVariableName().ToString();
+			}
+
 			TArray<FString>& LinkedPins = FromEntryNode.FindOrAdd(PinPath);
 
 			URigVMPin::SplitPinPathAtStart(Link->GetTargetPin()->GetPinPath(), NodeName, PinPath);
@@ -4136,52 +4162,15 @@ TArray<URigVMNode*> URigVMController::ExpandLibraryNode(URigVMLibraryNode* InNod
 			RemappedSourcePinPath = TargetPin->GetPinPath();
 		}
 
-		
 		for (const FString& FromEntryNodeTargetPinPath : FromEntryNodePair.Value)
 		{
-			TArray<URigVMPin*> TargetPins;
-
-			URigVMPin* SourcePin = GetGraph()->FindPin(RemappedSourcePinPath);
+			URigVMPin* SourcePin = GetGraph()->FindPin(*RemappedSourcePinPath);
 			URigVMPin* TargetPin = GetGraph()->FindPin(FromEntryNodeTargetPinPath);
-
-			// potentially the target pin was on the entry node,
-			// so there's no node been added for it. we'll have to look into the remapped
-			// pins for the "FromLibraryNode" map.
-			if(TargetPin == nullptr)
-			{
-				FString RemappedTargetPinPath = FromEntryNodeTargetPinPath;
-				FString ReturnNodeName, ReturnPinPath;
-				if (URigVMPin::SplitPinPathAtStart(RemappedTargetPinPath, ReturnNodeName, ReturnPinPath))
-				{
-					if(Cast<URigVMFunctionReturnNode>(InNode->GetContainedGraph()->FindNode(ReturnNodeName)))
-					{
-						if(FromLibraryNode.Contains(ReturnPinPath))
-						{
-							const TArray<FString>& FromLibraryNodeTargetPins = FromLibraryNode.FindChecked(ReturnPinPath);
-							for(const FString& FromLibraryNodeTargetPin : FromLibraryNodeTargetPins)
-							{
-								if(URigVMPin* MappedTargetPin = GetGraph()->FindPin(FromLibraryNodeTargetPin))
-								{
-									TargetPins.Add(MappedTargetPin);
-								}
-							}
-						}
-					}
-				}
-			}
-			else
-			{
-				TargetPins.Add(TargetPin);
-			}
-			
-			if (SourcePin)
+			if (SourcePin && TargetPin)
 			{
-				for(URigVMPin* EachTargetPin : TargetPins)
+				if (!SourcePin->IsLinkedTo(TargetPin))
 				{
-					if (!SourcePin->IsLinkedTo(EachTargetPin))
-					{
-						AddLink(SourcePin, EachTargetPin, false);
-					}
+					AddLink(SourcePin, TargetPin, false);
 				}
 			}
 		}
@@ -5654,7 +5643,7 @@ bool URigVMController::RenameVariable(const FName& InOldName, const FName& InNew
 	}
 
 	// If there is a local variable with the old name, a rename of the blueprint member variable does not affect this graph
-	for (FRigVMGraphVariableDescription& LocalVariable : Graph->GetLocalVariables())
+	for (FRigVMGraphVariableDescription& LocalVariable : Graph->GetLocalVariables(true))
 	{
 		if (LocalVariable.Name == InOldName)
 		{
@@ -6038,11 +6027,6 @@ bool URigVMController::SetPinDefaultValue(URigVMPin* InPin, const FString& InDef
 
 	TGuardValue<bool> Guard(bSuspendNotifications, !bNotify);
 
-	if (InPin->GetDirection() == ERigVMPinDirection::Hidden)
-	{
-		return false;
-	}
-	
 	URigVMGraph* Graph = GetGraph();
 	check(Graph);
  
@@ -6665,7 +6649,7 @@ bool URigVMController::BindPinToVariable(URigVMPin* InPin, const FString& InNewB
 		FString VariableName = InNewBoundVariablePath, SegmentPath;
 		InNewBoundVariablePath.Split(TEXT("."), &VariableName, &SegmentPath);
 
-		FRigVMExternalVariable Variable = GetVariableByName(*VariableName);
+		FRigVMExternalVariable Variable = GetVariableByName(*VariableName, true);
 		if (Variable.IsValid(true))
 		{
 #if UE_RIGVM_UCLASS_BASED_STORAGE_DISABLED
@@ -7382,7 +7366,15 @@ FName URigVMController::AddExposedPin(const FName& InPinName, ERigVMPinDirection
 
 	FName PinName = GetUniqueName(InPinName, [LibraryNode](const FName& InName) {
 
-		return LibraryNode->FindPin(InName.ToString()) == nullptr;
+		const TArray<FRigVMGraphVariableDescription>& LocalVariables = LibraryNode->GetContainedGraph()->GetLocalVariables(true);
+		for(const FRigVMGraphVariableDescription& VariableDescription : LocalVariables)
+		{
+			if (VariableDescription.Name == InName)
+			{
+				return false;
+			}
+		}
+		return true;
 
 	}, false, true);
 
@@ -7513,6 +7505,22 @@ bool URigVMController::RemoveExposedPin(const FName& InPinName, bool bSetupUndoR
 		bSuccessfullyRemovedPin = RemovePin(Pin, bSetupUndoRedo, true);
 	}
 
+	TArray<URigVMVariableNode*> NodesToRemove;
+	for (URigVMNode* Node : Graph->GetNodes())
+	{
+		if (URigVMVariableNode* VariableNode = Cast<URigVMVariableNode>(Node))
+		{
+			if (VariableNode->GetVariableName() == InPinName)
+			{
+				NodesToRemove.Add(VariableNode);
+			}
+		}
+	}
+	for (int32 i=NodesToRemove.Num()-1; i >= 0; --i)
+	{
+		RemoveNode(NodesToRemove[i], bSetupUndoRedo);
+	}
+
 	RefreshFunctionPins(Graph->GetEntryNode(), true);
 	RefreshFunctionPins(Graph->GetReturnNode(), true);
 	RefreshFunctionReferences(LibraryNode, false);
@@ -7592,7 +7600,15 @@ bool URigVMController::RenameExposedPin(const FName& InOldPinName, const FName&
 	}
 
 	FName PinName = GetUniqueName(InNewPinName, [LibraryNode](const FName& InName) {
-		return LibraryNode->FindPin(InName.ToString()) == nullptr;
+		const TArray<FRigVMGraphVariableDescription>& LocalVariables = LibraryNode->GetContainedGraph()->GetLocalVariables(true);
+		for(const FRigVMGraphVariableDescription& VariableDescription : LocalVariables)
+		{
+			if (VariableDescription.Name == InName)
+			{
+				return false;
+			}
+		}
+		return true;
 	}, false, true);
 
 	FRigVMRenameExposedPinAction Action;
@@ -7675,6 +7691,17 @@ bool URigVMController::RenameExposedPin(const FName& InOldPinName, const FName&
         });
 	}
 
+	for (URigVMNode* Node : Graph->GetNodes())
+	{
+		if (URigVMVariableNode* VariableNode = Cast<URigVMVariableNode>(Node))
+		{
+			if (VariableNode->GetVariableName() == InOldPinName)
+			{
+				SetVariableName(VariableNode, InNewPinName, bSetupUndoRedo);
+			}
+		}
+	}
+
 	if (bSetupUndoRedo)
 	{
 		ActionStack->EndAction(Action);
@@ -7789,6 +7816,22 @@ bool URigVMController::ChangeExposedPinType(const FName& InPinName, const FStrin
         });
 	}
 
+	for (URigVMNode* Node : Graph->GetNodes())
+	{
+		if (URigVMVariableNode* VariableNode = Cast<URigVMVariableNode>(Node))
+		{
+			if (VariableNode->GetVariableName() == InPinName)
+			{
+				URigVMPin* ValuePin = VariableNode->FindPin(URigVMVariableNode::ValueName);
+				if (ValuePin)
+				{
+					ChangePinType(ValuePin, InCPPType, InCPPTypeObjectPath, bSetupUndoRedo, bSetupOrphanPins);
+					RemoveUnusedOrphanedPins(VariableNode, true);
+				}
+			}
+		}
+	}
+
 	if (bSetupUndoRedo)
 	{
 		ActionStack->EndAction(Action);
@@ -8236,7 +8279,7 @@ FRigVMGraphVariableDescription URigVMController::AddLocalVariable(const FName& I
 	}
 
 	FName VariableName = GetUniqueName(InVariableName, [Graph](const FName& InName) {
-		for (FRigVMGraphVariableDescription LocalVariable : Graph->GetLocalVariables())
+		for (FRigVMGraphVariableDescription LocalVariable : Graph->GetLocalVariables(true))
 		{
 			if (LocalVariable.Name == InName)
 			{
@@ -8621,7 +8664,7 @@ bool URigVMController::SetLocalVariableType(const FName& InVariableName, const F
 			{
 				if (VariablePin->GetDefaultValue() == InVariableName.ToString())
 				{
-					RefreshVariableNode(Node->GetFName(), InVariableName, InCPPType, InCPPTypeObject, bSetupUndoRedo, false);
+					RefreshVariableNode(Node->GetFName(), InVariableName, InCPPType, InCPPTypeObject, bSetupUndoRedo);
 					continue;
 				}
 			}
@@ -8976,6 +9019,7 @@ bool URigVMController::SetVariableName(URigVMVariableNode* InVariableNode, const
 	SetPinDefaultValue(InVariableNode->FindPin(URigVMVariableNode::VariableName), VariableName.ToString(), false, bSetupUndoRedo, false);
 
 	Notify(ERigVMGraphNotifType::VariableAdded, InVariableNode);
+	Notify(ERigVMGraphNotifType::VariableRenamed, InVariableNode);
 
 	return true;
 }
@@ -10064,7 +10108,7 @@ bool URigVMController::CanAddNode(URigVMNode* InNode, bool bReportErrors, bool b
 			FString VarName = NamePin->GetDefaultValue();
 			if (!VarName.IsEmpty())
 			{
-				TArray<FRigVMExternalVariable> AllVariables = URigVMController::GetAllVariables();
+				TArray<FRigVMExternalVariable> AllVariables = URigVMController::GetAllVariables(true);
 				for(const FRigVMExternalVariable& Variable : AllVariables)
 				{
 					if(Variable.Name.ToString() == VarName)
@@ -12055,9 +12099,9 @@ void URigVMController::ValidatePin(URigVMPin* InPin)
 	FRigVMPropertyDescription(InPin->GetFName(), InPin->GetCPPType(), InPin->GetCPPTypeObject(), InPin->GetDefaultValue());
 }
 
-FRigVMExternalVariable URigVMController::GetVariableByName(const FName& InExternalVariableName)
+FRigVMExternalVariable URigVMController::GetVariableByName(const FName& InExternalVariableName, const bool bIncludeInputArguments)
 {
-	TArray<FRigVMExternalVariable> Variables = GetAllVariables();
+	TArray<FRigVMExternalVariable> Variables = GetAllVariables(bIncludeInputArguments);
 	for (const FRigVMExternalVariable& Variable : Variables)
 	{
 		if (Variable.Name == InExternalVariableName)
@@ -12069,13 +12113,13 @@ FRigVMExternalVariable URigVMController::GetVariableByName(const FName& InExtern
 	return FRigVMExternalVariable();
 }
 
-TArray<FRigVMExternalVariable> URigVMController::GetAllVariables()
+TArray<FRigVMExternalVariable> URigVMController::GetAllVariables(const bool bIncludeInputArguments)
 {
 	TArray<FRigVMExternalVariable> ExternalVariables;
 
 	if(URigVMGraph* Graph = GetGraph())
 	{
-		for (FRigVMGraphVariableDescription LocalVariable : Graph->GetLocalVariables())
+		for (FRigVMGraphVariableDescription LocalVariable : Graph->GetLocalVariables(bIncludeInputArguments))
 		{
 			ExternalVariables.Add(LocalVariable.ToExternalVariable());
 		}
diff --git a/Engine/Source/Developer/RigVMDeveloper/Private/RigVMModel/RigVMGraph.cpp b/Engine/Source/Developer/RigVMDeveloper/Private/RigVMModel/RigVMGraph.cpp
--- a/Engine/Source/Developer/RigVMDeveloper/Private/RigVMModel/RigVMGraph.cpp
+++ b/Engine/Source/Developer/RigVMDeveloper/Private/RigVMModel/RigVMGraph.cpp
@@ -279,6 +279,36 @@ TArray<FRigVMExternalVariable> URigVMGraph::GetExternalVariables() const
 	return Variables;
 }
 
+TArray<FRigVMGraphVariableDescription> URigVMGraph::GetLocalVariables(bool bIncludeInputArguments) const
+{
+	if (bIncludeInputArguments)
+	{
+		TArray<FRigVMGraphVariableDescription> Variables;
+		Variables.Append(LocalVariables);
+		Variables.Append(GetInputArguments());
+		return Variables;
+	}
+	
+	return LocalVariables;
+}
+
+TArray<FRigVMGraphVariableDescription> URigVMGraph::GetInputArguments() const
+{
+	TArray<FRigVMGraphVariableDescription> Inputs;
+	if (URigVMFunctionEntryNode* EntryNode = GetEntryNode())
+	{
+		for (URigVMPin* Pin : EntryNode->GetPins())
+		{			
+			FRigVMGraphVariableDescription Description;
+			Description.Name = Pin->GetFName();
+			Description.CPPType = Pin->GetCPPType();
+			Description.CPPTypeObject = Pin->GetCPPTypeObject();
+			Inputs.Add(Description);			
+		}
+	}
+	return Inputs;
+}
+
 FRigVMGraphModifiedEvent& URigVMGraph::OnModified()
 {
 	return ModifiedEvent;
diff --git a/Engine/Source/Developer/RigVMDeveloper/Public/RigVMModel/Nodes/RigVMVariableNode.h b/Engine/Source/Developer/RigVMDeveloper/Public/RigVMModel/Nodes/RigVMVariableNode.h
--- a/Engine/Source/Developer/RigVMDeveloper/Public/RigVMModel/Nodes/RigVMVariableNode.h
+++ b/Engine/Source/Developer/RigVMDeveloper/Public/RigVMModel/Nodes/RigVMVariableNode.h
@@ -38,6 +38,10 @@ public:
 	UFUNCTION(BlueprintCallable, Category = RigVMVariableNode)
 	bool IsLocalVariable() const;
 
+	// Returns true if this variable is an input argument
+	UFUNCTION(BlueprintCallable, Category = RigVMVariableNode)
+	bool IsInputArgument() const;
+
 	// Returns the C++ data type of the variable
 	UFUNCTION(BlueprintCallable, Category = RigVMVariableNode)
 	FString GetCPPType() const;
diff --git a/Engine/Source/Developer/RigVMDeveloper/Public/RigVMModel/RigVMController.h b/Engine/Source/Developer/RigVMDeveloper/Public/RigVMModel/RigVMController.h
--- a/Engine/Source/Developer/RigVMDeveloper/Public/RigVMModel/RigVMController.h
+++ b/Engine/Source/Developer/RigVMDeveloper/Public/RigVMModel/RigVMController.h
@@ -188,7 +188,7 @@ public:
 
 	// Refreshes the variable node with the new data
 	UFUNCTION(BlueprintCallable, Category = RigVMController)
-	void RefreshVariableNode(const FName& InNodeName, const FName& InVariableName, const FString& InCPPType, UObject* InCPPTypeObject, bool bSetupUndoRedo, bool bDetachAndReattachLinks = true);
+	void RefreshVariableNode(const FName& InNodeName, const FName& InVariableName, const FString& InCPPType, UObject* InCPPTypeObject, bool bSetupUndoRedo, bool bSetupOrphanPins = true);
 
 	// Removes all nodes related to a given variable
 	UFUNCTION(BlueprintCallable, Category = RigVMController)
@@ -868,8 +868,8 @@ private:
 	static bool EnsurePinValidity(URigVMPin* InPin, bool bRecursive);
 	static void ValidatePin(URigVMPin* InPin);
 
-	FRigVMExternalVariable GetVariableByName(const FName& InExternalVariableName);
-	TArray<FRigVMExternalVariable> GetAllVariables();
+	FRigVMExternalVariable GetVariableByName(const FName& InExternalVariableName, const bool bIncludeInputArguments = false);
+	TArray<FRigVMExternalVariable> GetAllVariables(const bool bIncludeInputArguments = false);
 
 	void RefreshFunctionReferences(URigVMLibraryNode* InFunctionDefinition, bool bSetupUndoRedo);
 
diff --git a/Engine/Source/Developer/RigVMDeveloper/Public/RigVMModel/RigVMGraph.h b/Engine/Source/Developer/RigVMDeveloper/Public/RigVMModel/RigVMGraph.h
--- a/Engine/Source/Developer/RigVMDeveloper/Public/RigVMModel/RigVMGraph.h
+++ b/Engine/Source/Developer/RigVMDeveloper/Public/RigVMModel/RigVMGraph.h
@@ -121,8 +121,13 @@ public:
 
 	TArray<FRigVMExternalVariable> GetExternalVariables() const;
 
+	// Returns the local variables of this function
 	UFUNCTION(BlueprintCallable, Category = RigVMGraph)
-	TArray<FRigVMGraphVariableDescription> GetLocalVariables() const { return LocalVariables; }
+	TArray<FRigVMGraphVariableDescription> GetLocalVariables(bool bIncludeInputArguments = false) const;
+
+	// Returns the input arguments of this graph
+	UFUNCTION(BlueprintCallable, Category = RigVMGraph)
+	TArray<FRigVMGraphVariableDescription> GetInputArguments() const;
 
 	// Returns the modified event, which can be used to 
 	// subscribe to changes happening within the Graph.
