Control Rig: Fix compile issue

diff --git a/Engine/Source/Developer/RigVMDeveloper/Private/RigVMCompiler/RigVMAST.cpp b/Engine/Source/Developer/RigVMDeveloper/Private/RigVMCompiler/RigVMAST.cpp
--- a/Engine/Source/Developer/RigVMDeveloper/Private/RigVMCompiler/RigVMAST.cpp
+++ b/Engine/Source/Developer/RigVMDeveloper/Private/RigVMCompiler/RigVMAST.cpp
@@ -2877,7 +2877,7 @@ void FRigVMParserAST::Inline(URigVMGraph* InGraph, const TArray<FRigVMASTProxy>&
 									FRigVMASTProxy EntryPinProxy = InPinProxy.GetSibling(EntryPin);
 									FRigVMASTProxy SourcePinProxy = FindSourcePin(EntryPinProxy, OutTraversalInfo);
 									SourcePinProxy = SourcePinProxy.IsValid() ? SourcePinProxy : EntryPinProxy;
-									OutTraversalInfo.SourcePins.FindOrAdd(InPinProxy) = SourcePinProxy;
+									OutTraversalInfo.SourcePins->FindOrAdd(InPinProxy) = SourcePinProxy;
 									return SourcePinProxy;
 								} 
 							}
