Control Rig: Fix Unable to redo Forward solve nodes paste

diff --git a/Engine/Source/Developer/RigVMDeveloper/Private/RigVMModel/RigVMControllerActions.cpp b/Engine/Source/Developer/RigVMDeveloper/Private/RigVMModel/RigVMControllerActions.cpp
--- a/Engine/Source/Developer/RigVMDeveloper/Private/RigVMModel/RigVMControllerActions.cpp
+++ b/Engine/Source/Developer/RigVMDeveloper/Private/RigVMModel/RigVMControllerActions.cpp
@@ -745,7 +745,7 @@ FRigVMRemoveNodeAction::FRigVMRemoveNodeAction(URigVMNode* InNode, URigVMControl
 
 	for (URigVMPin* Pin : InNode->GetPins())
 	{
-		if(Pin->IsExpanded())
+		if(Pin->IsExpanded() && Pin->GetSubPins().Num() > 0)
 		{
 			FRigVMSetPinExpansionAction ExpansionAction(Pin, true);
 			ExpansionAction.OldIsExpanded = false;
