Control Rig: Add "add_if_node_from_struct" in Python

diff --git a/Engine/Source/Developer/RigVMDeveloper/Private/RigVMModel/RigVMController.cpp b/Engine/Source/Developer/RigVMDeveloper/Private/RigVMModel/RigVMController.cpp
--- a/Engine/Source/Developer/RigVMDeveloper/Private/RigVMModel/RigVMController.cpp
+++ b/Engine/Source/Developer/RigVMDeveloper/Private/RigVMModel/RigVMController.cpp
@@ -9316,6 +9316,16 @@ URigVMIfNode* URigVMController::AddIfNode(const FString& InCPPType, const FName&
 	return Node;
 }
 
+URigVMIfNode* URigVMController::AddIfNodeFromStruct(UScriptStruct* InScriptStruct, const FVector2D& InPosition, const FString& InNodeName, bool bSetupUndoRedo)
+{
+	if (!InScriptStruct)
+	{
+		return nullptr;
+	}
+
+	return AddIfNode(InScriptStruct->GetStructCPPName(), FName(InScriptStruct->GetPathName()), InPosition, InNodeName, bSetupUndoRedo);
+}
+
 URigVMSelectNode* URigVMController::AddSelectNode(const FString& InCPPType, const FName& InCPPTypeObjectPath, const FVector2D& InPosition, const FString& InNodeName, bool bSetupUndoRedo, bool bPrintPythonCommand)
 {
 	if (!IsValidGraph())
diff --git a/Engine/Source/Developer/RigVMDeveloper/Public/RigVMModel/RigVMController.h b/Engine/Source/Developer/RigVMDeveloper/Public/RigVMModel/RigVMController.h
--- a/Engine/Source/Developer/RigVMDeveloper/Public/RigVMModel/RigVMController.h
+++ b/Engine/Source/Developer/RigVMDeveloper/Public/RigVMModel/RigVMController.h
@@ -263,6 +263,8 @@ public:
 	// If nodes can be used to pick between two values based on a condition.
 	UFUNCTION(BlueprintCallable, Category = RigVMController)
 	URigVMIfNode* AddIfNode(const FString& InCPPType, const FName& InCPPTypeObjectPath, const FVector2D& InPosition = FVector2D::ZeroVector, const FString& InNodeName = TEXT(""), bool bSetupUndoRedo = true, bool bPrintPythonCommand = false);
+	UFUNCTION(BlueprintCallable, Category = RigVMController)
+	URigVMIfNode* AddIfNodeFromStruct(UScriptStruct* InScriptStruct, const FVector2D& InPosition = FVector2D::ZeroVector, const FString& InNodeName = TEXT(""), bool bSetupUndoRedo = true);
 
 	// Adds a select node to the graph.
 	// Select nodes can be used to pick between multiple values based on an index.
