Control Rig: Fix failed set expansion when pin is already expanded

diff --git a/Engine/Source/Developer/RigVMDeveloper/Private/RigVMModel/RigVMController.cpp b/Engine/Source/Developer/RigVMDeveloper/Private/RigVMModel/RigVMController.cpp
--- a/Engine/Source/Developer/RigVMDeveloper/Private/RigVMModel/RigVMController.cpp
+++ b/Engine/Source/Developer/RigVMDeveloper/Private/RigVMModel/RigVMController.cpp
@@ -5831,14 +5831,10 @@ bool URigVMController::SetPinExpansion(const FString& InPinPath, bool bIsExpande
 
 bool URigVMController::SetPinExpansion(URigVMPin* InPin, bool bIsExpanded, bool bSetupUndoRedo)
 {
-	if (InPin->GetSubPins().Num() == 0)
+	// If there is nothing to do, just return success
+	if (InPin->GetSubPins().Num() == 0 || InPin->IsExpanded() == bIsExpanded)
 	{
-		return false;
-	}
-
-	if (InPin->IsExpanded() == bIsExpanded)
-	{
-		return false;
+		return true;
 	}
 
 	URigVMGraph* Graph = GetGraph();
