Control Rig: Bring back support for direct pin links during expansion

Seems to have been removed by mistake in CL 17643082

diff --git a/Engine/Source/Developer/RigVMDeveloper/Private/RigVMModel/RigVMController.cpp b/Engine/Source/Developer/RigVMDeveloper/Private/RigVMModel/RigVMController.cpp
--- a/Engine/Source/Developer/RigVMDeveloper/Private/RigVMModel/RigVMController.cpp
+++ b/Engine/Source/Developer/RigVMDeveloper/Private/RigVMModel/RigVMController.cpp
@@ -4164,13 +4164,49 @@ TArray<URigVMNode*> URigVMController::ExpandLibraryNode(URigVMLibraryNode* InNod
 
 		for (const FString& FromEntryNodeTargetPinPath : FromEntryNodePair.Value)
 		{
-			URigVMPin* SourcePin = GetGraph()->FindPin(*RemappedSourcePinPath);
+			TArray<URigVMPin*> TargetPins;
+
+			URigVMPin* SourcePin = GetGraph()->FindPin(RemappedSourcePinPath);
 			URigVMPin* TargetPin = GetGraph()->FindPin(FromEntryNodeTargetPinPath);
-			if (SourcePin && TargetPin)
+
+			// potentially the target pin was on the entry node,
+			// so there's no node been added for it. we'll have to look into the remapped
+			// pins for the "FromLibraryNode" map.
+			if(TargetPin == nullptr)
 			{
-				if (!SourcePin->IsLinkedTo(TargetPin))
+				FString RemappedTargetPinPath = FromEntryNodeTargetPinPath;
+				FString ReturnNodeName, ReturnPinPath;
+				if (URigVMPin::SplitPinPathAtStart(RemappedTargetPinPath, ReturnNodeName, ReturnPinPath))
 				{
-					AddLink(SourcePin, TargetPin, false);
+					if(Cast<URigVMFunctionReturnNode>(InNode->GetContainedGraph()->FindNode(ReturnNodeName)))
+					{
+						if(FromLibraryNode.Contains(ReturnPinPath))
+						{
+							const TArray<FString>& FromLibraryNodeTargetPins = FromLibraryNode.FindChecked(ReturnPinPath);
+							for(const FString& FromLibraryNodeTargetPin : FromLibraryNodeTargetPins)
+							{
+								if(URigVMPin* MappedTargetPin = GetGraph()->FindPin(FromLibraryNodeTargetPin))
+								{
+									TargetPins.Add(MappedTargetPin);
+								}
+							}
+						}
+					}
+				}
+			}
+			else
+			{
+				TargetPins.Add(TargetPin);
+			}
+			
+			if (SourcePin)
+			{
+				for(URigVMPin* EachTargetPin : TargetPins)
+				{
+					if (!SourcePin->IsLinkedTo(EachTargetPin))
+					{
+						AddLink(SourcePin, EachTargetPin, false);
+					}
 				}
 			}
 		}
