Control Rig: Ensure name uniqueness for Exposed Pins

diff --git a/Engine/Source/Developer/RigVMDeveloper/Private/RigVMModel/RigVMController.cpp b/Engine/Source/Developer/RigVMDeveloper/Private/RigVMModel/RigVMController.cpp
--- a/Engine/Source/Developer/RigVMDeveloper/Private/RigVMModel/RigVMController.cpp
+++ b/Engine/Source/Developer/RigVMDeveloper/Private/RigVMModel/RigVMController.cpp
@@ -7398,6 +7398,11 @@ FName URigVMController::AddExposedPin(const FName& InPinName, ERigVMPinDirection
 
 	FName PinName = GetUniqueName(InPinName, [LibraryNode](const FName& InName) {
 
+		if(LibraryNode->FindPin(InName.ToString()) != nullptr)
+		{
+			return false;
+		}
+
 		const TArray<FRigVMGraphVariableDescription>& LocalVariables = LibraryNode->GetContainedGraph()->GetLocalVariables(true);
 		for(const FRigVMGraphVariableDescription& VariableDescription : LocalVariables)
 		{
@@ -12061,12 +12066,14 @@ void URigVMController::DestroyObject(UObject* InObjectToDestroy)
 void URigVMController::AddNodePin(URigVMNode* InNode, URigVMPin* InPin)
 {
 	ValidatePin(InPin);
+	check(!InNode->Pins.Contains(InPin));
 	InNode->Pins.Add(InPin);
 }
 
 void URigVMController::AddSubPin(URigVMPin* InParentPin, URigVMPin* InPin)
 {
 	ValidatePin(InPin);
+	check(!InParentPin->SubPins.Contains(InPin));
 	InParentPin->SubPins.Add(InPin);
 }
 
