Control Rig: Constant folding fix for functions.

Folding needs to take pin overrides (for example default values coming from a function ref node) into account

diff --git a/Engine/Source/Developer/RigVMDeveloper/Private/RigVMCompiler/RigVMAST.cpp b/Engine/Source/Developer/RigVMDeveloper/Private/RigVMCompiler/RigVMAST.cpp
--- a/Engine/Source/Developer/RigVMDeveloper/Private/RigVMCompiler/RigVMAST.cpp
+++ b/Engine/Source/Developer/RigVMDeveloper/Private/RigVMCompiler/RigVMAST.cpp
@@ -1909,6 +1909,11 @@ bool FRigVMParserAST::FoldConstantValuesToLiterals(URigVMGraph* InGraph, URigVMC
 	// let's build an temporary AST which has only those nodes
 	TSharedPtr<FRigVMParserAST> TempAST = MakeShareable(new FRigVMParserAST(InGraph, NodesToCompute));
 
+	// share the pin overrides with the constant folding AST to ensure
+	// the complete view of default values across function references is available
+	// in the subset AST.
+	TempAST->PinOverrides = PinOverrides;
+
 	// build the VM to run this AST
 	TMap<FString, FRigVMOperand> Operands;
 	URigVM* TempVM = NewObject<URigVM>(InGraph);
