Control Rig: only allow one execute pin per function

diff --git a/Engine/Source/Developer/RigVMDeveloper/Private/RigVMModel/RigVMController.cpp b/Engine/Source/Developer/RigVMDeveloper/Private/RigVMModel/RigVMController.cpp
--- a/Engine/Source/Developer/RigVMDeveloper/Private/RigVMModel/RigVMController.cpp
+++ b/Engine/Source/Developer/RigVMDeveloper/Private/RigVMModel/RigVMController.cpp
@@ -7396,6 +7396,21 @@ FName URigVMController::AddExposedPin(const FName& InPinName, ERigVMPinDirection
 		}
 	}
 
+	// only allow one exposed pin of type execute context per direction
+	if(const UScriptStruct* CPPTypeStruct = Cast<UScriptStruct>(CPPTypeObject))
+	{
+		if(CPPTypeStruct->IsChildOf(FRigVMExecuteContext::StaticStruct()))
+		{
+			for(URigVMPin* ExistingPin : LibraryNode->Pins)
+			{
+				if(ExistingPin->IsExecuteContext())
+				{
+					return NAME_None;
+				}
+			}
+		}
+	}
+
 	FName PinName = GetUniqueName(InPinName, [LibraryNode](const FName& InName) {
 
 		if(LibraryNode->FindPin(InName.ToString()) != nullptr)
@@ -7788,6 +7803,30 @@ bool URigVMController::ChangeExposedPinType(const FName& InPinName, const FStrin
 	{
 		return false;
 	}
+	
+	// only allow one exposed pin of type execute context per direction
+	if (!InCPPTypeObjectPath.IsNone())
+	{
+		if(UObject* CPPTypeObject = URigVMPin::FindObjectFromCPPTypeObjectPath<UObject>(InCPPTypeObjectPath.ToString()))
+		{
+			if(const UScriptStruct* CPPTypeStruct = Cast<UScriptStruct>(CPPTypeObject))
+			{
+				if(CPPTypeStruct->IsChildOf(FRigVMExecuteContext::StaticStruct()))
+				{
+					for(URigVMPin* ExistingPin : LibraryNode->Pins)
+					{
+						if(ExistingPin != Pin)
+						{
+							if(ExistingPin->IsExecuteContext())
+							{
+								return false;
+							}
+						}
+					}
+				}
+			}
+		}
+	}
 
 	if(bSetupUndoRedo)
 	{
