maintain links when deleting nodes via Shift + Del/Backspace

diff --git a/Engine/Plugins/Animation/ControlRig/Source/ControlRigEditor/Private/Editor/ControlRigEditor.cpp b/Engine/Plugins/Animation/ControlRig/Source/ControlRigEditor/Private/Editor/ControlRigEditor.cpp
--- a/Engine/Plugins/Animation/ControlRig/Source/ControlRigEditor/Private/Editor/ControlRigEditor.cpp
+++ b/Engine/Plugins/Animation/ControlRig/Source/ControlRigEditor/Private/Editor/ControlRigEditor.cpp
@@ -2111,7 +2111,8 @@ void FControlRigEditor::DeleteSelectedNodes()
 				AnalyticsTrackNodeEvent(GetBlueprintObj(), Node, true);
 				if (UControlRigGraphNode* RigNode = Cast<UControlRigGraphNode>(Node))
 				{
-					if(GetFocusedController()->RemoveNodeByName(*RigNode->ModelNodePath, true, false, true))
+					const bool bReconnectPins = (FSlateApplication::Get().GetModifierKeys().IsShiftDown());
+					if(GetFocusedController()->RemoveNodeByName(*RigNode->ModelNodePath, true, false, true, bReconnectPins))
 					{
 						DeletedAnything = true;
 					}
diff --git a/Engine/Source/Developer/RigVMDeveloper/Private/RigVMModel/RigVMController.cpp b/Engine/Source/Developer/RigVMDeveloper/Private/RigVMModel/RigVMController.cpp
--- a/Engine/Source/Developer/RigVMDeveloper/Private/RigVMModel/RigVMController.cpp
+++ b/Engine/Source/Developer/RigVMDeveloper/Private/RigVMModel/RigVMController.cpp
@@ -4675,7 +4675,7 @@ void URigVMController::RefreshFunctionPins(URigVMNode* InNode, bool bNotify)
 	}
 }
 
-bool URigVMController::RemoveNode(URigVMNode* InNode, bool bSetupUndoRedo, bool bRecursive, bool bPrintPythonCommand)
+bool URigVMController::RemoveNode(URigVMNode* InNode, bool bSetupUndoRedo, bool bRecursive, bool bPrintPythonCommand, bool bRelinkPins)
 {
 	if (!IsValidNodeForGraph(InNode))
 	{
@@ -4847,6 +4847,12 @@ bool URigVMController::RemoveNode(URigVMNode* InNode, bool bSetupUndoRedo, bool
 		}
 	}
 
+	// try to reconnect source and target nodes based on the current links
+	if (bRelinkPins)
+	{
+		RelinkSourceAndTargetPins(InNode, bSetupUndoRedo);
+	}
+	
 	if (bSetupUndoRedo || bRecursive)
 	{
 		SelectNode(InNode, false, bSetupUndoRedo);
@@ -4913,11 +4919,12 @@ bool URigVMController::RemoveNode(URigVMNode* InNode, bool bSetupUndoRedo, bool
 		else
 		{
 			const FString NodePath = GetSanitizedPinPath(InNode->GetNodePath());
+
+			FString PythonCmd = FString::Printf(TEXT("blueprint.get_controller_by_name('%s')."), *GraphName );
+			PythonCmd += bRelinkPins ? FString::Printf(TEXT("remove_node_by_name('%s', relink_pins=True)"), *NodePath) :
+									   FString::Printf(TEXT("remove_node_by_name('%s')"), *NodePath);
 			
-			RigVMPythonUtils::Print(GetGraphOuterName(), 
-								FString::Printf(TEXT("blueprint.get_controller_by_name('%s').remove_node_by_name('%s')"),
-												*GraphName,
-												*NodePath));
+			RigVMPythonUtils::Print(GetGraphOuterName(), PythonCmd );
 		}
 	}
 
@@ -4945,7 +4952,7 @@ bool URigVMController::RemoveNode(URigVMNode* InNode, bool bSetupUndoRedo, bool
 	return true;
 }
 
-bool URigVMController::RemoveNodeByName(const FName& InNodeName, bool bSetupUndoRedo, bool bRecursive, bool bPrintPythonCommand)
+bool URigVMController::RemoveNodeByName(const FName& InNodeName, bool bSetupUndoRedo, bool bRecursive, bool bPrintPythonCommand, bool bRelinkPins)
 {
 	if (!IsValidGraph())
 	{
@@ -4955,7 +4962,7 @@ bool URigVMController::RemoveNodeByName(const FName& InNodeName, bool bSetupUndo
 	URigVMGraph* Graph = GetGraph();
 	check(Graph);
 
-	return RemoveNode(Graph->FindNodeByName(InNodeName), bSetupUndoRedo, bRecursive, bPrintPythonCommand);
+	return RemoveNode(Graph->FindNodeByName(InNodeName), bSetupUndoRedo, bRecursive, bPrintPythonCommand, bRelinkPins);
 }
 
 bool URigVMController::RenameNode(URigVMNode* InNode, const FName& InNewName, bool bSetupUndoRedo, bool bPrintPythonCommand)
@@ -7135,6 +7142,64 @@ bool URigVMController::AddLink(URigVMPin* OutputPin, URigVMPin* InputPin, bool b
 	return true;
 }
 
+void URigVMController::RelinkSourceAndTargetPins(URigVMNode* Node, bool bSetupUndoRedo)
+{
+	TArray<URigVMPin*> SourcePins;
+	TArray<URigVMPin*> TargetPins;
+	TArray<URigVMLink*> LinksToRemove;
+
+	// store source and target links 
+	const TArray<URigVMLink*> RigVMLinks = Node->GetLinks();
+	for (URigVMLink* Link: RigVMLinks)
+	{
+		URigVMPin* SrcPin = Link->GetSourcePin();
+		if (SrcPin && SrcPin->GetNode() != Node)
+		{
+			SourcePins.AddUnique(SrcPin);
+			LinksToRemove.AddUnique(Link);
+		}
+
+		URigVMPin* DstPin = Link->GetTargetPin();
+		if (DstPin && DstPin->GetNode() != Node)
+		{
+			TargetPins.AddUnique(DstPin);
+			LinksToRemove.AddUnique(Link);
+		}
+	}
+
+	if( SourcePins.Num() > 0 && TargetPins.Num() > 0 )
+	{
+		// remove previous links 
+		for (URigVMLink* Link: LinksToRemove)
+		{
+			BreakLink(Link->GetSourcePin(), Link->GetTargetPin(), bSetupUndoRedo); 
+		}
+
+		// relink pins if feasible 
+		TArray<bool> TargetHandled;
+		TargetHandled.AddZeroed(TargetPins.Num());
+		for (URigVMPin* Src: SourcePins)
+		{
+			for (int32 Index = 0; Index < TargetPins.Num(); Index++)
+			{
+				if (!TargetHandled[Index])
+				{
+					if (URigVMPin::CanLink(Src, TargetPins[Index], nullptr, nullptr))
+					{
+						// execute pins can be linked to one target only so link to the 1st compatible target
+						const bool bNeedNewLink = Src->IsExecuteContext() ? (Src->GetTargetLinks().Num() == 0) : true;
+						if (bNeedNewLink)
+						{
+							AddLink(Src, TargetPins[Index], bSetupUndoRedo);
+							TargetHandled[Index] = true;								
+						}
+					}
+				}
+			}
+		}
+	}
+}
+
 bool URigVMController::BreakLink(const FString& InOutputPinPath, const FString& InInputPinPath, bool bSetupUndoRedo, bool bPrintPythonCommand)
 {
 	if(!IsValidGraph())
diff --git a/Engine/Source/Developer/RigVMDeveloper/Public/RigVMModel/RigVMController.h b/Engine/Source/Developer/RigVMDeveloper/Public/RigVMModel/RigVMController.h
--- a/Engine/Source/Developer/RigVMDeveloper/Public/RigVMModel/RigVMController.h
+++ b/Engine/Source/Developer/RigVMDeveloper/Public/RigVMModel/RigVMController.h
@@ -396,13 +396,13 @@ public:
 	// Removes a node from the graph
 	// This causes a NodeRemoved modified event.
 	UFUNCTION(BlueprintCallable, Category = RigVMController)
-	bool RemoveNode(URigVMNode* InNode, bool bSetupUndoRedo = true, bool bRecursive = false, bool bPrintPythonCommand = false);
+	bool RemoveNode(URigVMNode* InNode, bool bSetupUndoRedo = true, bool bRecursive = false, bool bPrintPythonCommand = false, bool bRelinkPins = false);
 
 	// Removes a node from the graph given the node's name.
 	// This causes a NodeRemoved modified event.
 	UFUNCTION(BlueprintCallable, Category = RigVMController)
-	bool RemoveNodeByName(const FName& InNodeName, bool bSetupUndoRedo = true, bool bRecursive = false, bool bPrintPythonCommand = false);
-
+	bool RemoveNodeByName(const FName& InNodeName, bool bSetupUndoRedo = true, bool bRecursive = false, bool bPrintPythonCommand = false, bool bRelinkPins = false);
+	
 	// Renames a node in the graph
 	// This causes a NodeRenamed modified event.
 	UFUNCTION(BlueprintCallable, Category = RigVMController)
@@ -806,6 +806,9 @@ private:
 	bool MakeBindingsFromVariableNode(URigVMVariableNode* InNode, bool bSetupUndoRedo);
 	bool PromotePinToVariable(URigVMPin* InPin, bool bCreateVariableNode, const FVector2D& InNodePosition, bool bSetupUndoRedo);
 
+	// try to reconnect source and target pins after a node deletion
+	void RelinkSourceAndTargetPins(URigVMNode* RigNode, bool bSetupUndoRedo = true);
+	
 public:
 	bool AddLink(URigVMPin* OutputPin, URigVMPin* InputPin, bool bSetupUndoRedo = true);
 	bool BreakLink(URigVMPin* OutputPin, URigVMPin* InputPin, bool bSetupUndoRedo = true);
