Control Rig: Fix debugger not halting in local variable nodes

diff --git a/Engine/Source/Developer/RigVMDeveloper/Private/RigVMCompiler/RigVMCompiler.cpp b/Engine/Source/Developer/RigVMDeveloper/Private/RigVMCompiler/RigVMCompiler.cpp
--- a/Engine/Source/Developer/RigVMDeveloper/Private/RigVMCompiler/RigVMCompiler.cpp
+++ b/Engine/Source/Developer/RigVMDeveloper/Private/RigVMCompiler/RigVMCompiler.cpp
@@ -1017,26 +1017,21 @@ void URigVMCompiler::TraverseAssign(const FRigVMAssignExprAST* InExpr, FRigVMCom
 			int32 InstructionIndex = WorkData.VM->GetByteCode().GetNumInstructions() - 1;
 			if (Settings.SetupNodeInstructionIndex)
 			{
-				if (Source.GetMemoryType() == ERigVMMemoryType::External)
+				if (URigVMPin* SourcePin = InExpr->GetSourcePin())
 				{
-					if (URigVMPin* SourcePin = InExpr->GetSourcePin())
+					if (URigVMVariableNode* VariableNode = Cast<URigVMVariableNode>(SourcePin->GetNode()))
 					{
-						if (URigVMVariableNode* VariableNode = Cast<URigVMVariableNode>(SourcePin->GetNode()))
-						{
-							const FRigVMCallstack Callstack = SourceExpr->GetProxy().GetSibling(VariableNode).GetCallstack();
-							WorkData.VM->GetByteCode().SetSubject(InstructionIndex, Callstack.GetCallPath(), Callstack.GetStack());
-						}
+						const FRigVMCallstack Callstack = SourceExpr->GetProxy().GetSibling(VariableNode).GetCallstack();
+						WorkData.VM->GetByteCode().SetSubject(InstructionIndex, Callstack.GetCallPath(), Callstack.GetStack());
 					}
 				}
-				if (Target.GetMemoryType() == ERigVMMemoryType::External)
+
+				if (URigVMPin* TargetPin = InExpr->GetTargetPin())
 				{
-					if (URigVMPin* TargetPin = InExpr->GetTargetPin())
+					if (URigVMVariableNode* VariableNode = Cast<URigVMVariableNode>(TargetPin->GetNode()))
 					{
-						if (URigVMVariableNode* VariableNode = Cast<URigVMVariableNode>(TargetPin->GetNode()))
-						{
-							const FRigVMCallstack Callstack = TargetExpr->GetProxy().GetSibling(VariableNode).GetCallstack();
-							WorkData.VM->GetByteCode().SetSubject(InstructionIndex, Callstack.GetCallPath(), Callstack.GetStack());
-						}
+						const FRigVMCallstack Callstack = TargetExpr->GetProxy().GetSibling(VariableNode).GetCallstack();
+						WorkData.VM->GetByteCode().SetSubject(InstructionIndex, Callstack.GetCallPath(), Callstack.GetStack());
 					}
 				}
 			}
