Control Rig: Fix discrepancy between release and debug mode due to assignment folding

diff --git a/Engine/Source/Developer/RigVMDeveloper/Private/RigVMCompiler/RigVMCompiler.cpp b/Engine/Source/Developer/RigVMDeveloper/Private/RigVMCompiler/RigVMCompiler.cpp
--- a/Engine/Source/Developer/RigVMDeveloper/Private/RigVMCompiler/RigVMCompiler.cpp
+++ b/Engine/Source/Developer/RigVMDeveloper/Private/RigVMCompiler/RigVMCompiler.cpp
@@ -2783,6 +2783,7 @@ FRigVMOperand URigVMCompiler::FindOrAddRegister(const FRigVMVarExprAST* InVarExp
 	{
 		const TMap<FRigVMASTProxy, FRigVMASTProxy>& SharedOperandPins = InVarExpr->GetParser()->SharedOperandPins;
 		TArray<FRigVMASTProxy> PinProxies = FindProxiesWithSharedOperand(InVarExpr);
+		ensure(!PinProxies.IsEmpty());
 
 		for (const FRigVMASTProxy& Proxy : PinProxies)
 		{
@@ -2860,6 +2861,11 @@ TArray<FRigVMASTProxy> URigVMCompiler::FindProxiesWithSharedOperand(const FRigVM
 		}
 	}
 
+	if (PinProxies.IsEmpty())
+	{
+		PinProxies.Add(InVarExpr->GetProxy());
+	}
+
 	return PinProxies;
 }
 
