Control Rig: Fixed a RigVM crash caused by us using the wrong index/elementsize for ext variable's register offset

diff --git a/Engine/Source/Developer/RigVMDeveloper/Private/RigVMCompiler/RigVMCompiler.cpp b/Engine/Source/Developer/RigVMDeveloper/Private/RigVMCompiler/RigVMCompiler.cpp
--- a/Engine/Source/Developer/RigVMDeveloper/Private/RigVMCompiler/RigVMCompiler.cpp
+++ b/Engine/Source/Developer/RigVMDeveloper/Private/RigVMCompiler/RigVMCompiler.cpp
@@ -789,7 +789,7 @@ void URigVMCompiler::TraverseAssign(const FRigVMAssignExprAST* InExpr, FRigVMCom
 					{
 						const FRigVMExternalVariable& ExternalVariable = VM->GetExternalVariables()[Operand.GetRegisterIndex()];
 						UScriptStruct* ScriptStruct = CastChecked<UScriptStruct>(ExternalVariable.TypeObject);
-						int32 RegisterOffset = VM->GetWorkMemory().GetOrAddRegisterOffset(Operand.GetRegisterIndex(), ScriptStruct, SegmentPath, ArrayIndex == INDEX_NONE ? 0 : ArrayIndex);
+						int32 RegisterOffset = VM->GetWorkMemory().GetOrAddRegisterOffset(INDEX_NONE, ScriptStruct, SegmentPath, ArrayIndex == INDEX_NONE ? 0 : ArrayIndex, ExternalVariable.Size);
 						Operand = FRigVMOperand(ERigVMMemoryType::External, Operand.GetRegisterIndex(), RegisterOffset);
 					}
 					else
diff --git a/Engine/Source/Runtime/RigVM/Private/RigVMCore/RigVMMemory.cpp b/Engine/Source/Runtime/RigVM/Private/RigVMCore/RigVMMemory.cpp
--- a/Engine/Source/Runtime/RigVM/Private/RigVMCore/RigVMMemory.cpp
+++ b/Engine/Source/Runtime/RigVM/Private/RigVMCore/RigVMMemory.cpp
@@ -1691,8 +1691,12 @@ int32 FRigVMMemoryContainer::GetOrAddRegisterOffset(int32 InRegisterIndex, UScri
 		return INDEX_NONE;
 	}
 
-	if (InElementSize == 0 && Registers.IsValidIndex(InRegisterIndex))
+	// if this is a register offset for a external variable
+	// the register index is expected to be INDEX_NONE
+	// and it is also expected that InElementSize != 0
+	if (InElementSize == 0)
 	{
+		ensure(Registers.IsValidIndex(InRegisterIndex)); 
 		InElementSize = (int32)Registers[InRegisterIndex].ElementSize;
 	}
 
