Control Rig: RigVM UObject support

diff --git a/Engine/Plugins/Animation/ControlRig/Source/ControlRigDeveloper/Private/Graph/ControlRigGraphNode.cpp b/Engine/Plugins/Animation/ControlRig/Source/ControlRigDeveloper/Private/Graph/ControlRigGraphNode.cpp
--- a/Engine/Plugins/Animation/ControlRig/Source/ControlRigDeveloper/Private/Graph/ControlRigGraphNode.cpp
+++ b/Engine/Plugins/Animation/ControlRig/Source/ControlRigDeveloper/Private/Graph/ControlRigGraphNode.cpp
@@ -889,6 +889,19 @@ void UControlRigGraphNode::CopyPinDefaultsToModel(UEdGraphPin* Pin, bool bUndo,
 		}
 
 		FString DefaultValue = Pin->DefaultValue;
+
+		if(DefaultValue.IsEmpty() && (
+			Pin->PinType.PinCategory == UEdGraphSchema_K2::PC_Object ||
+			Pin->PinType.PinCategory == UEdGraphSchema_K2::PC_SoftObject ||
+			Pin->PinType.PinCategory == UEdGraphSchema_K2::AllObjectTypes
+			))
+		{
+			if(Pin->DefaultObject)
+			{
+				DefaultValue = Pin->DefaultObject->GetPathName();
+			}
+		}
+		
 		if (DefaultValue == FName(NAME_None).ToString() && Pin->PinType.PinSubCategory == UEdGraphSchema_K2::PC_Name)
 		{
 			DefaultValue = FString();
@@ -1141,6 +1154,11 @@ FEdGraphPinType UControlRigGraphNode::GetPinTypeForModelPin(URigVMPin* InModelPi
 		PinType.PinCategory = UEdGraphSchema_K2::PC_Struct;
 		PinType.PinSubCategoryObject = InModelPin->GetScriptStruct();
 	}
+	else if (InModelPin->IsUObject())
+	{
+		PinType.PinCategory = UEdGraphSchema_K2::PC_Object;
+		PinType.PinSubCategoryObject = InModelPin->GetCPPTypeObject();
+	}
 	else if (InModelPin->GetEnum() != nullptr)
 		{
 		PinType.PinCategory = UEdGraphSchema_K2::PC_Byte;
diff --git a/Engine/Plugins/Animation/ControlRig/Source/ControlRigDeveloper/Private/Graph/ControlRigGraphSchema.cpp b/Engine/Plugins/Animation/ControlRig/Source/ControlRigDeveloper/Private/Graph/ControlRigGraphSchema.cpp
--- a/Engine/Plugins/Animation/ControlRig/Source/ControlRigDeveloper/Private/Graph/ControlRigGraphSchema.cpp
+++ b/Engine/Plugins/Animation/ControlRig/Source/ControlRigDeveloper/Private/Graph/ControlRigGraphSchema.cpp
@@ -782,6 +782,10 @@ bool UControlRigGraphSchema::SupportsPinType(UScriptStruct* ScriptStruct) const
 				continue;
 			}
 		}
+		else if (CastField<FObjectProperty>(Property))
+		{
+			continue;
+		}
 	
 		return false;
 	}
@@ -809,6 +813,16 @@ bool UControlRigGraphSchema::SupportsPinType(TWeakPtr<const FEdGraphSchemaAction
 		return true;
 	}
 
+	if (PinType.PinCategory == UEdGraphSchema_K2::PC_Object ||
+		PinType.PinCategory == UEdGraphSchema_K2::PC_SoftObject ||
+		PinType.PinCategory == UEdGraphSchema_K2::AllObjectTypes)
+	{
+		if (PinType.PinSubCategoryObject.IsValid())
+		{
+			return PinType.PinSubCategoryObject->IsA<UClass>();
+		}
+	}
+
 	if (PinType.PinCategory == UEdGraphSchema_K2::PC_Struct)
 	{
 		if (UScriptStruct* ScriptStruct = Cast<UScriptStruct>(PinType.PinSubCategoryObject))
diff --git a/Engine/Source/Developer/RigVMDeveloper/Private/RigVMCompiler/RigVMCompiler.cpp b/Engine/Source/Developer/RigVMDeveloper/Private/RigVMCompiler/RigVMCompiler.cpp
--- a/Engine/Source/Developer/RigVMDeveloper/Private/RigVMCompiler/RigVMCompiler.cpp
+++ b/Engine/Source/Developer/RigVMDeveloper/Private/RigVMCompiler/RigVMCompiler.cpp
@@ -2773,7 +2773,7 @@ FRigVMOperand URigVMCompiler::FindOrAddRegister(const FRigVMVarExprAST* InVarExp
 		if(bIsDebugValue)
 		{
 			// debug values are always stored as arrays
-			CPPType = FString::Printf(TEXT("TArray<%s>"), *CPPType);
+			CPPType = FString::Printf(URigVMController::TArrayTemplate, *CPPType);
 			JoinedDefaultValue = URigVMPin::GetDefaultValueForArray({ JoinedDefaultValue });
 		}
 		else if(Pin->GetDirection() == ERigVMPinDirection::Hidden && Pin->GetNode()->IsA<URigVMUnitNode>())
@@ -2784,7 +2784,7 @@ FRigVMOperand URigVMCompiler::FindOrAddRegister(const FRigVMVarExprAST* InVarExp
 
 			if (!Property->HasMetaData(FRigVMStruct::SingletonMetaName))
 			{
-				CPPType = FString::Printf(TEXT("TArray<%s>"), *CPPType);
+				CPPType = FString::Printf(URigVMController::TArrayTemplate, *CPPType);
 				JoinedDefaultValue = URigVMPin::GetDefaultValueForArray({ JoinedDefaultValue });
 			}
 		}
diff --git a/Engine/Source/Developer/RigVMDeveloper/Private/RigVMModel/RigVMController.cpp b/Engine/Source/Developer/RigVMDeveloper/Private/RigVMModel/RigVMController.cpp
--- a/Engine/Source/Developer/RigVMDeveloper/Private/RigVMModel/RigVMController.cpp
+++ b/Engine/Source/Developer/RigVMDeveloper/Private/RigVMModel/RigVMController.cpp
@@ -29,6 +29,10 @@
 #endif
 
 TMap<URigVMController::FControlRigStructPinRedirectorKey, FString> URigVMController::PinPathCoreRedirectors;
+const TCHAR URigVMController::TArrayPrefix[] = TEXT("TArray<");
+const TCHAR URigVMController::TObjectPtrPrefix[] = TEXT("TObjectPtr<");
+const TCHAR URigVMController::TArrayTemplate[] = TEXT("TArray<%s>");
+const TCHAR URigVMController::TObjectPtrTemplate[] = TEXT("TObjectPtr<%s%s>");
 
 FRigVMControllerCompileBracketScope::FRigVMControllerCompileBracketScope(URigVMController* InController)
 : Graph(nullptr), bSuspendNotifications(InController->bSuspendNotifications)
@@ -978,11 +982,7 @@ URigVMVariableNode* URigVMController::AddVariableNode(const FName& InVariableNam
 		InCPPTypeObject = URigVMPin::FindObjectFromCPPTypeObjectPath<UObject>(InCPPType);
 	}
 
-	FString CPPType = InCPPType;
-	if (UScriptStruct* ScriptStruct = Cast<UScriptStruct>(InCPPTypeObject))
-	{
-		CPPType = ScriptStruct->GetStructCPPName();
-	}
+	FString CPPType = PostProcessCPPType(InCPPType, InCPPTypeObject, InCPPType.StartsWith(TArrayPrefix));
 	
 	FString Name = GetValidNodeName(InNodeName.IsEmpty() ? FString(TEXT("VariableNode")) : InNodeName);
 	URigVMVariableNode* Node = NewObject<URigVMVariableNode>(Graph, *Name);
@@ -1014,16 +1014,21 @@ URigVMVariableNode* URigVMController::AddVariableNode(const FName& InVariableNam
 		ValuePin->CPPTypeObject = ExternalVariable.TypeObject;
 		ValuePin->bIsDynamicArray = ExternalVariable.bIsArray;
 
-		if(ValuePin->bIsDynamicArray && !ValuePin->CPPType.StartsWith(TEXT("TArray<")))
+		if(ValuePin->bIsDynamicArray && !ValuePin->CPPType.StartsWith(TArrayPrefix))
 		{
-			ValuePin->CPPType = FString::Printf(TEXT("TArray<%s>"), *ValuePin->CPPType);
+			ValuePin->CPPType = FString::Printf(TArrayTemplate, *ValuePin->CPPType);
 		}
 	}
 	else
 	{
 		ValuePin->CPPType = CPPType;
 
-		if (UScriptStruct* ScriptStruct = Cast<UScriptStruct>(InCPPTypeObject))
+		if (UClass* Class = Cast<UClass>(InCPPTypeObject))
+		{
+			ValuePin->CPPTypeObject = Class;
+			ValuePin->CPPTypeObjectPath = *ValuePin->CPPTypeObject->GetPathName();
+		}
+		else if (UScriptStruct* ScriptStruct = Cast<UScriptStruct>(InCPPTypeObject))
 		{
 			ValuePin->CPPTypeObject = ScriptStruct;
 			ValuePin->CPPTypeObjectPath = *ValuePin->CPPTypeObject->GetPathName();
@@ -1604,28 +1609,18 @@ URigVMParameterNode* URigVMController::AddParameterNode(const FName& InParameter
 
 	if (DefaultValuePin)
 	{
-		DefaultValuePin->CPPType = InCPPType;
-	}
-	ValuePin->CPPType = InCPPType;
-
-	if (UScriptStruct* ScriptStruct = Cast<UScriptStruct>(InCPPTypeObject))
-	{
-		if (DefaultValuePin)
+		DefaultValuePin->CPPType = PostProcessCPPType(InCPPType, InCPPTypeObject, InCPPType.StartsWith(TArrayPrefix));
+		DefaultValuePin->CPPTypeObject = InCPPTypeObject;
+		if(DefaultValuePin->CPPTypeObject)
 		{
-			DefaultValuePin->CPPTypeObject = ScriptStruct;
 			DefaultValuePin->CPPTypeObjectPath = *DefaultValuePin->CPPTypeObject->GetPathName();
 		}
-		ValuePin->CPPTypeObject = ScriptStruct;
-		ValuePin->CPPTypeObjectPath = *ValuePin->CPPTypeObject->GetPathName();
 	}
-	else if (UEnum* Enum = Cast<UEnum>(InCPPTypeObject))
+	
+	ValuePin->CPPType = PostProcessCPPType(InCPPType, InCPPTypeObject, InCPPType.StartsWith(TArrayPrefix));
+	ValuePin->CPPTypeObject = InCPPTypeObject;
+	if(ValuePin->CPPTypeObject)
 	{
-		if (DefaultValuePin)
-		{
-			DefaultValuePin->CPPTypeObject = Enum;
-			DefaultValuePin->CPPTypeObjectPath = *DefaultValuePin->CPPTypeObject->GetPathName();
-		}
-		ValuePin->CPPTypeObject = Enum;
 		ValuePin->CPPTypeObjectPath = *ValuePin->CPPTypeObject->GetPathName();
 	}
 
@@ -6062,7 +6057,11 @@ bool URigVMController::SetPinDefaultValue(const FString& InPinPath, const FStrin
 bool URigVMController::SetPinDefaultValue(URigVMPin* InPin, const FString& InDefaultValue, bool bResizeArrays, bool bSetupUndoRedo, bool bMergeUndoAction, bool bNotify)
 {
 	check(InPin);
-	ensure(!InDefaultValue.IsEmpty());
+
+	if(!InPin->IsUObject())
+	{
+		ensure(!InDefaultValue.IsEmpty());
+	}
 
 	TGuardValue<bool> Guard(bSuspendNotifications, !bNotify);
 
@@ -7496,7 +7495,7 @@ FName URigVMController::AddExposedPin(const FName& InPinName, ERigVMPinDirection
 	}, false, true);
 
 	URigVMPin* Pin = NewObject<URigVMPin>(LibraryNode, PinName);
-	Pin->CPPType = InCPPType;
+	Pin->CPPType = PostProcessCPPType(InCPPType, CPPTypeObject, InCPPType.StartsWith(TArrayPrefix));
 	Pin->CPPTypeObjectPath = InCPPTypeObjectPath;
 	Pin->bIsConstant = false;
 	Pin->Direction = InDirection;
@@ -9363,7 +9362,6 @@ URigVMIfNode* URigVMController::AddIfNode(const FString& InCPPType, const FName&
 
 	ensure(!InCPPType.IsEmpty());
 
-	FString CPPType = InCPPType;
 	UObject* CPPTypeObject = nullptr;
 	if(!InCPPTypeObjectPath.IsNone())
 	{
@@ -9375,6 +9373,8 @@ URigVMIfNode* URigVMController::AddIfNode(const FString& InCPPType, const FName&
 		}
 	}
 
+	FString CPPType = PostProcessCPPType(InCPPType, CPPTypeObject, InCPPType.StartsWith(TArrayPrefix));
+
 	FString DefaultValue;
 	if(UScriptStruct* ScriptStruct = Cast<UScriptStruct>(CPPTypeObject))
 	{
@@ -9384,8 +9384,6 @@ URigVMIfNode* URigVMController::AddIfNode(const FString& InCPPType, const FName&
 			return nullptr;
 		}
 		CreateDefaultValueForStructIfRequired(ScriptStruct, DefaultValue);
-
-		CPPType = ScriptStruct->GetStructCPPName();
 	}
 	
 	FString Name = GetValidNodeName(InNodeName.IsEmpty() ? FString(TEXT("IfNode")) : InNodeName);
@@ -9479,7 +9477,6 @@ URigVMSelectNode* URigVMController::AddSelectNode(const FString& InCPPType, cons
 
 	ensure(!InCPPType.IsEmpty());
 
-	FString CPPType = InCPPType;
 	UObject* CPPTypeObject = nullptr;
 	if (!InCPPTypeObjectPath.IsNone())
 	{
@@ -9491,6 +9488,8 @@ URigVMSelectNode* URigVMController::AddSelectNode(const FString& InCPPType, cons
 		}
 	}
 
+	FString CPPType = PostProcessCPPType(InCPPType, CPPTypeObject, InCPPType.StartsWith(TArrayPrefix));
+
 	FString DefaultValue;
 	if (UScriptStruct* ScriptStruct = Cast<UScriptStruct>(CPPTypeObject))
 	{
@@ -9500,8 +9499,6 @@ URigVMSelectNode* URigVMController::AddSelectNode(const FString& InCPPType, cons
 			return nullptr;
 		}
 		CreateDefaultValueForStructIfRequired(ScriptStruct, DefaultValue);
-
-		CPPType = ScriptStruct->GetStructCPPName();
 	}
 
 	FString Name = GetValidNodeName(InNodeName.IsEmpty() ? FString(TEXT("IfNode")) : InNodeName);
@@ -9514,7 +9511,7 @@ URigVMSelectNode* URigVMController::AddSelectNode(const FString& InCPPType, cons
 	AddNodePin(Node, IndexPin);
 
 	URigVMPin* ValuePin = NewObject<URigVMPin>(Node, *URigVMSelectNode::ValueName);
-	ValuePin->CPPType = FString::Printf(TEXT("TArray<%s>"), *CPPType);
+	ValuePin->CPPType = FString::Printf(TArrayTemplate, *CPPType);
 	ValuePin->CPPTypeObject = CPPTypeObject;
 	ValuePin->CPPTypeObjectPath = InCPPTypeObjectPath;
 	ValuePin->Direction = ERigVMPinDirection::Input;
@@ -9758,11 +9755,7 @@ URigVMArrayNode* URigVMController::AddArrayNode(ERigVMOpCode InOpCode, const FSt
 		InCPPTypeObject = URigVMPin::FindObjectFromCPPTypeObjectPath<UObject>(InCPPType);
 	}
 
-	FString CPPType = InCPPType;
-	if (UScriptStruct* ScriptStruct = Cast<UScriptStruct>(InCPPTypeObject))
-	{
-		CPPType = ScriptStruct->GetStructCPPName();
-	}
+	FString CPPType = PostProcessCPPType(InCPPType, InCPPTypeObject, InCPPType.StartsWith(TArrayPrefix));
 	
 	const FString Name = GetValidNodeName(InNodeName.IsEmpty() ? FString(TEXT("ArrayNode")) : InNodeName);
 	URigVMArrayNode* Node = NewObject<URigVMArrayNode>(Graph, *Name);
@@ -9780,13 +9773,9 @@ URigVMArrayNode* URigVMController::AddArrayNode(ERigVMOpCode InOpCode, const FSt
 			{
 				Pin->CPPTypeObjectPath = *Pin->CPPTypeObject->GetPathName();
 			}
-			if(UScriptStruct* Struct = Cast<UScriptStruct>(Pin->CPPTypeObject))
-			{
-				Pin->CPPType = Struct->GetStructCPPName();
-			}
-			if(bIsArray && !Pin->CPPType.StartsWith(TEXT("TArray<")))
+			if(bIsArray && !Pin->CPPType.StartsWith(TArrayPrefix))
 			{
-				Pin->CPPType = FString::Printf(TEXT("TArray<%s>"), *Pin->CPPType);
+				Pin->CPPType = FString::Printf(TArrayTemplate, *Pin->CPPType);
 			}
 			Pin->Direction = InDirection;
 			Pin->bIsDynamicArray = bIsArray;
@@ -10557,6 +10546,10 @@ void URigVMController::ConfigurePinFromProperty(FProperty* InProperty, URigVMPin
 	{
 		InOutPin->CPPTypeObject = StructProperty->Struct;
 	}
+	else if (FObjectProperty* ObjectProperty = CastField<FObjectProperty>(PropertyForType))
+	{
+		InOutPin->CPPTypeObject = ObjectProperty->PropertyClass;
+	}
 	else if (FEnumProperty* EnumProperty = CastField<FEnumProperty>(PropertyForType))
 	{
 		InOutPin->CPPTypeObject = EnumProperty->GetEnum();
@@ -11780,6 +11773,34 @@ void URigVMController::PostProcessDefaultValue(URigVMPin* Pin, FString& OutDefau
 	}
 }
 
+FString URigVMController::PostProcessCPPType(const FString& InCPPType, UObject* InCPPTypeObject, bool bIsArray)
+{
+	FString CPPType = InCPPType;
+	
+	if (const UClass* Class = Cast<UClass>(InCPPTypeObject))
+	{
+		CPPType = FString::Printf(TObjectPtrTemplate, Class->GetPrefixCPP(), *Class->GetName());
+	}
+	else if (const UScriptStruct* ScriptStruct = Cast<UScriptStruct>(InCPPTypeObject))
+	{
+		CPPType = ScriptStruct->GetStructCPPName();
+	}
+	else if (UEnum* Enum = Cast<UEnum>(InCPPTypeObject))
+	{
+		CPPType = Enum->CppType;
+	}
+
+	if(bIsArray)
+	{
+		if(!CPPType.StartsWith(TArrayPrefix))
+		{
+			CPPType = FString::Printf(TArrayTemplate, *CPPType);
+		}
+	}
+	
+	return CPPType;
+}
+
 /*
 void URigVMController::PotentiallyResolvePrototypeNode(URigVMPrototypeNode* InNode, bool bSetupUndoRedo)
 {
@@ -11929,7 +11950,7 @@ void URigVMController::ResolveUnknownTypePin(URigVMPin* InPinToResolve, const UR
 	FString ResolvedCPPType = InTemplatePin->IsArray() ? InTemplatePin->GetArrayElementCppType() : InTemplatePin->GetCPPType();
 	if(InPinToResolve->IsArray())
 	{
-		ResolvedCPPType = FString::Printf(TEXT("TArray<%s>"), *ResolvedCPPType);
+		ResolvedCPPType = FString::Printf(TArrayTemplate, *ResolvedCPPType);
 	}
 
 	if(!ChangePinType(InPinToResolve, ResolvedCPPType, InTemplatePin->GetCPPTypeObject(), bSetupUndoRedo, false, false, false))
@@ -12216,18 +12237,7 @@ bool URigVMController::EnsurePinValidity(URigVMPin* InPin, bool bRecursive)
 		}
 	}
 
-	if(const UScriptStruct* ScriptStruct = Cast<UScriptStruct>(InPin->GetCPPTypeObject()))
-	{
-		if(InPin->IsArray())
-		{
-			static const TCHAR TemplateTArray[] = TEXT("TArray<%s>");
-			InPin->CPPType = FString::Printf(TemplateTArray, *ScriptStruct->GetStructCPPName());
-		}
-		else
-		{
-			InPin->CPPType = ScriptStruct->GetStructCPPName();
-		}
-	}
+	InPin->CPPType = PostProcessCPPType(InPin->CPPType, InPin->GetCPPTypeObject(), InPin->IsArray());
 
 	if(bRecursive)
 	{
diff --git a/Engine/Source/Developer/RigVMDeveloper/Private/RigVMModel/RigVMPin.cpp b/Engine/Source/Developer/RigVMDeveloper/Private/RigVMModel/RigVMPin.cpp
--- a/Engine/Source/Developer/RigVMDeveloper/Private/RigVMModel/RigVMPin.cpp
+++ b/Engine/Source/Developer/RigVMDeveloper/Private/RigVMModel/RigVMPin.cpp
@@ -4,6 +4,7 @@
 #include "RigVMModel/RigVMNode.h"
 #include "RigVMModel/RigVMGraph.h"
 #include "RigVMModel/RigVMLink.h"
+#include "RigVMModel/RigVMController.h"
 #include "RigVMModel/Nodes/RigVMPrototypeNode.h"
 #include "RigVMCompiler/RigVMCompiler.h"
 #include "RigVMCore/RigVMExecuteContext.h"
@@ -445,9 +446,14 @@ bool URigVMPin::IsStructMember() const
 	return ParentPin->IsStruct();
 }
 
+bool URigVMPin::IsUObject() const
+{
+	return CPPType.StartsWith(URigVMController::TObjectPtrPrefix);
+}
+
 bool URigVMPin::IsArray() const
 {
-	return CPPType.StartsWith(TEXT("TArray"));
+	return CPPType.StartsWith(URigVMController::TArrayPrefix);
 }
 
 bool URigVMPin::IsArrayElement() const
@@ -497,16 +503,7 @@ int32 URigVMPin::GetArraySize() const
 
 FString URigVMPin::GetCPPType() const
 {
-	if(UScriptStruct* ScriptStruct = GetScriptStruct())
-	{
-		const FString StructCPPType = ScriptStruct->GetStructCPPName();
-		if(IsArray())
-		{
-			return FString::Printf(TEXT("TArray<%s>"), *StructCPPType);
-		}
-		return StructCPPType;
-	}
-	return CPPType;
+	return URigVMController::PostProcessCPPType(CPPType, GetCPPTypeObject(), IsArray());
 }
 
 FString URigVMPin::GetArrayElementCppType() const
@@ -648,12 +645,33 @@ bool URigVMPin::IsValidDefaultValue(const FString& InDefaultValue) const
 		DefaultValues.Add(InDefaultValue);
 	}
 
-	FString BaseCPPType = GetCPPType().Replace(TEXT("TArray<"), TEXT("")).Replace(TEXT(">"), TEXT(""));
+	FString BaseCPPType = GetCPPType()
+		.Replace(URigVMController::TArrayPrefix, TEXT(""))
+		.Replace(URigVMController::TObjectPtrPrefix, TEXT(""))
+		.Replace(TEXT(">"), TEXT(""));
 
 	for (const FString& Value : DefaultValues)
 	{
 		// perform single value validation
-		if (UScriptStruct* ScriptStruct = Cast<UScriptStruct>(GetCPPTypeObject()))
+		if (UClass* Class = Cast<UClass>(GetCPPTypeObject()))
+		{
+			if(Value.IsEmpty())
+			{
+				return true;
+			}
+			
+			UObject* Object = FindObjectFromCPPTypeObjectPath(Value);
+			if(Object == nullptr)
+			{
+				return false;
+			}
+
+			if(!Object->GetClass()->IsChildOf(Class))
+			{
+				return false;
+			}
+		} 
+		else if (UScriptStruct* ScriptStruct = Cast<UScriptStruct>(GetCPPTypeObject()))
 		{
 			FRigVMPinDefaultValueImportErrorContext ErrorPipe;
 			TArray<uint8> TempStructBuffer;
diff --git a/Engine/Source/Developer/RigVMDeveloper/Public/RigVMModel/RigVMController.h b/Engine/Source/Developer/RigVMDeveloper/Public/RigVMModel/RigVMController.h
--- a/Engine/Source/Developer/RigVMDeveloper/Public/RigVMModel/RigVMController.h
+++ b/Engine/Source/Developer/RigVMDeveloper/Public/RigVMModel/RigVMController.h
@@ -853,6 +853,7 @@ private:
 	static FLinearColor GetColorFromMetadata(const FString& InMetadata);
 	static void CreateDefaultValueForStructIfRequired(UScriptStruct* InStruct, FString& InOutDefaultValue);
 	static void PostProcessDefaultValue(URigVMPin* Pin, FString& OutDefaultValue);
+	static FString PostProcessCPPType(const FString& InCPPType, UObject* InCPPTypeObject, bool bIsArray);
 
 	/*
 	void PotentiallyResolvePrototypeNode(URigVMPrototypeNode* InNode, bool bSetupUndoRedo);
@@ -887,7 +888,11 @@ private:
 	static FString GetSanitizedPinName(const FString& InName);
 	static FString GetSanitizedPinPath(const FString& InName);
 	static void SanitizeName(FString& InOutName, bool bAllowPeriod, bool bAllowSpace);
-	
+	static const TCHAR TArrayPrefix[];
+	static const TCHAR TObjectPtrPrefix[];
+	static const TCHAR TArrayTemplate[];
+	static const TCHAR TObjectPtrTemplate[];
+
 	UPROPERTY(transient)
 	TArray<TObjectPtr<URigVMGraph>> Graphs;
 
@@ -945,6 +950,7 @@ private:
 	bool bIsRunningUnitTest;
 
 	friend class URigVMGraph;
+	friend class URigVMPin;
 	friend class URigVMActionStack;
 	friend class URigVMCompiler;
 	friend struct FRigVMControllerObjectFactory;
diff --git a/Engine/Source/Developer/RigVMDeveloper/Public/RigVMModel/RigVMPin.h b/Engine/Source/Developer/RigVMDeveloper/Public/RigVMModel/RigVMPin.h
--- a/Engine/Source/Developer/RigVMDeveloper/Public/RigVMModel/RigVMPin.h
+++ b/Engine/Source/Developer/RigVMDeveloper/Public/RigVMModel/RigVMPin.h
@@ -166,6 +166,10 @@ public:
 	UFUNCTION(BlueprintCallable, Category = RigVMPin)
 	bool IsStructMember() const;
 
+	// Returns true if the data type of the Pin is a uobject
+	UFUNCTION(BlueprintCallable, Category = RigVMPin)
+	bool IsUObject() const;
+
 	// Returns true if the data type of the Pin is an array
 	UFUNCTION(BlueprintCallable, Category = RigVMPin)
 	bool IsArray() const;
diff --git a/Engine/Source/Developer/RigVMDeveloper/Public/RigVMModel/RigVMVariableDescription.h b/Engine/Source/Developer/RigVMDeveloper/Public/RigVMModel/RigVMVariableDescription.h
--- a/Engine/Source/Developer/RigVMDeveloper/Public/RigVMModel/RigVMVariableDescription.h
+++ b/Engine/Source/Developer/RigVMDeveloper/Public/RigVMModel/RigVMVariableDescription.h
@@ -118,6 +118,11 @@ public:
 			PinType.PinCategory = UEdGraphSchema_K2::PC_Byte;
 			PinType.PinSubCategoryObject = CPPTypeObject;
 		}
+		else if (Cast<UClass>(CPPTypeObject))
+		{
+			PinType.PinCategory = UEdGraphSchema_K2::PC_Object;
+			PinType.PinSubCategoryObject = CPPTypeObject;
+		}
 
 		return PinType;
 	}
@@ -166,6 +171,16 @@ public:
 				OutCPPTypeObject = Struct;
 			}
 		}
+		else if (InPinType.PinCategory == UEdGraphSchema_K2::PC_Object ||
+			InPinType.PinCategory == UEdGraphSchema_K2::PC_SoftObject ||
+			InPinType.PinCategory == UEdGraphSchema_K2::AllObjectTypes)
+		{
+			if (UClass* Class = Cast<UClass>(InPinType.PinSubCategoryObject))
+			{
+				OutCPPType = Prefix + FString::Printf(TEXT("TObjectPtr<%s%s>"), Class->GetPrefixCPP(), *Class->GetName()) + Suffix;
+				OutCPPTypeObject = Class;
+			}
+		}
 		else if (InPinType.PinCategory == UEdGraphSchema_K2::PC_Byte ||
 			InPinType.PinCategory == UEdGraphSchema_K2::PC_Enum)
 		{
diff --git a/Engine/Source/Runtime/RigVM/Private/RigVMCore/RigVMMemoryStorage.cpp b/Engine/Source/Runtime/RigVM/Private/RigVMCore/RigVMMemoryStorage.cpp
--- a/Engine/Source/Runtime/RigVM/Private/RigVMCore/RigVMMemoryStorage.cpp
+++ b/Engine/Source/Runtime/RigVM/Private/RigVMCore/RigVMMemoryStorage.cpp
@@ -176,7 +176,9 @@ bool FRigVMPropertyDescription::RequiresCPPTypeObject(const FString& InCPPType)
 		TEXT("E"), 
 		TEXT("TArray<E"), 
 		TEXT("U"), 
-		TEXT("TArray<U")
+		TEXT("TArray<U"),
+		TEXT("TObjectPtr<"), 
+		TEXT("TArray<TObjectPtr<")
 	};
 	static const TArray<FString> CPPTypesNotRequiringCPPTypeObject = {
 		TEXT("FString"), 
