do not create multiple forward / backward / setup event nodes

- prevent adding event nodes to non-top level graphs through copy/paste
- constness

diff --git a/Engine/Source/Developer/RigVMDeveloper/Private/RigVMModel/RigVMController.cpp b/Engine/Source/Developer/RigVMDeveloper/Private/RigVMModel/RigVMController.cpp
--- a/Engine/Source/Developer/RigVMDeveloper/Private/RigVMModel/RigVMController.cpp
+++ b/Engine/Source/Developer/RigVMDeveloper/Private/RigVMModel/RigVMController.cpp
@@ -191,7 +191,7 @@ FRigVMGraphModifiedEvent& URigVMController::OnModified()
 	return ModifiedEventStatic;
 }
 
-void URigVMController::Notify(ERigVMGraphNotifType InNotifType, UObject* InSubject)
+void URigVMController::Notify(ERigVMGraphNotifType InNotifType, UObject* InSubject) const
 {
 	if (bSuspendNotifications)
 	{
@@ -811,23 +811,30 @@ URigVMUnitNode* URigVMController::AddUnitNode(UScriptStruct* InScriptStruct, con
 		return nullptr;
 	}
 
-	// don't allow event nodes in anything but top level graphs
-	if (bSetupUndoRedo)
+	FStructOnScope StructOnScope(InScriptStruct);
+	FRigVMStruct* StructMemory = (FRigVMStruct*)StructOnScope.GetStructMemory();
+	InScriptStruct->InitializeDefaultValue((uint8*)StructMemory);
+	const bool bIsEventNode = (!StructMemory->GetEventName().IsNone());
+	if (bIsEventNode)
 	{
+		// don't allow event nodes in anything but top level graphs
 		if (!Graph->IsTopLevelGraph())
 		{
-			FStructOnScope StructOnScope(InScriptStruct);
-			FRigVMStruct* StructMemory = (FRigVMStruct*)StructOnScope.GetStructMemory();
-			InScriptStruct->InitializeDefaultValue((uint8*)StructMemory);
+			ReportAndNotifyError(TEXT("Event nodes can only be added to top level graphs."));
+			return nullptr;
+		}
 
-			if (!StructMemory->GetEventName().IsNone())
-			{
-				ReportAndNotifyError(TEXT("Event nodes can only be added to top level graphs."));
-				return nullptr;
-			}
+		// don't allow several event nodes in the main graph
+		TObjectPtr<URigVMNode> EventNode = FindEventNode(InScriptStruct);
+		if (EventNode != nullptr)
+		{
+			const FString ErrorMessage = FString::Printf(TEXT("Rig Graph can only contain one single %s node."),
+															*InScriptStruct->GetDisplayNameText().ToString());
+			ReportAndNotifyError(ErrorMessage);
+			return Cast<URigVMUnitNode>(EventNode);
 		}
 	}
-
+	
 	FString Name = GetValidNodeName(InNodeName.IsEmpty() ? InScriptStruct->GetName() : InNodeName);
 	URigVMUnitNode* Node = NewObject<URigVMUnitNode>(Graph, *Name);
 	Node->ScriptStruct = InScriptStruct;
@@ -10040,7 +10047,7 @@ FString URigVMController::GetValidNodeName(const FString& InPrefix)
 	}, false, true).ToString();
 }
 
-bool URigVMController::IsValidGraph()
+bool URigVMController::IsValidGraph() const
 {
 	URigVMGraph* Graph = GetGraph();
 	if (Graph == nullptr)
@@ -10260,10 +10267,90 @@ bool URigVMController::CanAddNode(URigVMNode* InNode, bool bReportErrors, bool b
 			}
 		}
 	}
-
+	else if (InNode->IsEvent())
+	{
+		if (URigVMUnitNode* InUnitNode = Cast<URigVMUnitNode>(InNode))
+		{
+			if (!CanAddEventNode(InUnitNode->GetScriptStruct(), bReportErrors))
+			{
+				return false;
+			}
+		}
+	}
+	
 	return true;
 }
 
+TObjectPtr<URigVMNode> URigVMController::FindEventNode(const UScriptStruct* InScriptStruct) const
+{
+	check(InScriptStruct);
+
+	// construct equivalent default struct
+	FStructOnScope InDefaultStructScope(InScriptStruct);
+	InScriptStruct->InitializeDefaultValue((uint8*)InDefaultStructScope.GetStructMemory());
+	
+	if (URigVMGraph* Graph = GetGraph())
+	{
+		TObjectPtr<URigVMNode>* FoundNode = 
+		Graph->Nodes.FindByPredicate( [&InDefaultStructScope](const TObjectPtr<URigVMNode>& Node) {
+			if (Node->IsEvent())
+			{
+				if (URigVMUnitNode* UnitNode = Cast<URigVMUnitNode>(Node))
+				{
+					// compare default structures
+					TSharedPtr<FStructOnScope> DefaultStructScope = UnitNode->ConstructStructInstance(true);
+					if (DefaultStructScope.IsValid() && InDefaultStructScope.GetStruct() == DefaultStructScope->GetStruct())
+					{
+						return true;
+					}
+				}
+			}
+			return false;
+		});
+
+		if (FoundNode)
+		{
+			return *FoundNode;
+		}
+	}
+	
+	return TObjectPtr<URigVMNode>();
+}
+
+bool URigVMController::CanAddEventNode(UScriptStruct* InScriptStruct, const bool bReportErrors) const
+{
+	if(!IsValidGraph())
+	{
+		return false;
+	}
+	
+	check(InScriptStruct);
+	
+	URigVMGraph* Graph = GetGraph();
+	check(Graph);
+
+	// check if we're trying to add a node within a graph which is not the top level one
+	if (!Graph->IsTopLevelGraph())
+	{
+		if (bReportErrors)
+		{
+			ReportAndNotifyError(TEXT("Event nodes can only be added to top level graphs."));
+		}
+		return false;
+	}
+
+	TObjectPtr<URigVMNode> EventNode = FindEventNode(InScriptStruct);
+	const bool bHasEventNode = EventNode != nullptr;
+	if (bHasEventNode && bReportErrors)
+	{
+		const FString ErrorMessage = FString::Printf(TEXT("Rig Graph can only contain one single %s node."),
+													 *InScriptStruct->GetDisplayNameText().ToString());
+		ReportAndNotifyError(ErrorMessage);
+	}
+		
+	return !bHasEventNode;
+}
+
 bool URigVMController::CanAddFunctionRefForDefinition(URigVMLibraryNode* InFunctionDefinition, bool bReportErrors)
 {
 	if(!IsValidGraph())
@@ -11647,7 +11734,7 @@ void URigVMController::ApplyPinStates(URigVMNode* InNode, const TMap<FString, UR
 	}
 }
 
-void URigVMController::ReportWarning(const FString& InMessage)
+void URigVMController::ReportWarning(const FString& InMessage) const
 {
 	if(!bReportWarningsAndErrors)
 	{
@@ -11666,7 +11753,7 @@ void URigVMController::ReportWarning(const FString& InMessage)
 	FScriptExceptionHandler::Get().HandleException(ELogVerbosity::Warning, *Message, *FString());
 }
 
-void URigVMController::ReportError(const FString& InMessage)
+void URigVMController::ReportError(const FString& InMessage) const
 {
 	if(!bReportWarningsAndErrors)
 	{
@@ -11685,7 +11772,7 @@ void URigVMController::ReportError(const FString& InMessage)
 	FScriptExceptionHandler::Get().HandleException(ELogVerbosity::Error, *Message, *FString());
 }
 
-void URigVMController::ReportAndNotifyError(const FString& InMessage)
+void URigVMController::ReportAndNotifyError(const FString& InMessage) const
 {
 	if (!bReportWarningsAndErrors)
 	{
diff --git a/Engine/Source/Developer/RigVMDeveloper/Public/RigVMModel/RigVMController.h b/Engine/Source/Developer/RigVMDeveloper/Public/RigVMModel/RigVMController.h
--- a/Engine/Source/Developer/RigVMDeveloper/Public/RigVMModel/RigVMController.h
+++ b/Engine/Source/Developer/RigVMDeveloper/Public/RigVMModel/RigVMController.h
@@ -137,7 +137,7 @@ public:
 	FRigVMGraphModifiedEvent& OnModified();
 
 	// Submits an event to the graph for broadcasting.
-	void Notify(ERigVMGraphNotifType InNotifType, UObject* InSubject);
+	void Notify(ERigVMGraphNotifType InNotifType, UObject* InSubject) const;
 
 	// Resends all notifications
 	void ResendAllNotifications();
@@ -741,9 +741,9 @@ public:
 
 	const FRigVMByteCode* GetCurrentByteCode() const;
 
-	void ReportWarning(const FString& InMessage);
-	void ReportError(const FString& InMessage);
-	void ReportAndNotifyError(const FString& InMessage);
+	void ReportWarning(const FString& InMessage) const;
+	void ReportError(const FString& InMessage) const;
+	void ReportAndNotifyError(const FString& InMessage) const;
 
 	template <typename FmtType, typename... Types>
 	void ReportWarningf(const FmtType& Fmt, Types... Args)
@@ -785,11 +785,13 @@ private:
 	void HandleModifiedEvent(ERigVMGraphNotifType InNotifType, URigVMGraph* InGraph, UObject* InSubject);
 
 	FString GetValidNodeName(const FString& InPrefix);
-	bool IsValidGraph();
+	bool IsValidGraph() const;
 	bool IsValidNodeForGraph(URigVMNode* InNode);
 	bool IsValidPinForGraph(URigVMPin* InPin);
 	bool IsValidLinkForGraph(URigVMLink* InLink);
 	bool CanAddNode(URigVMNode* InNode, bool bReportErrors, bool bIgnoreFunctionEntryReturnNodes = false);
+	TObjectPtr<URigVMNode> FindEventNode(const UScriptStruct* InScriptStruct) const;
+	bool CanAddEventNode(UScriptStruct* InScriptStruct, const bool bReportErrors) const;
 	bool CanAddFunctionRefForDefinition(URigVMLibraryNode* InFunctionDefinition, bool bReportErrors);
 	void AddPinsForStruct(UStruct* InStruct, URigVMNode* InNode, URigVMPin* InParentPin, ERigVMPinDirection InPinDirection, const FString& InDefaultValue, bool bAutoExpandArrays, bool bNotify = false);
 	void AddPinsForArray(FArrayProperty* InArrayProperty, URigVMNode* InNode, URigVMPin* InParentPin, ERigVMPinDirection InPinDirection, const TArray<FString>& InDefaultValues, bool bAutoExpandArrays);
