Control Rig: Fix undoing Control's transform Current Value undo all of the previous changes and can't be redone

diff --git a/Engine/Plugins/Animation/ControlRig/Source/ControlRigDeveloper/Private/ControlRigBlueprint.cpp b/Engine/Plugins/Animation/ControlRig/Source/ControlRigDeveloper/Private/ControlRigBlueprint.cpp
--- a/Engine/Plugins/Animation/ControlRig/Source/ControlRigDeveloper/Private/ControlRigBlueprint.cpp
+++ b/Engine/Plugins/Animation/ControlRig/Source/ControlRigDeveloper/Private/ControlRigBlueprint.cpp
@@ -1815,7 +1815,8 @@ void UControlRigBlueprint::PostTransacted(const FTransactionObjectEvent& Transac
 
 	if (TransactionEvent.GetEventType() == ETransactionObjectEventType::UndoRedo)
 	{
-		if (TransactionEvent.GetChangedProperties().Contains(TEXT("HierarchyContainer")))
+		TArray<FName> PropertiesChanged = TransactionEvent.GetChangedProperties();
+		if (PropertiesChanged.Contains(TEXT("HierarchyContainer")))
 		{
 			int32 TransactionIndex = GEditor->Trans->FindTransactionIndex(TransactionEvent.GetTransactionId());
 			const FTransaction* Transaction = GEditor->Trans->GetTransaction(TransactionIndex);
@@ -1843,15 +1844,24 @@ void UControlRigBlueprint::PostTransacted(const FTransactionObjectEvent& Transac
 			MarkPackageDirty();
 		}
 
-		else if (TransactionEvent.GetChangedProperties().Contains(TEXT("DrawContainer")))
+		else if (PropertiesChanged.Contains(TEXT("DrawContainer")))
 		{
 			PropagateDrawInstructionsFromBPToInstances();
 		}
 
-		else if (TransactionEvent.GetChangedProperties().Contains(TEXT("VMRuntimeSettings")))
+		else if (PropertiesChanged.Contains(TEXT("VMRuntimeSettings")))
 		{
 			PropagateRuntimeSettingsFromBPToInstances();
 		}
+
+		else if (PropertiesChanged.Contains(TEXT("NewVariables")))
+		{
+			if (RefreshEditorEvent.IsBound())
+			{
+				RefreshEditorEvent.Broadcast(this);
+			}
+			MarkPackageDirty();			
+		}
 	}
 }
 
diff --git a/Engine/Plugins/Animation/ControlRig/Source/ControlRigEditor/Private/Editor/ControlRigEditor.cpp b/Engine/Plugins/Animation/ControlRig/Source/ControlRigEditor/Private/Editor/ControlRigEditor.cpp
--- a/Engine/Plugins/Animation/ControlRig/Source/ControlRigEditor/Private/Editor/ControlRigEditor.cpp
+++ b/Engine/Plugins/Animation/ControlRig/Source/ControlRigEditor/Private/Editor/ControlRigEditor.cpp
@@ -2340,16 +2340,28 @@ void FControlRigEditor::PasteNodes()
 
 void FControlRigEditor::PostUndo(bool bSuccess)
 {
+	const FTransaction* Transaction = GEditor->Trans->GetTransaction(GEditor->Trans->GetQueueLength() - GEditor->Trans->GetUndoCount());
 	IControlRigEditor::PostUndo(bSuccess);
+	PostTransaction(bSuccess, Transaction, false);
+}
+
+void FControlRigEditor::PostRedo(bool bSuccess)
+{
+	const FTransaction* Transaction = GEditor->Trans->GetTransaction(GEditor->Trans->GetQueueLength() - GEditor->Trans->GetUndoCount() - 1);
+	IControlRigEditor::PostRedo(bSuccess);
+	PostTransaction(bSuccess, Transaction, true);
+}
+
+void FControlRigEditor::PostTransaction(bool bSuccess, const FTransaction* Transaction, bool bIsRedo)
+{
 	EnsureValidRigElementsInDetailPanel();
 
 	if (UControlRigBlueprint* RigBlueprint = Cast<UControlRigBlueprint>(GetBlueprintObj()))
 	{
-		if (RigBlueprint->Status == BS_Dirty)
-		{
-			Compile();
-		}
-
+		// Do not compile here. ControlRigBlueprint::PostTransacted decides when it is necessary to compile depending
+		// on the properties that are affected.
+		//Compile();
+		
 		USkeletalMesh* PreviewMesh = GetPersonaToolkit()->GetPreviewScene()->GetPreviewMesh();
 		if (PreviewMesh != RigBlueprint->GetPreviewMesh())
 		{
@@ -2373,36 +2385,6 @@ void FControlRigEditor::PostUndo(bool bSuccess)
 			EditMode->RequestToRecreateGizmoActors();
 		}
 	}
-	
-	Compile();
-}
-
-void FControlRigEditor::PostRedo(bool bSuccess)
-{
-	IControlRigEditor::PostRedo(bSuccess);
-	EnsureValidRigElementsInDetailPanel();
-
-	if (UControlRigBlueprint* RigBlueprint = Cast<UControlRigBlueprint>(GetBlueprintObj()))
-	{
-		if (RigBlueprint->Status == BS_Dirty)
-		{
-			Compile();
-		}
-
-		USkeletalMesh* PreviewMesh = GetPersonaToolkit()->GetPreviewScene()->GetPreviewMesh();
-		if (PreviewMesh != RigBlueprint->GetPreviewMesh())
-		{
-			RigBlueprint->SetPreviewMesh(PreviewMesh);
-			GetPersonaToolkit()->SetPreviewMesh(PreviewMesh, true);
-		}
-
-		if (FControlRigEditMode* EditMode = GetEditMode())
-		{
-			EditMode->RequestToRecreateGizmoActors();
-		}
-	}
-
-	Compile();
 }
 
 void FControlRigEditor::EnsureValidRigElementsInDetailPanel()
diff --git a/Engine/Plugins/Animation/ControlRig/Source/ControlRigEditor/Private/Editor/ControlRigEditor.h b/Engine/Plugins/Animation/ControlRig/Source/ControlRigEditor/Private/Editor/ControlRigEditor.h
--- a/Engine/Plugins/Animation/ControlRig/Source/ControlRigEditor/Private/Editor/ControlRigEditor.h
+++ b/Engine/Plugins/Animation/ControlRig/Source/ControlRigEditor/Private/Editor/ControlRigEditor.h
@@ -137,6 +137,8 @@ public:
 	virtual void PostUndo(bool bSuccess) override;
 	virtual void PostRedo(bool bSuccess) override;
 
+	virtual void PostTransaction(bool bSuccess, const FTransaction* Transaction, bool bIsRedo);
+
 	void EnsureValidRigElementsInDetailPanel();
 
 	virtual void OnStartWatchingPin();
diff --git a/Engine/Source/Developer/RigVMDeveloper/Private/RigVMModel/RigVMController.cpp b/Engine/Source/Developer/RigVMDeveloper/Private/RigVMModel/RigVMController.cpp
--- a/Engine/Source/Developer/RigVMDeveloper/Private/RigVMModel/RigVMController.cpp
+++ b/Engine/Source/Developer/RigVMDeveloper/Private/RigVMModel/RigVMController.cpp
@@ -859,6 +859,7 @@ URigVMUnitNode* URigVMController::AddUnitNode(UScriptStruct* InScriptStruct, con
 		Graph->MarkPackageDirty();
 	}
 
+	FRigVMControllerCompileBracketScope CompileScope(this);
 	FRigVMAddUnitNodeAction Action;
 	if (bSetupUndoRedo)
 	{
@@ -1072,6 +1073,7 @@ URigVMVariableNode* URigVMController::AddVariableNode(const FName& InVariableNam
 		Graph->MarkPackageDirty();
 	}
 
+	FRigVMControllerCompileBracketScope CompileScope(this);
 	FRigVMAddVariableNodeAction Action;
 	if (bSetupUndoRedo)
 	{
@@ -1213,6 +1215,7 @@ void URigVMController::OnExternalVariableRemoved(const FName& InVarName, bool bS
 	
 	const FString VarNameStr = InVarName.ToString();
 
+	FRigVMControllerCompileBracketScope CompileScope(this);
 	if (bSetupUndoRedo)
 	{
 		OpenUndoBracket(TEXT("Remove Variable Nodes"));
@@ -1308,6 +1311,7 @@ void URigVMController::OnExternalVariableRenamed(const FName& InOldVarName, cons
 
 	const FString VarNameStr = InOldVarName.ToString();
 
+	FRigVMControllerCompileBracketScope CompileScope(this);
 	if (bSetupUndoRedo)
 	{
 		OpenUndoBracket(TEXT("Rename Variable Nodes"));
@@ -1406,6 +1410,7 @@ void URigVMController::OnExternalVariableTypeChanged(const FName& InVarName, con
 
 	const FString VarNameStr = InVarName.ToString();
 
+	FRigVMControllerCompileBracketScope CompileScope(this);
 	if (bSetupUndoRedo)
 	{
 		OpenUndoBracket(TEXT("Change Variable Nodes Type"));
@@ -1681,6 +1686,7 @@ URigVMParameterNode* URigVMController::AddParameterNode(const FName& InParameter
 		Graph->MarkPackageDirty();
 	}
 
+	FRigVMControllerCompileBracketScope CompileScope(this);
 	FRigVMAddParameterNodeAction Action;
 	if (bSetupUndoRedo)
 	{
@@ -1760,6 +1766,7 @@ URigVMCommentNode* URigVMController::AddCommentNode(const FString& InCommentText
 		Graph->MarkPackageDirty();
 	}
 
+	FRigVMControllerCompileBracketScope CompileScope(this);
 	FRigVMAddCommentNodeAction Action;
 	if (bSetupUndoRedo)
 	{
@@ -1806,6 +1813,7 @@ URigVMRerouteNode* URigVMController::AddRerouteNodeOnLink(URigVMLink* InLink, bo
 
 	TGuardValue<bool> GuardCompactness(bIgnoreRerouteCompactnessChanges, true);
 
+	FRigVMControllerCompileBracketScope CompileScope(this);
 	FRigVMBaseAction Action;
 	if (bSetupUndoRedo)
 	{
@@ -1887,6 +1895,7 @@ URigVMRerouteNode* URigVMController::AddRerouteNodeOnPin(const FString& InPinPat
 
 	TGuardValue<bool> GuardCompactness(bIgnoreRerouteCompactnessChanges, true);
 
+	FRigVMControllerCompileBracketScope CompileScope(this);
 	FRigVMBaseAction Action;
 	if (bSetupUndoRedo)
 	{
@@ -2043,6 +2052,7 @@ URigVMInjectionInfo* URigVMController::AddInjectedNode(const FString& InPinPath,
 		return nullptr;
 	}
 
+	FRigVMControllerCompileBracketScope CompileScope(this);
 	FRigVMBaseAction Action;
 	if (bSetupUndoRedo)
 	{
@@ -2248,6 +2258,7 @@ URigVMNode* URigVMController::EjectNodeFromPin(const FString& InPinPath, bool bS
 		}
 	}
 
+	FRigVMControllerCompileBracketScope CompileScope(this);
 	FRigVMBaseAction Action;
 	if (bSetupUndoRedo)
 	{
@@ -2563,6 +2574,7 @@ TArray<FName> URigVMController::ImportNodesFromText(const FString& InText, bool
 		OpenUndoBracket(TEXT("Importing Nodes from Text"));
 	}
 
+	FRigVMControllerCompileBracketScope CompileScope(this);
 	FRigVMInverseAction AddNodesAction;
 	if (bSetupUndoRedo)
 	{
@@ -2871,6 +2883,7 @@ TMap<URigVMLibraryNode*, URigVMLibraryNode*> URigVMController::LocalizeFunctions
 		ExportedTextPerFunction.Add(FunctionToLocalize, ExportedText);
 	}
 
+	FRigVMControllerCompileBracketScope CompileScope(this);
 	if (bSetupUndoRedo)
 	{
 		OpenUndoBracket(TEXT("Localize functions"));
@@ -3291,6 +3304,7 @@ URigVMCollapseNode* URigVMController::CollapseNodes(const TArray<URigVMNode*>& I
 		InputExecutePins.Add(ExecutePin);
 	}
 
+	FRigVMControllerCompileBracketScope CompileScope(this);
 	FRigVMCollapseNodesAction CollapseAction;
 	CollapseAction.Title = TEXT("Collapse Nodes");
 
@@ -3679,6 +3693,7 @@ TArray<URigVMNode*> URigVMController::ExpandLibraryNode(URigVMLibraryNode* InNod
 		return TArray<URigVMNode*>();
 	}
 
+	FRigVMControllerCompileBracketScope CompileScope(this);
 	FRigVMExpandNodeAction ExpandAction;
 	ExpandAction.Title = FString::Printf(TEXT("Expand '%s' Node"), *InNode->GetName());
 
@@ -4394,6 +4409,7 @@ URigVMFunctionReferenceNode* URigVMController::PromoteCollapseNodeToFunctionRefe
 		return nullptr;
 	}
 
+	FRigVMControllerCompileBracketScope CompileScope(this);
 	if (bSetupUndoRedo)
 	{
 		OpenUndoBracket(TEXT("Promote to Function"));
@@ -4563,6 +4579,7 @@ URigVMCollapseNode* URigVMController::PromoteFunctionReferenceNodeToCollapseNode
 		}
 	}
 
+	FRigVMControllerCompileBracketScope CompileScope(this);
 	if (bSetupUndoRedo)
 	{
 		OpenUndoBracket(TEXT("Promote to Collapse Node"));
@@ -4754,6 +4771,7 @@ bool URigVMController::RemoveNode(URigVMNode* InNode, bool bSetupUndoRedo, bool
 
 	TGuardValue<bool> GuardCompactness(bIgnoreRerouteCompactnessChanges, true);
 
+	FRigVMControllerCompileBracketScope CompileScope(this);
 	FRigVMBaseAction Action;
 	if (bSetupUndoRedo)
 	{
@@ -5696,6 +5714,7 @@ bool URigVMController::RenameVariable(const FName& InOldName, const FName& InNew
 		}
 	}
 
+	FRigVMControllerCompileBracketScope CompileScope(this);
 	FRigVMRenameVariableAction Action;
 	if (bSetupUndoRedo)
 	{
@@ -5767,6 +5786,7 @@ bool URigVMController::RenameParameter(const FName& InOldName, const FName& InNe
 		}
 	}
 
+	FRigVMControllerCompileBracketScope CompileScope(this);
 	FRigVMRenameParameterAction Action;
 	if (bSetupUndoRedo)
 	{
@@ -5955,6 +5975,7 @@ bool URigVMController::SetPinIsWatched(URigVMPin* InPin, bool bIsWatched, bool b
 		return false;
 	}
 
+	FRigVMControllerCompileBracketScope CompileScope(this);
 	FRigVMSetPinWatchAction Action;
 	if (bSetupUndoRedo)
 	{
@@ -6083,6 +6104,7 @@ bool URigVMController::SetPinDefaultValue(URigVMPin* InPin, const FString& InDef
 		}
 	}
 
+	FRigVMControllerCompileBracketScope CompileScope(this);
 	FRigVMSetPinDefaultValueAction Action;
 	if (bSetupUndoRedo)
 	{
@@ -6375,6 +6397,7 @@ URigVMPin* URigVMController::InsertArrayPin(URigVMPin* ArrayPin, int32 InIndex,
 		InIndex = ArrayPin->GetSubPins().Num();
 	}
 
+	FRigVMControllerCompileBracketScope CompileScope(this);
 	FRigVMInsertArrayPinAction Action;
 	if (bSetupUndoRedo)
 	{
@@ -6461,6 +6484,7 @@ bool URigVMController::RemoveArrayPin(const FString& InArrayElementPinPath, bool
 	check(ArrayPin);
 	ensure(ArrayPin->IsArray());
 
+	FRigVMControllerCompileBracketScope CompileScope(this);
 	FRigVMRemoveArrayPinAction Action;
 	if (bSetupUndoRedo)
 	{
@@ -6507,6 +6531,7 @@ bool URigVMController::RemoveArrayPin(const FString& InArrayElementPinPath, bool
 
 bool URigVMController::RemovePin(URigVMPin* InPinToRemove, bool bSetupUndoRedo, bool bNotify)
 {
+	FRigVMControllerCompileBracketScope CompileScope(this);
 	if (bSetupUndoRedo)
 	{
 		BreakAllLinks(InPinToRemove, true, bSetupUndoRedo);
@@ -6571,6 +6596,7 @@ bool URigVMController::SetArrayPinSize(const FString& InArrayPinPath, int32 InSi
 		return false;
 	}
 
+	FRigVMControllerCompileBracketScope CompileScope(this);
 	FRigVMBaseAction Action;
 	if (bSetupUndoRedo)
 	{
@@ -6720,6 +6746,7 @@ bool URigVMController::BindPinToVariable(URigVMPin* InPin, const FString& InNewB
 		}
 	}
 
+	FRigVMControllerCompileBracketScope CompileScope(this);
 	FRigVMBaseAction Action;
 	if (bSetupUndoRedo)
 	{
@@ -6819,6 +6846,7 @@ bool URigVMController::MakeBindingsFromVariableNode(URigVMVariableNode* InNode,
 
 	if (Pairs.Num() > 0)
 	{
+		FRigVMControllerCompileBracketScope CompileScope(this);
 		if (bSetupUndoRedo)
 		{
 			OpenUndoBracket(TEXT("Turn Variable Node into Bindings"));
@@ -6949,6 +6977,7 @@ bool URigVMController::PromotePinToVariable(URigVMPin* InPin, bool bCreateVariab
 
 	if (bCreateVariableNode)
 	{
+		FRigVMControllerCompileBracketScope CompileScope(this);
 		if (URigVMVariableNode* VariableNode = AddVariableNode(
 			VariableForPin.Name,
 			VariableForPin.TypeName.ToString(),
@@ -6967,6 +6996,7 @@ bool URigVMController::PromotePinToVariable(URigVMPin* InPin, bool bCreateVariab
 	}
 	else
 	{
+		FRigVMControllerCompileBracketScope CompileScope(this);
 		return BindPinToVariable(InPin, VariableForPin.Name.ToString(), bSetupUndoRedo);
 	}
 
@@ -7075,6 +7105,7 @@ bool URigVMController::AddLink(URigVMPin* OutputPin, URigVMPin* InputPin, bool b
 	ensure(!OutputPin->IsLinkedTo(InputPin));
 	ensure(!InputPin->IsLinkedTo(OutputPin));
 
+	FRigVMControllerCompileBracketScope CompileScope(this);
 	FRigVMAddLinkAction Action;
 	if (bSetupUndoRedo)
 	{
@@ -7272,6 +7303,7 @@ bool URigVMController::BreakLink(URigVMPin* OutputPin, URigVMPin* InputPin, bool
 	{
 		if (Link->SourcePin == OutputPin && Link->TargetPin == InputPin)
 		{
+			FRigVMControllerCompileBracketScope CompileScope(this);
 			FRigVMBreakLinkAction Action;
 			if (bSetupUndoRedo)
 			{
@@ -7346,6 +7378,7 @@ bool URigVMController::BreakAllLinks(const FString& InPinPath, bool bAsInput, bo
 
 bool URigVMController::BreakAllLinks(URigVMPin* Pin, bool bAsInput, bool bSetupUndoRedo)
 {
+	FRigVMControllerCompileBracketScope CompileScope(this);
 	FRigVMBaseAction Action;
 	if (bSetupUndoRedo)
 	{
@@ -7517,6 +7550,7 @@ FName URigVMController::AddExposedPin(const FName& InPinName, ERigVMPinDirection
 		AddPinsForStruct(Pin->GetScriptStruct(), LibraryNode, Pin, Pin->Direction, DefaultValue, false);
 	}
 
+	FRigVMControllerCompileBracketScope CompileScope(this);
 	FRigVMAddExposedPinAction Action(Pin);
 	if (bSetupUndoRedo)
 	{
@@ -7616,6 +7650,7 @@ bool URigVMController::RemoveExposedPin(const FName& InPinName, bool bSetupUndoR
 		}
 	}
 
+	FRigVMControllerCompileBracketScope CompileScope(this);
 	FRigVMRemoveExposedPinAction Action(Pin);
 	if (bSetupUndoRedo)
 	{
@@ -7734,6 +7769,7 @@ bool URigVMController::RenameExposedPin(const FName& InOldPinName, const FName&
 		return true;
 	}, false, true);
 
+	FRigVMControllerCompileBracketScope CompileScope(this);
 	FRigVMRenameExposedPinAction Action;
 	if (bSetupUndoRedo)
 	{
@@ -7912,6 +7948,7 @@ bool URigVMController::ChangeExposedPinType(const FName& InPinName, const FStrin
 		}
 	}
 
+	FRigVMControllerCompileBracketScope CompileScope(this);
 	FRigVMBaseAction Action;
 	if (bSetupUndoRedo)
 	{
@@ -8120,6 +8157,8 @@ URigVMFunctionReferenceNode* URigVMController::AddFunctionReferenceNode(URigVMLi
 	FunctionRefNode->SetReferencedNode(InFunctionDefinition);
 	Graph->Nodes.Add(FunctionRefNode);
 
+	FRigVMControllerCompileBracketScope CompileScope(this);
+	
 	RepopulatePinsOnNode(FunctionRefNode, false, false);
 
 	Notify(ERigVMGraphNotifType::NodeAdded, FunctionRefNode);
@@ -8264,6 +8303,7 @@ bool URigVMController::SetRemappedVariable(URigVMFunctionReferenceNode* InFuncti
 
 	Notify(ERigVMGraphNotifType::VariableRemappingChanged, InFunctionRefNode);
 
+	FRigVMControllerCompileBracketScope CompileScope(this);
 	if(bSetupUndoRedo)
 	{
 		ActionStack->AddAction(FRigVMSetRemappedVariableAction(InFunctionRefNode, InInnerVariableName, OldOuterVariableName, InOuterVariableName));
@@ -8294,6 +8334,8 @@ URigVMLibraryNode* URigVMController::AddFunctionToLibrary(const FName& InFunctio
 	CollapseNode->Position = InNodePosition;
 	Graph->Nodes.Add(CollapseNode);
 
+	FRigVMControllerCompileBracketScope CompileScope(this);
+	
 	if (bMutable)
 	{
 		URigVMPin* ExecutePin = NewObject<URigVMPin>(CollapseNode, FRigVMStruct::ExecuteContextName);
@@ -8443,6 +8485,8 @@ FRigVMGraphVariableDescription URigVMController::AddLocalVariable(const FName& I
 
 	Graph->LocalVariables.Add(NewVariable);
 
+	FRigVMControllerCompileBracketScope CompileScope(this);
+	
 	for (URigVMNode* Node : Graph->GetNodes())
 	{
 		if (URigVMVariableNode* VariableNode = Cast<URigVMVariableNode>(Node))
@@ -8530,6 +8574,7 @@ bool URigVMController::RemoveLocalVariable(const FName& InVariableName, bool bSe
 
 	if (FoundIndex != INDEX_NONE)
 	{
+		FRigVMControllerCompileBracketScope CompileScope(this);
 		FRigVMInverseAction InverseAction;
 		if (bSetupUndoRedo)
 		{
@@ -8687,6 +8732,7 @@ bool URigVMController::RenameLocalVariable(const FName& InVariableName, const FN
 		Graph->MarkPackageDirty();
 	}
 
+	FRigVMControllerCompileBracketScope CompileScope(this);
 	if (bSetupUndoRedo)
 	{
 		FRigVMInverseAction InverseAction;
@@ -8777,7 +8823,8 @@ bool URigVMController::SetLocalVariableType(const FName& InVariableName, const F
 	{
 		Graph->MarkPackageDirty();
 	}
-	
+
+	FRigVMControllerCompileBracketScope CompileScope(this);
 	if (bSetupUndoRedo)
 	{
 		FRigVMInverseAction InverseAction;
@@ -8885,6 +8932,7 @@ bool URigVMController::SetLocalVariableDefaultValue(const FName& InVariableName,
 		return false;
 	}
 
+	FRigVMControllerCompileBracketScope CompileScope(this);
 	if (bSetupUndoRedo)
 	{
 		FRigVMInverseAction InverseAction;
@@ -9247,6 +9295,7 @@ URigVMRerouteNode* URigVMController::AddFreeRerouteNode(bool bShowAsFullNode, co
 		return nullptr;
 	}
 
+	FRigVMControllerCompileBracketScope CompileScope(this);
 	FRigVMBaseAction Action;
 	if (bSetupUndoRedo)
 	{
@@ -9339,6 +9388,7 @@ URigVMBranchNode* URigVMController::AddBranchNode(const FVector2D& InPosition, c
 
 	Notify(ERigVMGraphNotifType::NodeAdded, Node);
 
+	FRigVMControllerCompileBracketScope CompileScope(this);
 	if (bSetupUndoRedo)
 	{
 		ActionStack->AddAction(FRigVMAddBranchNodeAction(Node));
@@ -9444,6 +9494,7 @@ URigVMIfNode* URigVMController::AddIfNode(const FString& InCPPType, const FName&
 
 	Notify(ERigVMGraphNotifType::NodeAdded, Node);
 
+	FRigVMControllerCompileBracketScope CompileScope(this);
 	if (bSetupUndoRedo)
 	{
 		ActionStack->AddAction(FRigVMAddIfNodeAction(Node));
@@ -9538,6 +9589,7 @@ URigVMSelectNode* URigVMController::AddSelectNode(const FString& InCPPType, cons
 
 	SetArrayPinSize(ValuePin->GetPinPath(), 2, DefaultValue, false);
 
+	FRigVMControllerCompileBracketScope CompileScope(this);
 	if (bSetupUndoRedo)
 	{
 		ActionStack->AddAction(FRigVMAddSelectNodeAction(Node));
@@ -9615,6 +9667,7 @@ URigVMPrototypeNode* URigVMController::AddPrototypeNode(const FName& InNotation,
 
 	Notify(ERigVMGraphNotifType::NodeAdded, Node);
 
+	FRigVMControllerCompileBracketScope CompileScope(this);
 	if (bSetupUndoRedo)
 	{
 		ActionStack->AddAction(FRigVMAddPrototypeNodeAction(Node));
@@ -9680,6 +9733,7 @@ URigVMEnumNode* URigVMController::AddEnumNode(const FName& InCPPTypeObjectPath,
 
 	Notify(ERigVMGraphNotifType::NodeAdded, Node);
 
+	FRigVMControllerCompileBracketScope CompileScope(this);
 	if (bSetupUndoRedo)
 	{
 		ActionStack->AddAction(FRigVMAddEnumNodeAction(Node));
@@ -9962,6 +10016,7 @@ URigVMArrayNode* URigVMController::AddArrayNode(ERigVMOpCode InOpCode, const FSt
 		Graph->MarkPackageDirty();
 	}
 
+	FRigVMControllerCompileBracketScope CompileScope(this);
 	FRigVMAddArrayNodeAction Action;
 	if (bSetupUndoRedo)
 	{
@@ -12013,6 +12068,7 @@ void URigVMController::ResolveUnknownTypePin(URigVMPin* InPinToResolve, const UR
 	check(InPinToResolve->IsUnknownType());
 	check(!InTemplatePin->IsUnknownType());
 
+	FRigVMControllerCompileBracketScope CompileScope(this);
 	if(bTraverseParentPins)
 	{
 		URigVMPin* RootPin = InPinToResolve->GetRootPin();
@@ -12127,6 +12183,7 @@ bool URigVMController::ChangePinType(URigVMPin* InPin, const FString& InCPPType,
 		return false;
 	}
 
+	FRigVMControllerCompileBracketScope CompileScope(this);
 	FRigVMBaseAction Action;
 	if (bSetupUndoRedo)
 	{
@@ -12228,6 +12285,7 @@ void URigVMController::RewireLinks(URigVMPin* InOldPin, URigVMPin* InNewPin, boo
 {
 	ensure(InOldPin->GetRootPin() == InOldPin);
 	ensure(InNewPin->GetRootPin() == InNewPin);
+	FRigVMControllerCompileBracketScope CompileScope(this);
 
 	if (bAsInput)
 	{
