Control Rig: Fix "show bubble when zoom" change is not getting saved

diff --git a/Engine/Plugins/Animation/ControlRig/Source/ControlRigDeveloper/Private/Graph/ControlRigGraph.cpp b/Engine/Plugins/Animation/ControlRig/Source/ControlRigDeveloper/Private/Graph/ControlRigGraph.cpp
--- a/Engine/Plugins/Animation/ControlRig/Source/ControlRigDeveloper/Private/Graph/ControlRigGraph.cpp
+++ b/Engine/Plugins/Animation/ControlRig/Source/ControlRigDeveloper/Private/Graph/ControlRigGraph.cpp
@@ -482,6 +482,10 @@ void UControlRigGraph::HandleModifiedEvent(ERigVMGraphNotifType InNotifType, URi
 				if (EdNode)
 				{
 					EdNode->OnUpdateCommentText(ModelNode->GetCommentText());
+					EdNode->FontSize = ModelNode->GetCommentFontSize();
+					EdNode->bCommentBubbleVisible = ModelNode->GetCommentBubbleVisible();
+					EdNode->bCommentBubbleVisible_InDetailsPanel = ModelNode->GetCommentBubbleVisible();
+					EdNode->bColorCommentBubble = ModelNode->GetCommentColorBubble();
 				}
 			}
 			break;
diff --git a/Engine/Plugins/Animation/ControlRig/Source/ControlRigEditor/Private/Editor/ControlRigEditor.cpp b/Engine/Plugins/Animation/ControlRig/Source/ControlRigEditor/Private/Editor/ControlRigEditor.cpp
--- a/Engine/Plugins/Animation/ControlRig/Source/ControlRigEditor/Private/Editor/ControlRigEditor.cpp
+++ b/Engine/Plugins/Animation/ControlRig/Source/ControlRigEditor/Private/Editor/ControlRigEditor.cpp
@@ -4145,7 +4145,7 @@ void FControlRigEditor::OnNodeTitleCommitted(const FText& NewText, ETextCommit::
 	{
 		if (UControlRigBlueprint* ControlRigBP = GetControlRigBlueprint())
 		{
-			GetFocusedController()->SetCommentTextByName(CommentBeingChanged->GetFName(), NewText.ToString(), true, true);
+			GetFocusedController()->SetCommentTextByName(CommentBeingChanged->GetFName(), NewText.ToString(), CommentBeingChanged->FontSize, CommentBeingChanged->bCommentBubbleVisible, CommentBeingChanged->bColorCommentBubble, true, true);
 		}
 	}
 }
diff --git a/Engine/Plugins/Animation/ControlRig/Source/ControlRigEditor/Private/Graph/SControlRigGraphNodeComment.cpp b/Engine/Plugins/Animation/ControlRig/Source/ControlRigEditor/Private/Graph/SControlRigGraphNodeComment.cpp
--- a/Engine/Plugins/Animation/ControlRig/Source/ControlRigEditor/Private/Graph/SControlRigGraphNodeComment.cpp
+++ b/Engine/Plugins/Animation/ControlRig/Source/ControlRigEditor/Private/Graph/SControlRigGraphNodeComment.cpp
@@ -18,6 +18,7 @@
 SControlRigGraphNodeComment::SControlRigGraphNodeComment()
 	: SGraphNodeComment()
 	, CachedNodeCommentColor(FLinearColor(-1.f, -1.f, -1.f, -1.f))
+	, CachedColorBubble(INDEX_NONE)
 {
 }
 
@@ -86,8 +87,19 @@ void SControlRigGraphNodeComment::Tick(const FGeometry& AllottedGeometry, const
 	{
 		UEdGraphNode_Comment* CommentNode = CastChecked<UEdGraphNode_Comment>(GraphNode);
 
+		if (CachedColorBubble == INDEX_NONE)
+		{
+			CachedColorBubble = (int8) CommentNode->bColorCommentBubble;
+		}
+
 		const FString CurrentCommentTitle = GetNodeComment();
-		if (CurrentCommentTitle != CachedCommentTitle)
+		const int32 CurrentCommentFontSize = CommentNode->GetFontSize();
+		const bool CurrentCommentBubbleVisible = CommentNode->bCommentBubbleVisible;
+		const bool CurrentCommentColorBubble = CommentNode->bColorCommentBubble;
+		if (CurrentCommentTitle != CachedCommentTitle ||
+			CurrentCommentFontSize != CachedFontSize ||
+			CurrentCommentBubbleVisible != bCachedBubbleVisibility ||
+			CurrentCommentColorBubble != (bool) CachedColorBubble)
 		{
 			if (UControlRigGraph* Graph = Cast<UControlRigGraph>(CommentNode->GetOuter()))
 			{
@@ -95,7 +107,8 @@ void SControlRigGraphNodeComment::Tick(const FGeometry& AllottedGeometry, const
 				{
 					if (URigVMController* Controller = Blueprint->GetController(Graph))
 					{
-						Controller->SetCommentTextByName(CommentNode->GetFName(), CurrentCommentTitle, true, true);
+						Controller->SetCommentTextByName(CommentNode->GetFName(), CurrentCommentTitle, CurrentCommentFontSize, CurrentCommentBubbleVisible, CurrentCommentColorBubble, true, true);
+						CachedColorBubble = (int8) CurrentCommentColorBubble;
 					}
 				}
 			}
diff --git a/Engine/Plugins/Animation/ControlRig/Source/ControlRigEditor/Private/Graph/SControlRigGraphNodeComment.h b/Engine/Plugins/Animation/ControlRig/Source/ControlRigEditor/Private/Graph/SControlRigGraphNodeComment.h
--- a/Engine/Plugins/Animation/ControlRig/Source/ControlRigEditor/Private/Graph/SControlRigGraphNodeComment.h
+++ b/Engine/Plugins/Animation/ControlRig/Source/ControlRigEditor/Private/Graph/SControlRigGraphNodeComment.h
@@ -28,4 +28,6 @@ protected:
 private:
 
 	FLinearColor CachedNodeCommentColor;
+
+	int8 CachedColorBubble;
 };
diff --git a/Engine/Source/Developer/RigVMDeveloper/Private/RigVMModel/Nodes/RigVMCommentNode.cpp b/Engine/Source/Developer/RigVMDeveloper/Private/RigVMModel/Nodes/RigVMCommentNode.cpp
--- a/Engine/Source/Developer/RigVMDeveloper/Private/RigVMModel/Nodes/RigVMCommentNode.cpp
+++ b/Engine/Source/Developer/RigVMDeveloper/Private/RigVMModel/Nodes/RigVMCommentNode.cpp
@@ -5,9 +5,27 @@
 URigVMCommentNode::URigVMCommentNode()
 {
 	Size = FVector2D(400.f, 300.f);
+	FontSize = 18;
+	bBubbleVisible = false;
+	bColorBubble = false;
 }
 
 FString URigVMCommentNode::GetCommentText() const
 {
 	return CommentText;
 }
+
+int32 URigVMCommentNode::GetCommentFontSize() const
+{
+	return FontSize;
+}
+
+bool URigVMCommentNode::GetCommentBubbleVisible() const
+{
+	return bBubbleVisible;
+}
+
+bool URigVMCommentNode::GetCommentColorBubble() const
+{
+	return bColorBubble;
+}
diff --git a/Engine/Source/Developer/RigVMDeveloper/Private/RigVMModel/RigVMController.cpp b/Engine/Source/Developer/RigVMDeveloper/Private/RigVMModel/RigVMController.cpp
--- a/Engine/Source/Developer/RigVMDeveloper/Private/RigVMModel/RigVMController.cpp
+++ b/Engine/Source/Developer/RigVMDeveloper/Private/RigVMModel/RigVMController.cpp
@@ -220,6 +220,11 @@ void URigVMController::ResendAllNotifications()
 		for (URigVMNode* Node : Graph->Nodes)
 		{
 			Notify(ERigVMGraphNotifType::NodeAdded, Node);
+
+			if (URigVMCommentNode* CommentNode = Cast<URigVMCommentNode>(Node))
+			{
+				Notify(ERigVMGraphNotifType::CommentTextChanged, Node);
+			}
 		}
 
 		for (URigVMLink* Link : Graph->Links)
@@ -5567,7 +5572,7 @@ bool URigVMController::SetNodeDescriptionByName(const FName& InNodeName, const F
 	return SetNodeDescription(Node, InDescription, bSetupUndoRedo, bMergeUndoAction);
 }
 
-bool URigVMController::SetCommentText(URigVMNode* InNode, const FString& InCommentText, bool bSetupUndoRedo, bool bPrintPythonCommand)
+bool URigVMController::SetCommentText(URigVMNode* InNode, const FString& InCommentText, const int32& InCommentFontSize, const bool& bInCommentBubbleVisible, const bool& bInCommentColorBubble, bool bSetupUndoRedo, bool bPrintPythonCommand)
 {
 	if (!IsValidNodeForGraph(InNode))
 	{
@@ -5576,7 +5581,7 @@ bool URigVMController::SetCommentText(URigVMNode* InNode, const FString& InComme
 
 	if (URigVMCommentNode* CommentNode = Cast<URigVMCommentNode>(InNode))
 	{
-		if(CommentNode->CommentText == InCommentText)
+		if(CommentNode->CommentText == InCommentText && CommentNode->FontSize == InCommentFontSize && CommentNode->bBubbleVisible == bInCommentBubbleVisible && CommentNode->bColorBubble == bInCommentColorBubble)
 		{
 			return false;
 		}
@@ -5584,12 +5589,15 @@ bool URigVMController::SetCommentText(URigVMNode* InNode, const FString& InComme
 		FRigVMSetCommentTextAction Action;
 		if (bSetupUndoRedo)
 		{
-			Action = FRigVMSetCommentTextAction(CommentNode, InCommentText);
+			Action = FRigVMSetCommentTextAction(CommentNode, InCommentText, InCommentFontSize, bInCommentBubbleVisible, bInCommentColorBubble);
 			Action.Title = FString::Printf(TEXT("Set Comment Text"));
 			ActionStack->BeginAction(Action);
 		}
 
 		CommentNode->CommentText = InCommentText;
+		CommentNode->FontSize = InCommentFontSize;
+		CommentNode->bBubbleVisible = bInCommentBubbleVisible;
+		CommentNode->bColorBubble = bInCommentColorBubble;
 		Notify(ERigVMGraphNotifType::CommentTextChanged, InNode);
 
 		if (bSetupUndoRedo)
@@ -5615,7 +5623,7 @@ bool URigVMController::SetCommentText(URigVMNode* InNode, const FString& InComme
 	return false;
 }
 
-bool URigVMController::SetCommentTextByName(const FName& InNodeName, const FString& InCommentText, bool bSetupUndoRedo, bool bPrintPythonCommand)
+bool URigVMController::SetCommentTextByName(const FName& InNodeName, const FString& InCommentText, const int32& InCommentFontSize, const bool& bInCommentBubbleVisible, const bool& bInCommentColorBubble, bool bSetupUndoRedo, bool bPrintPythonCommand)
 {
 	if(!IsValidGraph())
 	{
@@ -5626,7 +5634,7 @@ bool URigVMController::SetCommentTextByName(const FName& InNodeName, const FStri
 	check(Graph);
 
 	URigVMNode* Node = Graph->FindNodeByName(InNodeName);
-	return SetCommentText(Node, InCommentText, bSetupUndoRedo, bPrintPythonCommand);
+	return SetCommentText(Node, InCommentText, InCommentFontSize, bInCommentBubbleVisible, bInCommentColorBubble, bSetupUndoRedo, bPrintPythonCommand);
 }
 
 bool URigVMController::SetRerouteCompactness(URigVMNode* InNode, bool bShowAsFullNode, bool bSetupUndoRedo, bool bPrintPythonCommand)
diff --git a/Engine/Source/Developer/RigVMDeveloper/Private/RigVMModel/RigVMControllerActions.cpp b/Engine/Source/Developer/RigVMDeveloper/Private/RigVMModel/RigVMControllerActions.cpp
--- a/Engine/Source/Developer/RigVMDeveloper/Private/RigVMModel/RigVMControllerActions.cpp
+++ b/Engine/Source/Developer/RigVMDeveloper/Private/RigVMModel/RigVMControllerActions.cpp
@@ -1057,10 +1057,31 @@ bool FRigVMSetNodeDescriptionAction::Redo(URigVMController* InController)
 	return FRigVMBaseAction::Redo(InController);
 }
 
-FRigVMSetCommentTextAction::FRigVMSetCommentTextAction(URigVMCommentNode* InNode, const FString& InNewText)
+FRigVMSetCommentTextAction::FRigVMSetCommentTextAction()
+: NodePath()
+, OldText()
+, NewText()
+, OldFontSize(18)
+, NewFontSize(18)
+, bOldBubbleVisible(false)
+, bNewBubbleVisible(false)
+, bOldColorBubble(false)
+, bNewColorBubble(false)
+{
+	
+}
+
+
+FRigVMSetCommentTextAction::FRigVMSetCommentTextAction(URigVMCommentNode* InNode, const FString& InNewText, const int32& InNewFontSize, const bool& bInNewBubbleVisible, const bool& bInNewColorBubble)
 : NodePath(InNode->GetNodePath())
 , OldText(InNode->GetCommentText())
 , NewText(InNewText)
+, OldFontSize(InNode->GetCommentFontSize())
+, NewFontSize(InNewFontSize)
+, bOldBubbleVisible(InNode->GetCommentBubbleVisible())
+, bNewBubbleVisible(bInNewBubbleVisible)
+, bOldColorBubble(InNode->GetCommentColorBubble())
+, bNewColorBubble(bInNewColorBubble)
 {
 }
 
@@ -1070,12 +1091,12 @@ bool FRigVMSetCommentTextAction::Undo(URigVMController* InController)
 	{
 		return false;
 	}
-	return InController->SetCommentTextByName(*NodePath, OldText, false);
+	return InController->SetCommentTextByName(*NodePath, OldText, OldFontSize, bOldBubbleVisible, bOldColorBubble, false);
 }
 
 bool FRigVMSetCommentTextAction::Redo(URigVMController* InController)
 {
-	if(!InController->SetCommentTextByName(*NodePath, NewText, false))
+	if(!InController->SetCommentTextByName(*NodePath, NewText, NewFontSize, bNewBubbleVisible, bNewColorBubble, false))
 	{
 		return false;
 	}
diff --git a/Engine/Source/Developer/RigVMDeveloper/Public/RigVMModel/Nodes/RigVMCommentNode.h b/Engine/Source/Developer/RigVMDeveloper/Public/RigVMModel/Nodes/RigVMCommentNode.h
--- a/Engine/Source/Developer/RigVMDeveloper/Public/RigVMModel/Nodes/RigVMCommentNode.h
+++ b/Engine/Source/Developer/RigVMDeveloper/Public/RigVMModel/Nodes/RigVMCommentNode.h
@@ -28,11 +28,32 @@ public:
 	UFUNCTION(BlueprintCallable, Category = RigVMCommentNode)
 	FString GetCommentText() const;
 
+	// Returns the current user provided font size of this comment.
+	UFUNCTION(BlueprintCallable, Category = RigVMCommentNode)
+	int32 GetCommentFontSize() const;
+
+	// Returns the current user provided bubble visibility of this comment.
+	UFUNCTION(BlueprintCallable, Category = RigVMCommentNode)
+	bool GetCommentBubbleVisible() const;
+
+	// Returns the current user provided bubble color inheritance of this comment.
+	UFUNCTION(BlueprintCallable, Category = RigVMCommentNode)
+	bool GetCommentColorBubble() const;
+
 private:
 
 	UPROPERTY()
 	FString CommentText;
 
+	UPROPERTY()
+	int32 FontSize;
+
+	UPROPERTY()
+	bool bBubbleVisible;
+
+	UPROPERTY()
+	bool bColorBubble;
+
 	friend class URigVMController;
 };
 
diff --git a/Engine/Source/Developer/RigVMDeveloper/Public/RigVMModel/RigVMController.h b/Engine/Source/Developer/RigVMDeveloper/Public/RigVMModel/RigVMController.h
--- a/Engine/Source/Developer/RigVMDeveloper/Public/RigVMModel/RigVMController.h
+++ b/Engine/Source/Developer/RigVMDeveloper/Public/RigVMModel/RigVMController.h
@@ -488,15 +488,15 @@ public:
 	UFUNCTION(BlueprintCallable, Category = RigVMController)
 	bool SetNodeDescriptionByName(const FName& InNodeName, const FString& InDescription, bool bSetupUndoRedo = true, bool bMergeUndoAction = false);
 
-	// Sets the comment text of a comment node in the graph.
+	// Sets the comment text and properties of a comment node in the graph.
 	// This causes a CommentTextChanged modified event.
 	UFUNCTION(BlueprintCallable, Category = RigVMController)
-	bool SetCommentText(URigVMNode* InNode, const FString& InCommentText, bool bSetupUndoRedo = true, bool bPrintPythonCommand = false);
+	bool SetCommentText(URigVMNode* InNode, const FString& InCommentText, const int32& InCommentFontSize, const bool& bInCommentBubbleVisible, const bool& bInCommentColorBubble, bool bSetupUndoRedo = true, bool bPrintPythonCommand = false);
 
-	// Sets the comment text of a comment node in the graph by name.
+	// Sets the comment text and properties of a comment node in the graph by name.
 	// This causes a CommentTextChanged modified event.
 	UFUNCTION(BlueprintCallable, Category = RigVMController)
-	bool SetCommentTextByName(const FName& InNodeName, const FString& InCommentText, bool bSetupUndoRedo = true, bool bPrintPythonCommand = false);
+	bool SetCommentTextByName(const FName& InNodeName, const FString& InCommentText, const int32& InCommentFontSize, const bool& bInCommentBubbleVisible, const bool& bInCommentColorBubble, bool bSetupUndoRedo = true, bool bPrintPythonCommand = false);
 
 	// Sets the compactness of a reroute node in the graph.
 	// This causes a RerouteCompactnessChanged modified event.
diff --git a/Engine/Source/Developer/RigVMDeveloper/Public/RigVMModel/RigVMControllerActions.h b/Engine/Source/Developer/RigVMDeveloper/Public/RigVMModel/RigVMControllerActions.h
--- a/Engine/Source/Developer/RigVMDeveloper/Public/RigVMModel/RigVMControllerActions.h
+++ b/Engine/Source/Developer/RigVMDeveloper/Public/RigVMModel/RigVMControllerActions.h
@@ -885,8 +885,8 @@ struct FRigVMSetCommentTextAction : public FRigVMBaseAction
 
 public:
 
-	FRigVMSetCommentTextAction() {}
-	FRigVMSetCommentTextAction(URigVMCommentNode* InNode, const FString& InNewText);
+	FRigVMSetCommentTextAction();
+	FRigVMSetCommentTextAction(URigVMCommentNode* InNode, const FString& InNewText, const int32& InNewFontSize, const bool& bInNewBubbleVisible, const bool& bInNewColorBubble);
 	virtual ~FRigVMSetCommentTextAction() {};
 	virtual bool Undo(URigVMController* InController) override;
 	virtual bool Redo(URigVMController* InController) override;
@@ -899,6 +899,24 @@ public:
 
 	UPROPERTY()
 	FString NewText;
+
+	UPROPERTY()
+	int32 OldFontSize;
+
+	UPROPERTY()
+	int32 NewFontSize;
+
+	UPROPERTY()
+	bool bOldBubbleVisible;
+
+	UPROPERTY()
+	bool bNewBubbleVisible;
+
+	UPROPERTY()
+	bool bOldColorBubble;
+
+	UPROPERTY()
+	bool bNewColorBubble;
 };
 
 /**
diff --git a/Engine/Source/Editor/GraphEditor/Public/SGraphNodeComment.h b/Engine/Source/Editor/GraphEditor/Public/SGraphNodeComment.h
--- a/Engine/Source/Editor/GraphEditor/Public/SGraphNodeComment.h
+++ b/Engine/Source/Editor/GraphEditor/Public/SGraphNodeComment.h
@@ -80,8 +80,6 @@ protected:
 	virtual FVector2D GetNodeMaximumSize() const override;
 	//~ Begin SGraphNodeResizable Interface
 
-private:
-
 	/** @return the color to tint the comment body */
 	FSlateColor GetCommentBodyColor() const;
 
@@ -91,16 +89,14 @@ private:
 	/** @return the color to tint the comment bubble */
 	FSlateColor GetCommentBubbleColor() const;
 
+private:
+	
 	/** Returns the width to wrap the text of the comment at */
 	float GetWrapAt() const;
 
-private:
 	/** The comment bubble widget (used when zoomed out) */
 	TSharedPtr<SCommentBubble> CommentBubble;
 
-	/** Was the bubble desired to be visible last frame? */
-	mutable bool bCachedBubbleVisibility;
-
 	/** The current selection state of the comment */
 	mutable bool bIsSelected;
 
@@ -109,15 +105,18 @@ private:
 
 protected:
 	/** cached comment title */
-	FString CachedCommentTitle; 
+	FString CachedCommentTitle;
+
+	/** cached font size */
+	int32 CachedFontSize;
+
+	/** Was the bubble desired to be visible last frame? */
+	mutable bool bCachedBubbleVisibility;
 
 private:
 	/** cached comment title */
 	int32 CachedWidth;
 
-	/** cached font size */
-	int32 CachedFontSize;
-
 	/** Local copy of the comment style */
 	FInlineEditableTextBlockStyle CommentStyle;
 };
