reset pin value capabilities for FunctionReferenceNode

diff --git a/Engine/Plugins/Animation/ControlRig/Source/ControlRigEditor/Private/ControlRigEditorModule.cpp b/Engine/Plugins/Animation/ControlRig/Source/ControlRigEditor/Private/ControlRigEditorModule.cpp
--- a/Engine/Plugins/Animation/ControlRig/Source/ControlRigEditor/Private/ControlRigEditorModule.cpp
+++ b/Engine/Plugins/Animation/ControlRig/Source/ControlRigEditor/Private/ControlRigEditorModule.cpp
@@ -1425,34 +1425,36 @@ void FControlRigEditorModule::GetContextMenuActions(const UControlRigGraphSchema
 						}
 					}
 
-					if(Cast<URigVMUnitNode>(ModelPin->GetNode()) != nullptr || 
+					if (Cast<URigVMUnitNode>(ModelPin->GetNode()) != nullptr || 
 						Cast<URigVMLibraryNode>(ModelPin->GetNode()) != nullptr)
 					{
 						if (ModelPin->GetDirection() == ERigVMPinDirection::Input && 
-							!ModelPin->IsExecuteContext() &&
-							!ModelPin->IsBoundToVariable())
+							!ModelPin->IsExecuteContext())
 						{
-							FToolMenuSection& VariablesSection = Menu->FindOrAddSection(TEXT("Variables"));
+							if (!ModelPin->IsBoundToVariable())
+							{
+								FToolMenuSection& VariablesSection = Menu->FindOrAddSection(TEXT("Variables"));
 
-							TSharedRef<SControlRigVariableBinding> VariableBindingWidget =
-								SNew(SControlRigVariableBinding)
-								.Blueprint(RigBlueprint)
-								.ModelPin(ModelPin)
-								.CanRemoveBinding(false);
+								TSharedRef<SControlRigVariableBinding> VariableBindingWidget =
+									SNew(SControlRigVariableBinding)
+									.Blueprint(RigBlueprint)
+									.ModelPin(ModelPin)
+									.CanRemoveBinding(false);
 
-							VariablesSection.AddEntry(FToolMenuEntry::InitWidget("BindPinToVariableWidget", VariableBindingWidget, FText(), true));
-						}
+								VariablesSection.AddEntry(FToolMenuEntry::InitWidget("BindPinToVariableWidget", VariableBindingWidget, FText(), true));
+							}
 
-						FToolMenuSection& Section = Menu->AddSection("EdGraphSchemaPinDefaults", LOCTEXT("PinDefaults", "Pin Defaults"));
-						Section.AddMenuEntry(
-							"ResetPinDefaultValue",
-							LOCTEXT("ResetPinDefaultValue", "Reset Pin Value"),
-							LOCTEXT("ResetPinDefaultValue_Tooltip", "Resets the pin's value to its default."),
-							FSlateIcon(),
-							FUIAction(FExecuteAction::CreateLambda([Controller, ModelPin]() {
-								Controller->ResetPinDefaultValue(ModelPin->GetPinPath());
-							})
-						));
+							FToolMenuSection& Section = Menu->AddSection("EdGraphSchemaPinDefaults", LOCTEXT("PinDefaults", "Pin Defaults"));
+							Section.AddMenuEntry(
+								"ResetPinDefaultValue",
+								LOCTEXT("ResetPinDefaultValue", "Reset Pin Value"),
+								LOCTEXT("ResetPinDefaultValue_Tooltip", "Resets the pin's value to its default."),
+								FSlateIcon(),
+								FUIAction(FExecuteAction::CreateLambda([Controller, ModelPin]() {
+									Controller->ResetPinDefaultValue(ModelPin->GetPinPath());
+								})
+							));
+						}
 					}
 
 					if ((ModelPin->GetCPPType() == TEXT("FVector") ||
diff --git a/Engine/Source/Developer/RigVMDeveloper/Private/RigVMModel/RigVMController.cpp b/Engine/Source/Developer/RigVMDeveloper/Private/RigVMModel/RigVMController.cpp
--- a/Engine/Source/Developer/RigVMDeveloper/Private/RigVMModel/RigVMController.cpp
+++ b/Engine/Source/Developer/RigVMDeveloper/Private/RigVMModel/RigVMController.cpp
@@ -6217,9 +6217,10 @@ bool URigVMController::ResetPinDefaultValue(const FString& InPinPath, bool bSetu
 		return false;
 	}
 
-	if (Cast<URigVMUnitNode>(Pin->GetNode()) == nullptr)
+	URigVMNode* Node = Pin->GetNode();
+	if (!Node->IsA<URigVMUnitNode>() && !Node->IsA<URigVMFunctionReferenceNode>())
 	{
-		ReportErrorf(TEXT("Pin '%s' is not on a unit node."), *InPinPath);
+		ReportErrorf(TEXT("Pin '%s' is neither part of a unit nor a function reference node."), *InPinPath);
 		return false;
 	}
 
@@ -6241,26 +6242,129 @@ bool URigVMController::ResetPinDefaultValue(URigVMPin* InPin, bool bSetupUndoRed
 {
 	check(InPin);
 
-	if (URigVMUnitNode* UnitNode = Cast<URigVMUnitNode>(InPin->GetNode()))
+	URigVMNode* RigVMNode = InPin->GetNode();
+
+	// unit nodes
+	if (URigVMUnitNode* UnitNode = Cast<URigVMUnitNode>(RigVMNode))
+	{
+		// cut off the first one since it's the node
+		static const uint32 Offset = 1;
+		const FString DefaultValue = GetPinInitialDefaultValueFromStruct(UnitNode->GetScriptStruct(), InPin, Offset);
+		if (!DefaultValue.IsEmpty())
+		{
+			SetPinDefaultValue(InPin, DefaultValue, true, bSetupUndoRedo, false);
+			return true;
+		}
+	}
+
+	// function reference nodes
+	URigVMFunctionReferenceNode* RefNode = Cast<URigVMFunctionReferenceNode>(RigVMNode);
+	if (RefNode != nullptr)
+	{
+		const FString DefaultValue = GetPinInitialDefaultValue(InPin);
+		if (!DefaultValue.IsEmpty())
+		{
+			SetPinDefaultValue(InPin, DefaultValue, true, bSetupUndoRedo, false);
+			return true;
+		}
+	}
+
+	return false;
+}
+
+FString URigVMController::GetPinInitialDefaultValue(URigVMPin* InPin)
+{
+	static const FString EmptyValue;
+	static const FString TArrayInitValue( TEXT("()") );
+	static const FString TObjectInitValue( TEXT("()") );
+	static const TMap<FString, FString> InitValues =
 	{
-		TSharedPtr<FStructOnScope> StructOnScope = UnitNode->ConstructStructInstance(true /* use default */);
+		{ TEXT("bool"),	TEXT("False") },
+		{ TEXT("int32"),	TEXT("0") },
+		{ TEXT("float"),	TEXT("0.000000") },
+		{ TEXT("double"),	TEXT("0.000000") },
+		{ TEXT("FName"),	FName(NAME_None).ToString() },
+		{ TEXT("FString"),	TEXT("") }
+	};
+
+	if (InPin->IsStruct())
+	{
+		// offset is useless here as we are going to get the full struct default value
+		static const uint32 Offset = 0;
+		return GetPinInitialDefaultValueFromStruct(InPin->GetScriptStruct(), InPin, Offset);
+	}
+		
+	if (InPin->IsStructMember())
+	{
+		if (URigVMPin* ParentPin = InPin->GetParentPin())
+		{
+			// cut off node's and parent struct's paths if func reference node, only node instead
+			static const uint32 Offset = InPin->GetNode()->IsA<URigVMFunctionReferenceNode>() ? 2 : 1;
+			return GetPinInitialDefaultValueFromStruct(ParentPin->GetScriptStruct(), InPin, Offset);
+		}
+	}
+
+	if (InPin->IsArray())
+	{
+		return TArrayInitValue;
+	}
+		
+	if (InPin->IsUObject())
+	{
+		return TObjectInitValue;
+	}
+		
+	if (UEnum* Enum = InPin->GetEnum())
+	{
+		return Enum->GetNameStringByIndex(0);
+	}
+	
+	if (const FString* BasicDefault = InitValues.Find(InPin->GetCPPType()))
+	{
+		return *BasicDefault;
+	}
+	
+	return EmptyValue;
+}
+
+FString URigVMController::GetPinInitialDefaultValueFromStruct(UScriptStruct* ScriptStruct, URigVMPin* InPin, uint32 InOffset)
+{
+	FString DefaultValue;
+	if (InPin && ScriptStruct)
+	{
+		TSharedPtr<FStructOnScope> StructOnScope = MakeShareable(new FStructOnScope(ScriptStruct));
+		uint8* Memory = (uint8*)StructOnScope->GetStructMemory();
+		ScriptStruct->InitializeDefaultValue(Memory);
+
+		if (InPin->GetScriptStruct() == ScriptStruct)
+		{
+			ScriptStruct->ExportText(DefaultValue, Memory, nullptr, nullptr, PPF_None, nullptr, true);
+			return DefaultValue;
+		}
+
+		const FString PinPath = InPin->GetPinPath();
 
 		TArray<FString> Parts;
-		if (!URigVMPin::SplitPinPath(InPin->GetPinPath(), Parts))
+		if (!URigVMPin::SplitPinPath(PinPath, Parts))
 		{
-			return false;
+			return DefaultValue;
 		}
 
-		int32 PartIndex = 1; // cut off the first one since it's the node
+		const uint32 NumParts = Parts.Num();
+		if (InOffset >= NumParts)
+		{
+			return DefaultValue;
+		}
 
-		UStruct* Struct = UnitNode->ScriptStruct;
+		uint32 PartIndex = InOffset;
+
+		UStruct* Struct = ScriptStruct;
 		FProperty* Property = Struct->FindPropertyByName(*Parts[PartIndex++]);
 		check(Property);
 
-		uint8* Memory = StructOnScope->GetStructMemory();
 		Memory = Property->ContainerPtrToValuePtr<uint8>(Memory);
 
-		while (PartIndex < Parts.Num() && Property != nullptr)
+		while (PartIndex < NumParts && Property != nullptr)
 		{
 			if (FArrayProperty* ArrayProperty = CastField<FArrayProperty>(Property))
 			{
@@ -6292,23 +6396,14 @@ bool URigVMController::ResetPinDefaultValue(URigVMPin* InPin, bool bSetupUndoRed
 
 		if (Memory)
 		{
-			FString DefaultValue;
 			check(Property);
 			Property->ExportTextItem(DefaultValue, Memory, nullptr, nullptr, PPF_None);
-
-			if (!DefaultValue.IsEmpty())
-			{
-				SetPinDefaultValue(InPin, DefaultValue, true, bSetupUndoRedo, false);
-				return true;
-			}
 		}
 	}
 
-	return false;
+	return DefaultValue;
 }
 
-
-
 FString URigVMController::AddArrayPin(const FString& InArrayPinPath, const FString& InDefaultValue, bool bSetupUndoRedo, bool bPrintPythonCommand)
 {
 	return InsertArrayPin(InArrayPinPath, INDEX_NONE, InDefaultValue, bSetupUndoRedo, bPrintPythonCommand);
@@ -12249,6 +12344,7 @@ bool URigVMController::ChangePinType(URigVMPin* InPin, const FString& InCPPType,
 	InPin->CPPType = InCPPType;
 	InPin->CPPTypeObjectPath = CPPTypeObjectPath;
 	InPin->CPPTypeObject = InCPPTypeObject;
+	// we might want to use GetPinInitialDefaultValue here for a better default value
 	InPin->DefaultValue = FString();
 
 	if (InPin->IsStruct())
@@ -12295,7 +12391,7 @@ void URigVMController::RewireLinks(URigVMPin* InOldPin, URigVMPin* InNewPin, boo
 	ensure(InNewPin->GetRootPin() == InNewPin);
 	FRigVMControllerCompileBracketScope CompileScope(this);
 
-	if (bAsInput)
+ 	if (bAsInput)
 	{
 		TArray<URigVMLink*> Links = InLinks;
 		if (Links.Num() == 0)
diff --git a/Engine/Source/Developer/RigVMDeveloper/Public/RigVMModel/RigVMController.h b/Engine/Source/Developer/RigVMDeveloper/Public/RigVMModel/RigVMController.h
--- a/Engine/Source/Developer/RigVMDeveloper/Public/RigVMModel/RigVMController.h
+++ b/Engine/Source/Developer/RigVMDeveloper/Public/RigVMModel/RigVMController.h
@@ -801,6 +801,8 @@ private:
 	virtual bool ShouldPinBeUnfolded(URigVMPin* InPin);
 	bool SetPinDefaultValue(URigVMPin* InPin, const FString& InDefaultValue, bool bResizeArrays, bool bSetupUndoRedo, bool bMergeUndoAction, bool bNotify = true);
 	bool ResetPinDefaultValue(URigVMPin* InPin, bool bSetupUndoRedo);
+	static FString GetPinInitialDefaultValue(URigVMPin* InPin);
+	static FString GetPinInitialDefaultValueFromStruct(UScriptStruct* ScriptStruct, URigVMPin* InPin, uint32 InOffset);
 	URigVMPin* InsertArrayPin(URigVMPin* ArrayPin, int32 InIndex, const FString& InDefaultValue, bool bSetupUndoRedo);
 	bool RemovePin(URigVMPin* InPinToRemove, bool bSetupUndoRedo, bool bNotify);
 	FProperty* FindPropertyForPin(const FString& InPinPath);
