Control Rig: Fix Array Element Pins not having correct default values, which causes VM to allocate fewer elements for the array pin.

diff --git a/Engine/Source/Developer/RigVMDeveloper/Private/RigVMModel/RigVMController.cpp b/Engine/Source/Developer/RigVMDeveloper/Private/RigVMModel/RigVMController.cpp
--- a/Engine/Source/Developer/RigVMDeveloper/Private/RigVMModel/RigVMController.cpp
+++ b/Engine/Source/Developer/RigVMDeveloper/Private/RigVMModel/RigVMController.cpp
@@ -6272,7 +6272,7 @@ bool URigVMController::ResetPinDefaultValue(URigVMPin* InPin, bool bSetupUndoRed
 	return false;
 }
 
-FString URigVMController::GetPinInitialDefaultValue(URigVMPin* InPin)
+FString URigVMController::GetPinInitialDefaultValue(const URigVMPin* InPin)
 {
 	static const FString EmptyValue;
 	static const FString TArrayInitValue( TEXT("()") );
@@ -6327,7 +6327,7 @@ FString URigVMController::GetPinInitialDefaultValue(URigVMPin* InPin)
 	return EmptyValue;
 }
 
-FString URigVMController::GetPinInitialDefaultValueFromStruct(UScriptStruct* ScriptStruct, URigVMPin* InPin, uint32 InOffset)
+FString URigVMController::GetPinInitialDefaultValueFromStruct(UScriptStruct* ScriptStruct, const URigVMPin* InPin, uint32 InOffset)
 {
 	FString DefaultValue;
 	if (InPin && ScriptStruct)
diff --git a/Engine/Source/Developer/RigVMDeveloper/Private/RigVMModel/RigVMPin.cpp b/Engine/Source/Developer/RigVMDeveloper/Private/RigVMModel/RigVMPin.cpp
--- a/Engine/Source/Developer/RigVMDeveloper/Private/RigVMModel/RigVMPin.cpp
+++ b/Engine/Source/Developer/RigVMDeveloper/Private/RigVMModel/RigVMPin.cpp
@@ -571,10 +571,12 @@ FString URigVMPin::GetDefaultValue(const URigVMPin::FPinOverride& InOverride) co
 				{
 					ElementDefaultValue = TEXT("\"") + ElementDefaultValue + TEXT("\"");
 				}
-				else if (ElementDefaultValue.IsEmpty())
-				{
-					continue;
-				}
+
+				// if default value is empty, VM may believe that the array is smaller than # of sub-pins 
+				// since VM determines the array size by counting the number of element default values.
+				// see URigVMCompiler::FindOrAddRegister() for how pin default values are put into VM memory
+				ensure(!ElementDefaultValue.IsEmpty());
+
 				ElementDefaultValues.Add(ElementDefaultValue);
 			}
 			if (ElementDefaultValues.Num() == 0)
@@ -613,6 +615,18 @@ FString URigVMPin::GetDefaultValue(const URigVMPin::FPinOverride& InOverride) co
 
 		return DefaultValue.IsEmpty() ? TEXT("()") : DefaultValue;
 	}
+	else if (IsArrayElement() && DefaultValue.IsEmpty())
+	{
+		// array element cannot have an empty default value because when an array pin is
+		// added as a property to a memory storage class, its default value needs to reflect 
+		// the number of array elements in that array pin.
+		// for example:
+		// for an array pin of 1 float, the final default value should be "(0.0)" instead of "()".
+		// This default value is used during URigVMCompiler::FindOrAddRegister(...)
+		// Thus in this block, we have to return something like 0.0 instead of empty string
+
+		return URigVMController::GetPinInitialDefaultValue(this);
+	}
 
 	return DefaultValue;
 }
diff --git a/Engine/Source/Developer/RigVMDeveloper/Public/RigVMModel/RigVMController.h b/Engine/Source/Developer/RigVMDeveloper/Public/RigVMModel/RigVMController.h
--- a/Engine/Source/Developer/RigVMDeveloper/Public/RigVMModel/RigVMController.h
+++ b/Engine/Source/Developer/RigVMDeveloper/Public/RigVMModel/RigVMController.h
@@ -801,8 +801,8 @@ private:
 	virtual bool ShouldPinBeUnfolded(URigVMPin* InPin);
 	bool SetPinDefaultValue(URigVMPin* InPin, const FString& InDefaultValue, bool bResizeArrays, bool bSetupUndoRedo, bool bMergeUndoAction, bool bNotify = true);
 	bool ResetPinDefaultValue(URigVMPin* InPin, bool bSetupUndoRedo);
-	static FString GetPinInitialDefaultValue(URigVMPin* InPin);
-	static FString GetPinInitialDefaultValueFromStruct(UScriptStruct* ScriptStruct, URigVMPin* InPin, uint32 InOffset);
+	static FString GetPinInitialDefaultValue(const URigVMPin* InPin);
+	static FString GetPinInitialDefaultValueFromStruct(UScriptStruct* ScriptStruct, const URigVMPin* InPin, uint32 InOffset);
 	URigVMPin* InsertArrayPin(URigVMPin* ArrayPin, int32 InIndex, const FString& InDefaultValue, bool bSetupUndoRedo);
 	bool RemovePin(URigVMPin* InPinToRemove, bool bSetupUndoRedo, bool bNotify);
 	FProperty* FindPropertyForPin(const FString& InPinPath);
