Control Rig: Local variable non-latin name change to '____' when undo deletion in Control Rig

diff --git a/Engine/Source/Developer/RigVMDeveloper/Private/RigVMModel/RigVMController.cpp b/Engine/Source/Developer/RigVMDeveloper/Private/RigVMModel/RigVMController.cpp
--- a/Engine/Source/Developer/RigVMDeveloper/Private/RigVMModel/RigVMController.cpp
+++ b/Engine/Source/Developer/RigVMDeveloper/Private/RigVMModel/RigVMController.cpp
@@ -12619,11 +12619,11 @@ void URigVMController::SanitizeName(FString& InOutName, bool bAllowPeriod, bool
 		TCHAR& C = InOutName[i];
 
 		const bool bGoodChar =
-			((C >= 'A') && (C <= 'Z')) || ((C >= 'a') && (C <= 'z')) ||		// A-Z (upper and lowercase) anytime
+			FChar::IsAlpha(C) ||											// Any letter (upper and lowercase) anytime
 			(C == '_') || (C == '-') || 									// _  and - anytime
 			(bAllowPeriod && (C == '.')) ||
 			(bAllowSpace && (C == ' ')) ||
-			((i > 0) && (C >= '0') && (C <= '9'));							// 0-9 after the first character
+			((i > 0) && FChar::IsDigit(C));									// 0-9 after the first character
 
 		if (!bGoodChar)
 		{
