Control Rig: Control Rig Parameters Not Respecting Clamping Meta data

diff --git a/Engine/Source/Developer/RigVMDeveloper/Private/RigVMModel/RigVMController.cpp b/Engine/Source/Developer/RigVMDeveloper/Private/RigVMModel/RigVMController.cpp
--- a/Engine/Source/Developer/RigVMDeveloper/Private/RigVMModel/RigVMController.cpp
+++ b/Engine/Source/Developer/RigVMDeveloper/Private/RigVMModel/RigVMController.cpp
@@ -6121,12 +6121,14 @@ bool URigVMController::SetPinDefaultValue(URigVMPin* InPin, const FString& InDef
 		ActionStack->BeginAction(Action);
 	}
 
+	const FString ClampedDefaultValue = InPin->IsRootPin() ? InPin->ClampDefaultValueFromMetaData(InDefaultValue) : InDefaultValue;
+
 	bool bSetPinDefaultValueSucceeded = false;
 	if (InPin->IsArray())
 	{
 		if (ShouldPinBeUnfolded(InPin))
 		{
-			TArray<FString> Elements = URigVMPin::SplitDefaultValue(InDefaultValue);
+			TArray<FString> Elements = URigVMPin::SplitDefaultValue(ClampedDefaultValue);
 
 			if (bResizeArrays)
 			{
@@ -6158,7 +6160,7 @@ bool URigVMController::SetPinDefaultValue(URigVMPin* InPin, const FString& InDef
 	}
 	else if (InPin->IsStruct())
 	{
-		TArray<FString> MemberValuePairs = URigVMPin::SplitDefaultValue(InDefaultValue);
+		TArray<FString> MemberValuePairs = URigVMPin::SplitDefaultValue(ClampedDefaultValue);
 
 		for (const FString& MemberValuePair : MemberValuePairs)
 		{
@@ -6183,7 +6185,7 @@ bool URigVMController::SetPinDefaultValue(URigVMPin* InPin, const FString& InDef
 	{
 		if (InPin->GetSubPins().Num() == 0)
 		{
-			InPin->DefaultValue = InDefaultValue;
+			InPin->DefaultValue = ClampedDefaultValue;
 			Notify(ERigVMGraphNotifType::PinDefaultValueChanged, InPin);
 			if (!bSuspendNotifications)
 			{
diff --git a/Engine/Source/Developer/RigVMDeveloper/Private/RigVMModel/RigVMPin.cpp b/Engine/Source/Developer/RigVMDeveloper/Private/RigVMModel/RigVMPin.cpp
--- a/Engine/Source/Developer/RigVMDeveloper/Private/RigVMModel/RigVMPin.cpp
+++ b/Engine/Source/Developer/RigVMDeveloper/Private/RigVMModel/RigVMPin.cpp
@@ -631,6 +631,34 @@ FString URigVMPin::GetDefaultValue(const URigVMPin::FPinOverride& InOverride) co
 	return DefaultValue;
 }
 
+template< typename Type>
+static FString ClampValue(const FString& InValueString, const FString& InMinValueString, const FString& InMaxValueString)
+{
+	FString RetValString = InValueString;
+	Type RetVal;
+	TTypeFromString<Type>::FromString(RetVal, *RetValString);
+
+	// Enforce min
+	if(!InMinValueString.IsEmpty())
+	{
+		checkSlow(InMinValueString.IsNumeric());
+		Type MinValue;
+		TTypeFromString<Type>::FromString(MinValue, *InMinValueString);
+		RetVal = FMath::Max<Type>(MinValue, RetVal);
+	}
+	//Enforce max 
+	if(!InMaxValueString.IsEmpty())
+	{
+		checkSlow(InMaxValueString.IsNumeric());
+		Type MaxValue;
+		TTypeFromString<Type>::FromString(MaxValue, *InMaxValueString);
+		RetVal = FMath::Min<Type>(MaxValue, RetVal);
+	}
+
+	RetValString = TTypeToString<Type>::ToString(RetVal);
+	return RetValString;
+}
+
 bool URigVMPin::IsValidDefaultValue(const FString& InDefaultValue) const
 {
 	TArray<FString> DefaultValues; 
@@ -761,6 +789,76 @@ bool URigVMPin::IsValidDefaultValue(const FString& InDefaultValue) const
 	return true;
 }
 
+FString URigVMPin::ClampDefaultValueFromMetaData(const FString& InDefaultValue) const
+{
+	FString RetVal = InDefaultValue;
+	if (URigVMUnitNode* UnitNode = Cast<URigVMUnitNode>(GetNode()))
+	{
+		TArray<FString> RetVals;
+		TArray<FString> DefaultValues; 
+
+		if (IsArray())
+		{
+			DefaultValues = URigVMPin::SplitDefaultValue(InDefaultValue);
+		}
+		else
+		{
+			DefaultValues.Add(InDefaultValue);
+		}
+
+		FString MinValue, MaxValue;	
+		if (UScriptStruct* ScriptStruct = UnitNode->GetScriptStruct())
+		{
+			if (FProperty* Property = ScriptStruct->FindPropertyByName(*GetName()))
+			{
+				MinValue = Property->GetMetaData(TEXT("ClampMin"));
+				MaxValue = Property->GetMetaData(TEXT("ClampMax"));
+			}
+		}
+		
+
+		FString BaseCPPType = GetCPPType()
+			.Replace(URigVMController::TArrayPrefix, TEXT(""))
+			.Replace(URigVMController::TObjectPtrPrefix, TEXT(""))
+			.Replace(TEXT(">"), TEXT(""));
+
+		RetVals.SetNumZeroed(DefaultValues.Num());
+		for (int32 Index = 0; Index < DefaultValues.Num(); ++Index)
+		{
+			const FString& Value = DefaultValues[Index]; 
+			
+			// perform single value validation
+			if (BaseCPPType == TEXT("float"))
+			{ 
+				RetVals[Index] = ClampValue<float>(Value, MinValue, MaxValue);
+			}
+			else if (BaseCPPType == TEXT("double"))
+			{ 
+				RetVals[Index] = ClampValue<double>(Value, MinValue, MaxValue);
+			}
+			else if (BaseCPPType == TEXT("int32"))
+			{ 
+				RetVals[Index] = ClampValue<int32>(Value, MinValue, MaxValue);
+			}
+			else
+			{
+				RetVals[Index] = Value;
+			}
+		}
+
+		if (IsArray())
+		{
+			RetVal = GetDefaultValueForArray(RetVals);
+		}
+		else
+		{
+			RetVal = RetVals[0];
+		}
+	}
+
+	return RetVal;
+}
+
 FName URigVMPin::GetCustomWidgetName() const
 {
 	if (IsArrayElement())
diff --git a/Engine/Source/Developer/RigVMDeveloper/Public/RigVMModel/RigVMPin.h b/Engine/Source/Developer/RigVMDeveloper/Public/RigVMModel/RigVMPin.h
--- a/Engine/Source/Developer/RigVMDeveloper/Public/RigVMModel/RigVMPin.h
+++ b/Engine/Source/Developer/RigVMDeveloper/Public/RigVMModel/RigVMPin.h
@@ -228,6 +228,9 @@ public:
 	// Returns true if the default value provided is valid
 	bool IsValidDefaultValue(const FString& InDefaultValue) const;
 
+	// Returns the default value clamped with the limit meta values defined by the UPROPERTY in URigVMUnitNodes 
+	FString ClampDefaultValueFromMetaData(const FString& InDefaultValue) const;
+
 	// Returns the name of a custom widget to be used
 	// for editing the Pin.
 	UFUNCTION(BlueprintCallable, Category = RigVMPin)
