Control Rig: Bubble up expressions linked to true and false blocks

diff --git a/Engine/Source/Developer/RigVMDeveloper/Private/RigVMCompiler/RigVMAST.cpp b/Engine/Source/Developer/RigVMDeveloper/Private/RigVMCompiler/RigVMAST.cpp
--- a/Engine/Source/Developer/RigVMDeveloper/Private/RigVMCompiler/RigVMAST.cpp
+++ b/Engine/Source/Developer/RigVMDeveloper/Private/RigVMCompiler/RigVMAST.cpp
@@ -1367,12 +1367,39 @@ void FRigVMParserAST::BubbleUpExpressions()
 
 		if (Blocks.Num() > 1)
 		{
-			// this expression is part of multiple blocks, and it needs to be bubbled up.
-			// for this we'll walk up the block tree and find the first block which contains all of them
 			TArray<FRigVMBlockExprAST*> BlockCandidates;
 			BlockCandidates.Append(Blocks);
 			FRigVMBlockExprAST* OuterBlock = nullptr;
 
+			// deal with a case where an expression is linked within both the true and false case of an "if" node
+			if(Blocks.Num() == 2)
+			{
+				const FRigVMExprAST* Parent0 = Blocks[0]->GetParent();
+				const FRigVMExprAST* Parent1 = Blocks[1]->GetParent();
+				if(Parent0 && Parent1)
+				{
+					const FRigVMExprAST* GrandParent0 = Parent0->GetParent();
+					const FRigVMExprAST* GrandParent1 = Parent1->GetParent();
+					if(GrandParent0 && GrandParent1 && GrandParent0 == GrandParent1)
+					{
+						if(GrandParent0->IsA(FRigVMExprAST::EType::If))
+						{
+							const FRigVMIfExprAST* IfExpression = GrandParent0->To<FRigVMIfExprAST>();
+							const FRigVMExprAST* ConditionBlockExpression = IfExpression->GetConditionExpr()->GetFirstChildOfType(FRigVMExprAST::Block);
+							if(ConditionBlockExpression)
+							{
+								OuterBlock = (FRigVMBlockExprAST*)ConditionBlockExpression->To<FRigVMBlockExprAST>();
+								OuterBlock->Children.Add(Expression);
+								Expression->Parents.Insert(OuterBlock, 0);
+								continue;
+							}
+						}
+					}
+				}
+			}
+
+			// this expression is part of multiple blocks, and it needs to be bubbled up.
+			// for this we'll walk up the block tree and find the first block which contains all of them
 			for (int32 BlockCandidateIndex = 0; BlockCandidateIndex < BlockCandidates.Num(); BlockCandidateIndex++)
 			{
 				FRigVMBlockExprAST* BlockCandidate = BlockCandidates[BlockCandidateIndex];
