Control Rig: Fix creating reroute nodes from pins with existing connections disconnects all connections except for reroute

diff --git a/Engine/Source/Developer/RigVMDeveloper/Private/RigVMModel/RigVMController.cpp b/Engine/Source/Developer/RigVMDeveloper/Private/RigVMModel/RigVMController.cpp
--- a/Engine/Source/Developer/RigVMDeveloper/Private/RigVMModel/RigVMController.cpp
+++ b/Engine/Source/Developer/RigVMDeveloper/Private/RigVMModel/RigVMController.cpp
@@ -1910,7 +1910,10 @@ URigVMRerouteNode* URigVMController::AddRerouteNodeOnPin(const FString& InPinPat
 
 	//in case an injected node is present, use its pins for any new links
 	URigVMPin *PinForLink = Pin->GetPinForLink(); 
-	BreakAllLinks(PinForLink, bAsInput, bSetupUndoRedo);
+	if (bAsInput)
+	{
+		BreakAllLinks(PinForLink, bAsInput, bSetupUndoRedo);
+	}
 
 	FString Name = GetValidNodeName(InNodeName.IsEmpty() ? FString(TEXT("RerouteNode")) : InNodeName);
 	URigVMRerouteNode* Node = NewObject<URigVMRerouteNode>(Graph, *Name);
