Control Rig: Fix unit test failing RigVM.Graph.PinDefaultValueValidation

diff --git a/Engine/Source/Developer/RigVMDeveloper/Private/RigVMModel/RigVMPin.cpp b/Engine/Source/Developer/RigVMDeveloper/Private/RigVMModel/RigVMPin.cpp
--- a/Engine/Source/Developer/RigVMDeveloper/Private/RigVMModel/RigVMPin.cpp
+++ b/Engine/Source/Developer/RigVMDeveloper/Private/RigVMModel/RigVMPin.cpp
@@ -826,19 +826,26 @@ FString URigVMPin::ClampDefaultValueFromMetaData(const FString& InDefaultValue)
 		for (int32 Index = 0; Index < DefaultValues.Num(); ++Index)
 		{
 			const FString& Value = DefaultValues[Index]; 
-			
-			// perform single value validation
-			if (BaseCPPType == TEXT("float"))
-			{ 
-				RetVals[Index] = ClampValue<float>(Value, MinValue, MaxValue);
-			}
-			else if (BaseCPPType == TEXT("double"))
-			{ 
-				RetVals[Index] = ClampValue<double>(Value, MinValue, MaxValue);
-			}
-			else if (BaseCPPType == TEXT("int32"))
-			{ 
-				RetVals[Index] = ClampValue<int32>(Value, MinValue, MaxValue);
+
+			if (!MinValue.IsEmpty() || !MaxValue.IsEmpty())
+			{
+				// perform single value validation
+				if (BaseCPPType == TEXT("float"))
+				{ 
+					RetVals[Index] = ClampValue<float>(Value, MinValue, MaxValue);
+				}
+				else if (BaseCPPType == TEXT("double"))
+				{ 
+					RetVals[Index] = ClampValue<double>(Value, MinValue, MaxValue);
+				}
+				else if (BaseCPPType == TEXT("int32"))
+				{ 
+					RetVals[Index] = ClampValue<int32>(Value, MinValue, MaxValue);
+				}
+				else
+				{
+					RetVals[Index] = Value;
+				}
 			}
 			else
 			{
