Control Rig: setup a redirector to redirect old spring interp struct name to new ones. This should clarify that on release there is only one version of spring interp quat that we support.

diff --git a/Engine/Config/BaseEngine.ini b/Engine/Config/BaseEngine.ini
--- a/Engine/Config/BaseEngine.ini
+++ b/Engine/Config/BaseEngine.ini
@@ -862,6 +862,7 @@ NearClipPlane=10.0
 +ClassRedirects=(OldName="WaterMeshActorFactory",NewName="/Script/Water.WaterZoneActorFactory")
 
 ;5.0
++StructRedirects=(OldName="/Script/ControlRig.RigUnit_SpringInterpQuaternion",NewName="/Script/ControlRig.RigUnit_SpringInterpQuaternionV2")
 +ClassRedirects=(OldName="RigVMStructNode",NewName="/Script/RigVMDeveloper.RigVMUnitNode")
 +EnumRedirects=(OldName="/Script/Slate.ETextTransformPolicy", NewName="/Script/SlateCore.ETextTransformPolicy")
 +PropertyRedirects=(OldName="BlendProfile.BlendProfileMode", NewName="Mode")
diff --git a/Engine/Plugins/Animation/ControlRig/Source/ControlRig/Private/Units/Simulation/RigUnit_SpringInterp.h b/Engine/Plugins/Animation/ControlRig/Source/ControlRig/Private/Units/Simulation/RigUnit_SpringInterp.h
--- a/Engine/Plugins/Animation/ControlRig/Source/ControlRig/Private/Units/Simulation/RigUnit_SpringInterp.h
+++ b/Engine/Plugins/Animation/ControlRig/Source/ControlRig/Private/Units/Simulation/RigUnit_SpringInterp.h
@@ -278,6 +278,7 @@ struct CONTROLRIG_API FRigUnit_SpringInterpQuaternionV2 : public FRigUnit_SimBas
 {
 	GENERATED_BODY()
 	
+	// RigUnit_SpringInterpQuaternion is redirected to RigUnit_SpringInterpQuaternionV2
 	FRigUnit_SpringInterpQuaternionV2()
 	{
 		bUseCurrentInput = false;
diff --git a/Engine/Source/Developer/RigVMDeveloper/Private/RigVMModel/RigVMController.cpp b/Engine/Source/Developer/RigVMDeveloper/Private/RigVMModel/RigVMController.cpp
--- a/Engine/Source/Developer/RigVMDeveloper/Private/RigVMModel/RigVMController.cpp
+++ b/Engine/Source/Developer/RigVMDeveloper/Private/RigVMModel/RigVMController.cpp
@@ -11546,7 +11546,10 @@ void URigVMController::RepopulatePinsOnNode(URigVMNode* InNode, bool bFollowCore
 	else if (CollapseNode)
 	{
 		FRigVMControllerGraphGuard GraphGuard(this, CollapseNode->GetContainedGraph(), false);
-		for (URigVMNode* ContainedNode : CollapseNode->GetContainedNodes())
+		// need to get a copy of the node array since the following function could remove nodes from the graph
+		// we don't want to remove elements from the array we are iterating over.
+		TArray<URigVMNode*> ContainedNodes = CollapseNode->GetContainedNodes();
+		for (URigVMNode* ContainedNode : ContainedNodes)
 		{
 			RepopulatePinsOnNode(ContainedNode, bFollowCoreRedirectors, bNotify, bSetupOrphanedPins);
 		}
