RigVM: Deal with faulty nodes within a top level graph

diff --git a/Engine/Source/Developer/RigVMDeveloper/Private/RigVMModel/RigVMController.cpp b/Engine/Source/Developer/RigVMDeveloper/Private/RigVMModel/RigVMController.cpp
--- a/Engine/Source/Developer/RigVMDeveloper/Private/RigVMModel/RigVMController.cpp
+++ b/Engine/Source/Developer/RigVMDeveloper/Private/RigVMModel/RigVMController.cpp
@@ -4718,7 +4718,12 @@ bool URigVMController::RemoveNode(URigVMNode* InNode, bool bSetupUndoRedo, bool
 		if ((Cast<URigVMFunctionEntryNode>(InNode) != nullptr && InNode->GetName() == TEXT("Entry")) ||
 			(Cast<URigVMFunctionReturnNode>(InNode) != nullptr && InNode->GetName() == TEXT("Return")))
 		{
-			return false;
+			// due to earlier bugs in the copy & paste code entry and return nodes could end up in
+			// root graphs - in those cases we allow deletion
+			if(!Graph->IsRootGraph())
+			{
+				return false;
+			}
 		}
 
 		// check if the operation will cause to dirty assets
@@ -11539,8 +11544,8 @@ void URigVMController::RepopulatePinsOnNode(URigVMNode* InNode, bool bFollowCore
 		}
 		else
 		{
-			// in the future we'll likely have function libraries as outers here
-			checkNoEntry();
+			// due to earlier bugs with copy and paste we can find entry and return nodes under the top level
+			// graph. we'll ignore these for now.
 		}
 	}
 	else if (CollapseNode)
