RigVM: Less ambiguous object search for cpptypes in URigVMController::EnsurePinValidity

#edigrate of 18125091

diff --git a/Engine/Source/Developer/RigVMDeveloper/Private/RigVMModel/RigVMController.cpp b/Engine/Source/Developer/RigVMDeveloper/Private/RigVMModel/RigVMController.cpp
--- a/Engine/Source/Developer/RigVMDeveloper/Private/RigVMModel/RigVMController.cpp
+++ b/Engine/Source/Developer/RigVMDeveloper/Private/RigVMModel/RigVMController.cpp
@@ -12469,7 +12469,7 @@ bool URigVMController::EnsurePinValidity(URigVMPin* InPin, bool bRecursive)
 		{
 			// try to find the CPPTypeObject by name
 			FString CPPType = InPin->IsArray() ? InPin->GetArrayElementCppType() : InPin->GetCPPType();
-			UObject* CPPTypeObject = FindObject<UObject>(ANY_PACKAGE, *CPPType);
+			UObject* CPPTypeObject = FindObject<UField>(ANY_PACKAGE, *CPPType);
 			if(CPPTypeObject == nullptr)
 			{
 				if(CPPType.StartsWith(TEXT("E"), ESearchCase::CaseSensitive) ||
@@ -12477,7 +12477,7 @@ bool URigVMController::EnsurePinValidity(URigVMPin* InPin, bool bRecursive)
 					CPPType.StartsWith(TEXT("U"), ESearchCase::CaseSensitive))
 				{
 					CPPType = CPPType.Mid(1);
-					CPPTypeObject = FindObject<UObject>(ANY_PACKAGE, *CPPType);
+					CPPTypeObject = FindObject<UField>(ANY_PACKAGE, *CPPType);
 				}
 
 				if(CPPTypeObject == nullptr)
