RigVM: Implement undo / redo for promote actions

diff --git a/Engine/Source/Developer/RigVMDeveloper/Private/RigVMModel/RigVMController.cpp b/Engine/Source/Developer/RigVMDeveloper/Private/RigVMModel/RigVMController.cpp
--- a/Engine/Source/Developer/RigVMDeveloper/Private/RigVMModel/RigVMController.cpp
+++ b/Engine/Source/Developer/RigVMDeveloper/Private/RigVMModel/RigVMController.cpp
@@ -4355,7 +4355,7 @@ TArray<URigVMNode*> URigVMController::ExpandLibraryNode(URigVMLibraryNode* InNod
 	return ExpandedNodes;
 }
 
-FName URigVMController::PromoteCollapseNodeToFunctionReferenceNode(const FName& InNodeName, bool bSetupUndoRedo, bool bPrintPythonCommand)
+FName URigVMController::PromoteCollapseNodeToFunctionReferenceNode(const FName& InNodeName, bool bSetupUndoRedo, bool bPrintPythonCommand, const FString& InExistingFunctionDefinitionPath)
 {
 	if (!IsValidGraph())
 	{
@@ -4365,7 +4365,7 @@ FName URigVMController::PromoteCollapseNodeToFunctionReferenceNode(const FName&
 	URigVMGraph* Graph = GetGraph();
 	check(Graph);
 
-	URigVMNode* Result = PromoteCollapseNodeToFunctionReferenceNode(Cast<URigVMCollapseNode>(Graph->FindNodeByName(InNodeName)), bSetupUndoRedo);
+	URigVMNode* Result = PromoteCollapseNodeToFunctionReferenceNode(Cast<URigVMCollapseNode>(Graph->FindNodeByName(InNodeName)), bSetupUndoRedo, InExistingFunctionDefinitionPath);
 	if (Result)
 	{
 		if (bPrintPythonCommand)
@@ -4383,7 +4383,7 @@ FName URigVMController::PromoteCollapseNodeToFunctionReferenceNode(const FName&
 	return NAME_None;
 }
 
-FName URigVMController::PromoteFunctionReferenceNodeToCollapseNode(const FName& InNodeName, bool bSetupUndoRedo, bool bPrintPythonCommand)
+FName URigVMController::PromoteFunctionReferenceNodeToCollapseNode(const FName& InNodeName, bool bSetupUndoRedo, bool bPrintPythonCommand, bool bRemoveFunctionDefinition)
 {
 	if (!IsValidGraph())
 	{
@@ -4393,7 +4393,7 @@ FName URigVMController::PromoteFunctionReferenceNodeToCollapseNode(const FName&
 	URigVMGraph* Graph = GetGraph();
 	check(Graph);
 
-	URigVMNode* Result = PromoteFunctionReferenceNodeToCollapseNode(Cast<URigVMFunctionReferenceNode>(Graph->FindNodeByName(InNodeName)), bSetupUndoRedo);
+	URigVMNode* Result = PromoteFunctionReferenceNodeToCollapseNode(Cast<URigVMFunctionReferenceNode>(Graph->FindNodeByName(InNodeName)), bSetupUndoRedo, bRemoveFunctionDefinition);
 	if (Result)
 	{
 		return Result->GetFName();
@@ -4401,7 +4401,7 @@ FName URigVMController::PromoteFunctionReferenceNodeToCollapseNode(const FName&
 	return NAME_None;
 }
 
-URigVMFunctionReferenceNode* URigVMController::PromoteCollapseNodeToFunctionReferenceNode(URigVMCollapseNode* InCollapseNode, bool bSetupUndoRedo)
+URigVMFunctionReferenceNode* URigVMController::PromoteCollapseNodeToFunctionReferenceNode(URigVMCollapseNode* InCollapseNode, bool bSetupUndoRedo, const FString& InExistingFunctionDefinitionPath)
 {
 	if (!IsValidNodeForGraph(InCollapseNode))
 	{
@@ -4418,28 +4418,31 @@ URigVMFunctionReferenceNode* URigVMController::PromoteCollapseNodeToFunctionRefe
 	}
 
 	FRigVMControllerCompileBracketScope CompileScope(this);
-	if (bSetupUndoRedo)
-	{
-		OpenUndoBracket(TEXT("Promote to Function"));
-	}
-
 	URigVMFunctionReferenceNode* FunctionRefNode = nullptr;
 
 	// Create Function
 	URigVMLibraryNode* FunctionDefinition = nullptr;
+	if (!InExistingFunctionDefinitionPath.IsEmpty())
 	{
-		FRigVMControllerGraphGuard GraphGuard(this, FunctionLibrary, bSetupUndoRedo);
-		FString FunctionName = GetValidNodeName(InCollapseNode->GetName());
-		FunctionDefinition = AddFunctionToLibrary(*FunctionName, InCollapseNode->IsMutable());		
+		FunctionDefinition = FindObject<URigVMLibraryNode>(ANY_PACKAGE, *InExistingFunctionDefinitionPath);
 	}
 
-	// Add interface pins in function
-	if (FunctionDefinition)
+	if (FunctionDefinition == nullptr)
 	{
-		FRigVMControllerGraphGuard GraphGuard(this, FunctionDefinition->GetContainedGraph(), bSetupUndoRedo);
-		for(const URigVMPin* Pin : InCollapseNode->GetPins())
 		{
-			AddExposedPin(Pin->GetFName(), Pin->GetDirection(), Pin->GetCPPType(), (Pin->GetCPPTypeObject() ? *Pin->GetCPPTypeObject()->GetPathName() : TEXT("")), Pin->GetDefaultValue(), bSetupUndoRedo);
+			FRigVMControllerGraphGuard GraphGuard(this, FunctionLibrary, false);
+			const FString FunctionName = GetValidNodeName(InCollapseNode->GetName());
+			FunctionDefinition = AddFunctionToLibrary(*FunctionName, InCollapseNode->IsMutable(), FVector2D::Zero(), false);		
+		}
+	
+		// Add interface pins in function
+		if (FunctionDefinition)
+		{
+			FRigVMControllerGraphGuard GraphGuard(this, FunctionDefinition->GetContainedGraph(), false);
+			for(const URigVMPin* Pin : InCollapseNode->GetPins())
+			{
+				AddExposedPin(Pin->GetFName(), Pin->GetDirection(), Pin->GetCPPType(), (Pin->GetCPPTypeObject() ? *Pin->GetCPPTypeObject()->GetPathName() : TEXT("")), Pin->GetDefaultValue(), false);
+			}
 		}
 	}
 
@@ -4448,7 +4451,7 @@ URigVMFunctionReferenceNode* URigVMController::PromoteCollapseNodeToFunctionRefe
 	{
 		FString TextContent;
 		{
-			FRigVMControllerGraphGuard GraphGuard(this, InCollapseNode->GetContainedGraph(), bSetupUndoRedo);
+			FRigVMControllerGraphGuard GraphGuard(this, InCollapseNode->GetContainedGraph(), false);
 			TArray<FName> NodeNames;
 			for (const URigVMNode* Node : InCollapseNode->GetContainedNodes())
 			{
@@ -4462,16 +4465,16 @@ URigVMFunctionReferenceNode* URigVMController::PromoteCollapseNodeToFunctionRefe
 			TextContent = ExportNodesToText(NodeNames);
 		}
 		{
-			FRigVMControllerGraphGuard GraphGuard(this, FunctionDefinition->GetContainedGraph(), bSetupUndoRedo);
-			ImportNodesFromText(TextContent, bSetupUndoRedo);
+			FRigVMControllerGraphGuard GraphGuard(this, FunctionDefinition->GetContainedGraph(), false);
+			ImportNodesFromText(TextContent, false);
 			if (FunctionDefinition->GetContainedGraph()->GetEntryNode() && InCollapseNode->GetContainedGraph()->GetEntryNode())
 			{ 
-				SetNodePosition(FunctionDefinition->GetContainedGraph()->GetEntryNode(), InCollapseNode->GetContainedGraph()->GetEntryNode()->GetPosition());
+				SetNodePosition(FunctionDefinition->GetContainedGraph()->GetEntryNode(), InCollapseNode->GetContainedGraph()->GetEntryNode()->GetPosition(), false);
 			}
 			
 			if (FunctionDefinition->GetContainedGraph()->GetReturnNode() && InCollapseNode->GetContainedGraph()->GetReturnNode())
 			{ 
-				SetNodePosition(FunctionDefinition->GetContainedGraph()->GetReturnNode(), InCollapseNode->GetContainedGraph()->GetReturnNode()->GetPosition());
+				SetNodePosition(FunctionDefinition->GetContainedGraph()->GetReturnNode(), InCollapseNode->GetContainedGraph()->GetReturnNode()->GetPosition(), false);
 			}
 
 			for (const URigVMLink* InnerLink : InCollapseNode->GetContainedGraph()->GetLinks())
@@ -4482,7 +4485,7 @@ URigVMFunctionReferenceNode* URigVMController::PromoteCollapseNodeToFunctionRefe
 				{
 					if (!SourcePin->IsLinkedTo(TargetPin))
 					{
-						AddLink(InnerLink->SourcePinPath, InnerLink->TargetPinPath, bSetupUndoRedo);	
+						AddLink(InnerLink->SourcePinPath, InnerLink->TargetPinPath, false);	
 					}
 				}				
 			}
@@ -4503,29 +4506,29 @@ URigVMFunctionReferenceNode* URigVMController::PromoteCollapseNodeToFunctionRefe
 			LinkPaths.Add(TPair< FString, FString >(Link->GetSourcePin()->GetPinPath(), Link->GetTargetPin()->GetPinPath()));
 		}
 
-		RemoveNode(InCollapseNode, bSetupUndoRedo);
+		RemoveNode(InCollapseNode, false, true);
 
-		FunctionRefNode = AddFunctionReferenceNode(FunctionDefinition, NodePosition, NodeName, bSetupUndoRedo);
+		FunctionRefNode = AddFunctionReferenceNode(FunctionDefinition, NodePosition, NodeName, false);
 
 		if (FunctionRefNode)
 		{
 			ApplyPinStates(FunctionRefNode, PinStates);
 			for (const TPair<FString, FString>& LinkPath : LinkPaths)
 			{
-				AddLink(LinkPath.Key, LinkPath.Value, bSetupUndoRedo);
+				AddLink(LinkPath.Key, LinkPath.Value, false);
 			}
 		}
-	}
 
-	if (bSetupUndoRedo)
-	{
-		CloseUndoBracket();
+		if (bSetupUndoRedo)
+		{
+			ActionStack->AddAction(FRigVMPromoteNodeAction(InCollapseNode, NodeName, FString()));
+		}
 	}
 
 	return FunctionRefNode;
 }
 
-URigVMCollapseNode* URigVMController::PromoteFunctionReferenceNodeToCollapseNode(URigVMFunctionReferenceNode* InFunctionRefNode, bool bSetupUndoRedo)
+URigVMCollapseNode* URigVMController::PromoteFunctionReferenceNodeToCollapseNode(URigVMFunctionReferenceNode* InFunctionRefNode, bool bSetupUndoRedo, bool bRemoveFunctionDefinition)
 {
 	if (!IsValidNodeForGraph(InFunctionRefNode))
 	{
@@ -4588,10 +4591,6 @@ URigVMCollapseNode* URigVMController::PromoteFunctionReferenceNodeToCollapseNode
 	}
 
 	FRigVMControllerCompileBracketScope CompileScope(this);
-	if (bSetupUndoRedo)
-	{
-		OpenUndoBracket(TEXT("Promote to Collapse Node"));
-	}
 
 	FString NodeName = InFunctionRefNode->GetName();
 	FVector2D NodePosition = InFunctionRefNode->GetPosition();
@@ -4604,7 +4603,7 @@ URigVMCollapseNode* URigVMController::PromoteFunctionReferenceNodeToCollapseNode
 		LinkPaths.Add(TPair< FString, FString >(Link->GetSourcePin()->GetPinPath(), Link->GetTargetPin()->GetPinPath()));
 	}
 
-	RemoveNode(InFunctionRefNode, bSetupUndoRedo);
+	RemoveNode(InFunctionRefNode, false, true);
 
 	if (RequestNewExternalVariableDelegate.IsBound())
 	{
@@ -4618,7 +4617,7 @@ URigVMCollapseNode* URigVMController::PromoteFunctionReferenceNodeToCollapseNode
 	if(CollapseNode)
 	{
 		{
-			FRigVMControllerGraphGuard Guard(this, CollapseNode->GetContainedGraph(), bSetupUndoRedo);
+			FRigVMControllerGraphGuard Guard(this, CollapseNode->GetContainedGraph(), false);
 			ReattachLinksToPinObjects();
 
 			for (URigVMNode* Node : CollapseNode->GetContainedGraph()->GetNodes())
@@ -4643,13 +4642,19 @@ URigVMCollapseNode* URigVMController::PromoteFunctionReferenceNodeToCollapseNode
 		ApplyPinStates(CollapseNode, PinStates);
 		for (const TPair<FString, FString>& LinkPath : LinkPaths)
 		{
-			AddLink(LinkPath.Key, LinkPath.Value, bSetupUndoRedo);
+			AddLink(LinkPath.Key, LinkPath.Value, false);
 		}
 	}
 
 	if (bSetupUndoRedo)
 	{
-		CloseUndoBracket();
+		ActionStack->AddAction(FRigVMPromoteNodeAction(InFunctionRefNode, NodeName, FunctionDefinition->GetPathName()));
+	}
+
+	if(bRemoveFunctionDefinition)
+	{
+		FRigVMControllerGraphGuard Guard(this, FunctionDefinition->GetRootGraph(), false);
+		RemoveFunctionFromLibrary(FunctionDefinition->GetFName(), false);
 	}
 
 	return CollapseNode;
diff --git a/Engine/Source/Developer/RigVMDeveloper/Private/RigVMModel/RigVMControllerActions.cpp b/Engine/Source/Developer/RigVMDeveloper/Private/RigVMModel/RigVMControllerActions.cpp
--- a/Engine/Source/Developer/RigVMDeveloper/Private/RigVMModel/RigVMControllerActions.cpp
+++ b/Engine/Source/Developer/RigVMDeveloper/Private/RigVMModel/RigVMControllerActions.cpp
@@ -1999,3 +1999,40 @@ bool FRigVMAddArrayNodeAction::Redo(URigVMController* InController)
 #endif
 	return false;
 }
+
+FRigVMPromoteNodeAction::FRigVMPromoteNodeAction()
+	: LibraryNodePath()
+	, FunctionDefinitionPath()
+	, bFromFunctionToCollapseNode(false)
+{
+}
+
+FRigVMPromoteNodeAction::FRigVMPromoteNodeAction(const URigVMNode* InNodeToPromote, const FString& InNodePath, const FString& InFunctionDefinitionPath)
+	: LibraryNodePath(InNodePath)
+	, FunctionDefinitionPath(InFunctionDefinitionPath)
+	, bFromFunctionToCollapseNode(InNodeToPromote->IsA<URigVMFunctionReferenceNode>())
+{
+}
+
+bool FRigVMPromoteNodeAction::Undo(URigVMController* InController)
+{
+	if(bFromFunctionToCollapseNode)
+	{
+		const FName FunctionRefNodeName = InController->PromoteCollapseNodeToFunctionReferenceNode(*LibraryNodePath, false, false, FunctionDefinitionPath);
+		return FunctionRefNodeName.ToString() == LibraryNodePath;
+	}
+
+	const FName CollapseNodeName = InController->PromoteFunctionReferenceNodeToCollapseNode(*LibraryNodePath, false, false, true);
+	return CollapseNodeName.ToString() == LibraryNodePath;
+}
+
+bool FRigVMPromoteNodeAction::Redo(URigVMController* InController)
+{
+	if(bFromFunctionToCollapseNode)
+	{
+		const FName CollapseNodeName = InController->PromoteFunctionReferenceNodeToCollapseNode(*LibraryNodePath, false, false, false);
+		return CollapseNodeName.ToString() == LibraryNodePath;
+	}
+	const FName FunctionRefNodeName = InController->PromoteCollapseNodeToFunctionReferenceNode(*LibraryNodePath, false, false);
+	return FunctionRefNodeName.ToString() == LibraryNodePath;
+}
diff --git a/Engine/Source/Developer/RigVMDeveloper/Public/RigVMModel/RigVMController.h b/Engine/Source/Developer/RigVMDeveloper/Public/RigVMModel/RigVMController.h
--- a/Engine/Source/Developer/RigVMDeveloper/Public/RigVMModel/RigVMController.h
+++ b/Engine/Source/Developer/RigVMDeveloper/Public/RigVMModel/RigVMController.h
@@ -385,11 +385,11 @@ public:
 
 	// Turns a collapse node into a function node
 	UFUNCTION(BlueprintCallable, Category = RigVMController)
-	FName PromoteCollapseNodeToFunctionReferenceNode(const FName& InNodeName, bool bSetupUndoRedo = true, bool bPrintPythonCommand = false);
+	FName PromoteCollapseNodeToFunctionReferenceNode(const FName& InNodeName, bool bSetupUndoRedo = true, bool bPrintPythonCommand = false, const FString& InExistingFunctionDefinitionPath = TEXT(""));
 
 	// Turns a collapse node into a function node
 	UFUNCTION(BlueprintCallable, Category = RigVMController)
-	FName PromoteFunctionReferenceNodeToCollapseNode(const FName& InNodeName, bool bSetupUndoRedo = true, bool bPrintPythonCommand = false);
+	FName PromoteFunctionReferenceNodeToCollapseNode(const FName& InNodeName, bool bSetupUndoRedo = true, bool bPrintPythonCommand = false, bool bRemoveFunctionDefinition = false);
 
 #endif
 
@@ -830,8 +830,8 @@ private:
 	static void ForEveryPinRecursively(URigVMNode* InNode, TFunction<void(URigVMPin*)> OnEachPinFunction);
 	URigVMCollapseNode* CollapseNodes(const TArray<URigVMNode*>& InNodes, const FString& InCollapseNodeName, bool bSetupUndoRedo);
 	TArray<URigVMNode*> ExpandLibraryNode(URigVMLibraryNode* InNode, bool bSetupUndoRedo);
-	URigVMFunctionReferenceNode* PromoteCollapseNodeToFunctionReferenceNode(URigVMCollapseNode* InCollapseNode, bool bSetupUndoRedo);
-	URigVMCollapseNode* PromoteFunctionReferenceNodeToCollapseNode(URigVMFunctionReferenceNode* InFunctionRefNode, bool bSetupUndoRedo);
+	URigVMFunctionReferenceNode* PromoteCollapseNodeToFunctionReferenceNode(URigVMCollapseNode* InCollapseNode, bool bSetupUndoRedo, const FString& InExistingFunctionDefinitionPath);
+	URigVMCollapseNode* PromoteFunctionReferenceNodeToCollapseNode(URigVMFunctionReferenceNode* InFunctionRefNode, bool bSetupUndoRedo, bool bRemoveFunctionDefinition);
 	void SetReferencedFunction(URigVMFunctionReferenceNode* InFunctionRefNode, URigVMLibraryNode* InNewReferencedNode, bool bSetupUndoRedo);
 
 	void RefreshFunctionPins(URigVMNode* InNode, bool bNotify = true);
diff --git a/Engine/Source/Developer/RigVMDeveloper/Public/RigVMModel/RigVMControllerActions.h b/Engine/Source/Developer/RigVMDeveloper/Public/RigVMModel/RigVMControllerActions.h
--- a/Engine/Source/Developer/RigVMDeveloper/Public/RigVMModel/RigVMControllerActions.h
+++ b/Engine/Source/Developer/RigVMDeveloper/Public/RigVMModel/RigVMControllerActions.h
@@ -1671,3 +1671,29 @@ public:
 	UPROPERTY()
 	FString NodePath;
 };
+
+/**
+ * An action to promote a function to collapse node or vice versa
+ */
+USTRUCT()
+struct FRigVMPromoteNodeAction : public FRigVMBaseAction
+{
+	GENERATED_BODY()
+
+public:
+
+	FRigVMPromoteNodeAction();
+	FRigVMPromoteNodeAction(const URigVMNode* InNodeToPromote, const FString& InNodePath, const FString& InFunctionDefinitionPath);
+	virtual ~FRigVMPromoteNodeAction() {};
+	virtual bool Undo(URigVMController* InController) override;
+	virtual bool Redo(URigVMController* InController) override;
+
+	UPROPERTY()
+	FString LibraryNodePath;
+
+	UPROPERTY()
+	FString FunctionDefinitionPath;
+
+	UPROPERTY()
+	bool bFromFunctionToCollapseNode;
+};
\ No newline at end of file
