RigVM: Followup fixing build breakage in main

diff --git a/Engine/Source/Developer/RigVMDeveloper/Private/RigVMModel/RigVMController.cpp b/Engine/Source/Developer/RigVMDeveloper/Private/RigVMModel/RigVMController.cpp
--- a/Engine/Source/Developer/RigVMDeveloper/Private/RigVMModel/RigVMController.cpp
+++ b/Engine/Source/Developer/RigVMDeveloper/Private/RigVMModel/RigVMController.cpp
@@ -4432,7 +4432,7 @@ URigVMFunctionReferenceNode* URigVMController::PromoteCollapseNodeToFunctionRefe
 		{
 			FRigVMControllerGraphGuard GraphGuard(this, FunctionLibrary, false);
 			const FString FunctionName = GetValidNodeName(InCollapseNode->GetName());
-			FunctionDefinition = AddFunctionToLibrary(*FunctionName, InCollapseNode->IsMutable(), FVector2D::Zero(), false);		
+			FunctionDefinition = AddFunctionToLibrary(*FunctionName, InCollapseNode->IsMutable(), FVector2D::ZeroVector, false);		
 		}
 	
 		// Add interface pins in function
