Control Rig: Fix binding to local variables and input arguments

diff --git a/Engine/Source/Developer/RigVMDeveloper/Private/RigVMCompiler/RigVMAST.cpp b/Engine/Source/Developer/RigVMDeveloper/Private/RigVMCompiler/RigVMAST.cpp
--- a/Engine/Source/Developer/RigVMDeveloper/Private/RigVMCompiler/RigVMAST.cpp
+++ b/Engine/Source/Developer/RigVMDeveloper/Private/RigVMCompiler/RigVMAST.cpp
@@ -1056,8 +1056,16 @@ FRigVMExprAST* FRigVMParserAST::TraversePin(const FRigVMASTProxy& InPinProxy, FR
 		Links.Num() == 0 &&
 		SubPinsBoundToVariables.Num() == 0)
 	{
-		if (Cast<URigVMParameterNode>(Pin->GetNode()) ||
-			Cast<URigVMVariableNode>(Pin->GetNode()))
+		if (Pin->IsBoundToExternalVariable(PinOverride))
+		{
+			PinExpr = MakeExpr<FRigVMExternalVarExprAST>(InPinProxy);
+		}
+		else if (Pin->IsBoundToLocalVariable(PinOverride))
+		{
+			PinExpr = MakeExpr<FRigVMVarExprAST>(FRigVMExprAST::EType::Var, InPinProxy);
+		}
+		else if (Cast<URigVMParameterNode>(Pin->GetNode()) ||
+				 Cast<URigVMVariableNode>(Pin->GetNode()))
 		{
 			PinExpr = MakeExpr<FRigVMVarExprAST>(FRigVMExprAST::EType::Var, InPinProxy);
 			FRigVMExprAST* PinLiteralExpr = MakeExpr<FRigVMLiteralExprAST>(InPinProxy);
@@ -1066,10 +1074,6 @@ FRigVMExprAST* FRigVMParserAST::TraversePin(const FRigVMASTProxy& InPinProxy, FR
 			PinCopyExpr->AddParent(PinExpr);
 			PinLiteralExpr->AddParent(PinCopyExpr);
 		}
-		else if (Pin->IsBoundToVariable(PinOverride))
-		{
-			PinExpr = MakeExpr<FRigVMExternalVarExprAST>(InPinProxy);
-		}
 		else
 		{
 			PinExpr = MakeExpr<FRigVMLiteralExprAST>(InPinProxy);
@@ -3013,9 +3017,24 @@ void FRigVMParserAST::Inline(URigVMGraph* InGraph, const TArray<FRigVMASTProxy>&
 				}
 
 				TArray<URigVMLink*> SourceLinks = ChildPin->GetSourceLinks(false /* recursive */);
-				if (SourceLinks.Num() > 0)
+				if (SourceLinks.Num() > 0 || ChildPin->IsBoundToInputArgument())
 				{
-					URigVMPin* SourcePin = SourceLinks[0]->GetSourcePin();
+					URigVMPin* SourcePin = nullptr;
+					if (SourceLinks.Num() > 0)
+					{
+						SourcePin = SourceLinks[0]->GetSourcePin();
+					}
+					else
+					{
+						if (URigVMGraph* Graph = ChildPin->GetGraph())
+						{
+							if (URigVMFunctionEntryNode* EntryNode = Graph->GetEntryNode())
+							{
+								SourcePin = EntryNode->FindPin(ChildPin->GetBoundVariableName());
+							}
+						}
+					}
+					check(SourcePin);
 					SourcePinProxy = InPinProxy.GetSibling(SourcePin);
 
 					// only continue the recursion on reroutes
diff --git a/Engine/Source/Developer/RigVMDeveloper/Private/RigVMCompiler/RigVMCompiler.cpp b/Engine/Source/Developer/RigVMDeveloper/Private/RigVMCompiler/RigVMCompiler.cpp
--- a/Engine/Source/Developer/RigVMDeveloper/Private/RigVMCompiler/RigVMCompiler.cpp
+++ b/Engine/Source/Developer/RigVMDeveloper/Private/RigVMCompiler/RigVMCompiler.cpp
@@ -340,7 +340,22 @@ bool URigVMCompiler::Compile(URigVMGraph* InGraph, URigVMController* InControlle
 	{
 		for (const FRigVMGraphVariableDescription& LocalVariable : VisitedGraph->LocalVariables)
 		{
-			// To create the default value in the literal memory, we need to find a pin in a variable node that
+			auto AddDefaultValueOperand = [&](URigVMPin* Pin)
+			{
+				FRigVMASTProxy PinProxy = FRigVMASTProxy::MakeFromUObject(Pin);
+				FRigVMVarExprAST* TempVarExpr = AST->MakeExpr<FRigVMVarExprAST>(FRigVMExprAST::EType::Literal, PinProxy);
+				FRigVMOperand Operand = FindOrAddRegister(TempVarExpr, WorkData, false);
+
+#if UE_RIGVM_UCLASS_BASED_STORAGE_DISABLED
+				WorkData.VM->GetLiteralMemory().SetRegisterValueFromString(Operand, LocalVariable.CPPType, LocalVariable.CPPTypeObject, {LocalVariable.DefaultValue});
+#else
+				check(Operand.GetMemoryType() == ERigVMMemoryType::Literal);
+				TArray<FRigVMPropertyDescription>& LiteralProperties = WorkData.PropertyDescriptions.FindChecked(Operand.GetMemoryType());
+				LiteralProperties[Operand.GetRegisterIndex()].DefaultValue = LocalVariable.DefaultValue;
+#endif
+			};
+			
+			// To create the default value in the literal memory, we need to find a pin in a variable node (or bounded to a local variable) that
 			// uses this local variable
 			for (URigVMNode* Node : VisitedGraph->GetNodes())
 			{
@@ -351,21 +366,31 @@ bool URigVMCompiler::Compile(URigVMGraph* InGraph, URigVMController* InControlle
 						if (Pin->GetDefaultValue() == LocalVariable.Name.ToString())
 						{
 							URigVMPin* ValuePin = VariableNode->FindPin(URigVMVariableNode::ValueName);
-							FRigVMASTProxy PinProxy = FRigVMASTProxy::MakeFromUObject(ValuePin);
-							FRigVMVarExprAST* TempVarExpr = AST->MakeExpr<FRigVMVarExprAST>(FRigVMExprAST::EType::Literal, PinProxy);
-							FRigVMOperand Operand = FindOrAddRegister(TempVarExpr, WorkData, false);
-
-#if UE_RIGVM_UCLASS_BASED_STORAGE_DISABLED
-							WorkData.VM->GetLiteralMemory().SetRegisterValueFromString(Operand, LocalVariable.CPPType, LocalVariable.CPPTypeObject, {LocalVariable.DefaultValue});
-#else
-							check(Operand.GetMemoryType() == ERigVMMemoryType::Literal);
-							TArray<FRigVMPropertyDescription>& LiteralProperties = WorkData.PropertyDescriptions.FindChecked(Operand.GetMemoryType());
-							LiteralProperties[Operand.GetRegisterIndex()].DefaultValue = LocalVariable.DefaultValue;
-#endif
+							AddDefaultValueOperand(ValuePin);
 							break;
 						}
 					}
 				}
+				else
+				{
+					bool bAdded = false;
+					for (URigVMPin* Pin : Node->GetPins())
+					{
+						if (Pin->IsBoundToLocalVariable())
+						{
+							if (Pin->GetBoundVariableName() == LocalVariable.Name.ToString())
+							{
+								AddDefaultValueOperand(Pin);
+								bAdded = true;
+								break;
+							}
+						}
+					}
+					if (bAdded)
+					{
+						break;
+					}
+				}
 			}
 		}
 	}	
@@ -1910,6 +1935,24 @@ FString URigVMCompiler::GetPinHash(const URigVMPin* InPin, const FRigVMVarExprAS
 			return FString::Printf(TEXT("%sVariable::%s%s"), *Prefix, *VariablePath, *Suffix);
 		}
 
+		if (InPin->IsBoundToLocalVariable())
+		{
+			URigVMPin::FPinOverride PinOverride(InVarExpr->GetProxy(), InVarExpr->GetParser()->GetPinOverrides());
+			FString VariablePath = InPin->GetBoundVariablePath(PinOverride);
+
+			FRigVMASTProxy ParentProxy = InVarExpr->GetProxy();
+			while(ParentProxy.GetCallstack().Num() > 1)
+			{
+				ParentProxy = ParentProxy.GetParent();
+
+				if(URigVMLibraryNode* LibraryNode = ParentProxy.GetSubject<URigVMLibraryNode>())
+				{
+					// Local variables are in the format "LocalVariable::PathToGraph|VariableName"
+					return FString::Printf(TEXT("%sLocalVariable::%s|%s%s"), *Prefix, *LibraryNode->GetNodePath(true), *VariablePath, *Suffix);
+				}
+			}
+		}
+
 		// for IO array pins we'll walk left and use that pin hash instead
 		if(const FRigVMVarExprAST* SourceVarExpr = GetSourceVarExpr(InVarExpr))
 		{
diff --git a/Engine/Source/Developer/RigVMDeveloper/Private/RigVMModel/RigVMPin.cpp b/Engine/Source/Developer/RigVMDeveloper/Private/RigVMModel/RigVMPin.cpp
--- a/Engine/Source/Developer/RigVMDeveloper/Private/RigVMModel/RigVMPin.cpp
+++ b/Engine/Source/Developer/RigVMDeveloper/Private/RigVMModel/RigVMPin.cpp
@@ -956,6 +956,80 @@ bool URigVMPin::IsBoundToVariable(const URigVMPin::FPinOverride& InOverride) con
 	return !GetBoundVariablePath(InOverride).IsEmpty();
 }
 
+bool URigVMPin::IsBoundToExternalVariable() const
+{
+	return IsBoundToExternalVariable(EmptyPinOverride);
+}
+
+bool URigVMPin::IsBoundToExternalVariable(const FPinOverride& InOverride) const
+{
+	FString VariableName = GetBoundVariableName(InOverride);
+	if (VariableName.IsEmpty())
+	{
+		return false;
+	}
+
+	TArray<FRigVMGraphVariableDescription> LocalVariables = GetGraph()->GetLocalVariables(true);
+	for (FRigVMGraphVariableDescription& LocalVariable : LocalVariables)
+	{
+		if (LocalVariable.Name == *VariableName)
+		{
+			return false;
+		}
+	}
+
+	return true;
+}
+
+bool URigVMPin::IsBoundToLocalVariable() const
+{
+	return IsBoundToLocalVariable(EmptyPinOverride);
+}
+
+bool URigVMPin::IsBoundToLocalVariable(const FPinOverride& InOverride) const
+{
+	FString VariableName = GetBoundVariableName(InOverride);
+	if (VariableName.IsEmpty())
+	{
+		return false;
+	}
+
+	TArray<FRigVMGraphVariableDescription> LocalVariables = GetGraph()->GetLocalVariables(false);
+	for (FRigVMGraphVariableDescription& LocalVariable : LocalVariables)
+	{
+		if (LocalVariable.Name == *VariableName)
+		{
+			return true;
+		}
+	}
+
+	return false;
+}
+
+bool URigVMPin::IsBoundToInputArgument() const
+{
+	return IsBoundToInputArgument(EmptyPinOverride);
+}
+
+bool URigVMPin::IsBoundToInputArgument(const FPinOverride& InOverride) const
+{
+	FString VariableName = GetBoundVariableName(InOverride);
+	if (VariableName.IsEmpty())
+	{
+		return false;
+	}
+
+	if (URigVMFunctionEntryNode* EntryNode = GetGraph()->GetEntryNode())
+	{
+		if (EntryNode->FindPin(VariableName))
+		{
+			return true;
+		}
+	}
+	
+	return false;
+}
+
 #if UE_RIGVM_UCLASS_BASED_STORAGE_DISABLED
 bool URigVMPin::CanBeBoundToVariable(const FRigVMExternalVariable& InExternalVariable, const FRigVMRegisterOffset& InOffset) const
 #else
diff --git a/Engine/Source/Developer/RigVMDeveloper/Public/RigVMModel/RigVMPin.h b/Engine/Source/Developer/RigVMDeveloper/Public/RigVMModel/RigVMPin.h
--- a/Engine/Source/Developer/RigVMDeveloper/Public/RigVMModel/RigVMPin.h
+++ b/Engine/Source/Developer/RigVMDeveloper/Public/RigVMModel/RigVMPin.h
@@ -354,6 +354,24 @@ public:
 	// Returns true if this pin is bound to a variable
 	bool IsBoundToVariable(const FPinOverride& InOverride) const;
 
+	// Returns true if this pin is bound to an external variable
+	bool IsBoundToExternalVariable() const;
+
+	// Returns true if this pin is bound to an external variable
+	bool IsBoundToExternalVariable(const FPinOverride& InOverride) const;
+
+	// Returns true if this pin is bound to a local variable
+	bool IsBoundToLocalVariable() const;
+
+	// Returns true if this pin is bound to a local variable
+	bool IsBoundToLocalVariable(const FPinOverride& InOverride) const;
+
+	// Returns true if this pin is bound to an input argument
+	bool IsBoundToInputArgument() const;
+
+	// Returns true if this pin is bound to an input argument
+	bool IsBoundToInputArgument(const FPinOverride& InOverride) const;
+
 #if UE_RIGVM_UCLASS_BASED_STORAGE_DISABLED
 	// Returns true if the pin can be bound to a given variable
 	bool CanBeBoundToVariable(const FRigVMExternalVariable& InExternalVariable, const FRigVMRegisterOffset& InOffset = FRigVMRegisterOffset()) const;
