Control Rig: Fix error on python script when adding link to array node iterator with generic type pin

diff --git a/Engine/Plugins/Animation/ControlRig/Source/ControlRig/Private/ControlRig.cpp b/Engine/Plugins/Animation/ControlRig/Source/ControlRig/Private/ControlRig.cpp
--- a/Engine/Plugins/Animation/ControlRig/Source/ControlRig/Private/ControlRig.cpp
+++ b/Engine/Plugins/Animation/ControlRig/Source/ControlRig/Private/ControlRig.cpp
@@ -295,7 +295,7 @@ TArray<FRigVMExternalVariable> UControlRig::GetExternalVariablesImpl(bool bFallb
 			TArray<FRigVMExternalVariable> BlueprintVariables;
 			for (const FBPVariableDescription& VariableDescription : Blueprint->NewVariables)
 			{
-				FRigVMExternalVariable ExternalVariable = GetExternalVariableFromDescription(VariableDescription);
+				FRigVMExternalVariable ExternalVariable = RigVMTypeUtils::ExternalVariableFromBPVariableDescription(VariableDescription);
 				if (ExternalVariable.TypeName.IsNone())
 				{
 					continue;
@@ -2524,143 +2524,6 @@ void UControlRig::HandleInteractionRigControlSelected(UControlRig* Subject, FRig
 
 #if WITH_EDITOR
 
-FEdGraphPinType UControlRig::GetPinTypeFromExternalVariable(const FRigVMExternalVariable& InExternalVariable)
-{
-	FEdGraphPinType PinType;
-	PinType.ResetToDefaults();
-	PinType.PinCategory = NAME_None;
-
-	if (InExternalVariable.TypeName == TEXT("bool"))
-	{
-		PinType.PinCategory = UEdGraphSchema_K2::PC_Boolean;
-	}
-	else if (InExternalVariable.TypeName == TEXT("int32"))
-	{
-		PinType.PinCategory = UEdGraphSchema_K2::PC_Int;
-	}
-	else if (InExternalVariable.TypeName == TEXT("float"))
-	{
-		PinType.PinCategory = UEdGraphSchema_K2::PC_Float;
-	}
-	else if (InExternalVariable.TypeName == TEXT("double"))
-	{
-		PinType.PinCategory = UEdGraphSchema_K2::PC_Double;
-	}
-	else if (InExternalVariable.TypeName == TEXT("FName"))
-	{
-		PinType.PinCategory = UEdGraphSchema_K2::PC_Name;
-	}
-	else if (InExternalVariable.TypeName == TEXT("FString"))
-	{
-		PinType.PinCategory = UEdGraphSchema_K2::PC_String;
-	}
-	else if (Cast<UScriptStruct>(InExternalVariable.TypeObject))
-	{
-		PinType.PinCategory = UEdGraphSchema_K2::PC_Struct;
-		PinType.PinSubCategoryObject = InExternalVariable.TypeObject;
-	}
-	else if (Cast<UEnum>(InExternalVariable.TypeObject))
-	{
-		PinType.PinCategory = UEdGraphSchema_K2::PC_Byte;
-		PinType.PinSubCategoryObject = InExternalVariable.TypeObject;
-	}
-
-	if (InExternalVariable.bIsArray)
-	{
-		PinType.ContainerType = EPinContainerType::Array;
-	}
-	else
-	{
-		PinType.ContainerType = EPinContainerType::None;
-	}
-
-	return PinType;
-}
-
-FRigVMExternalVariable UControlRig::GetExternalVariableFromPinType(const FName& InName, const FEdGraphPinType& InPinType, bool bInPublic, bool bInReadonly)
-{
-	FRigVMExternalVariable ExternalVariable;
-	ExternalVariable.Name = InName;
-	ExternalVariable.bIsPublic = bInPublic;
-	ExternalVariable.bIsReadOnly = bInReadonly;
-
-	if (InPinType.ContainerType == EPinContainerType::None)
-	{
-		ExternalVariable.bIsArray = false;
-	}
-	else if (InPinType.ContainerType == EPinContainerType::Array)
-	{
-		ExternalVariable.bIsArray = true;
-	}
-	else
-	{
-		return FRigVMExternalVariable();
-	}
-
-	if (InPinType.PinCategory == UEdGraphSchema_K2::PC_Boolean)
-	{
-		ExternalVariable.TypeName = TEXT("bool");
-		ExternalVariable.Size = sizeof(bool);
-	}
-	else if (InPinType.PinCategory == UEdGraphSchema_K2::PC_Int)
-	{
-		ExternalVariable.TypeName = TEXT("int32");
-		ExternalVariable.Size = sizeof(int32);
-	}
-	else if (InPinType.PinCategory == UEdGraphSchema_K2::PC_Enum ||
-		InPinType.PinCategory == UEdGraphSchema_K2::PC_Byte)
-	{
-		if (UEnum* Enum = Cast<UEnum>(InPinType.PinSubCategoryObject))
-		{
-			ExternalVariable.TypeName = Enum->GetFName();
-			ExternalVariable.TypeObject = Enum;
-		}
-		else
-		{
-			ExternalVariable.TypeName = TEXT("uint8");
-		}
-		ExternalVariable.Size = sizeof(uint8);
-	}
-	else if (InPinType.PinCategory == UEdGraphSchema_K2::PC_Float)
-	{
-		ExternalVariable.TypeName = TEXT("float");
-		ExternalVariable.Size = sizeof(float);
-	}
-	else if (InPinType.PinCategory == UEdGraphSchema_K2::PC_Double)
-	{
-		ExternalVariable.TypeName = TEXT("double");
-		ExternalVariable.Size = sizeof(double);
-	}
-	else if (InPinType.PinCategory == UEdGraphSchema_K2::PC_Name)
-	{
-		ExternalVariable.TypeName = TEXT("FName");
-		ExternalVariable.Size = sizeof(FName);
-	}
-	else if (InPinType.PinCategory == UEdGraphSchema_K2::PC_String)
-	{
-		ExternalVariable.TypeName = TEXT("FString");
-		ExternalVariable.Size = sizeof(FString);
-	}
-	else if (InPinType.PinCategory == UEdGraphSchema_K2::PC_Struct)
-	{
-		if (UScriptStruct* Struct = Cast<UScriptStruct>(InPinType.PinSubCategoryObject))
-		{
-			ExternalVariable.TypeName = *Struct->GetStructCPPName();
-			ExternalVariable.TypeObject = Struct;
-			ExternalVariable.Size = Struct->GetStructureSize();
-		}
-	}
-
-	return ExternalVariable;
-}
-
-FRigVMExternalVariable UControlRig::GetExternalVariableFromDescription(const FBPVariableDescription& InVariableDescription)
-{
-	const bool bIsPublic = !((InVariableDescription.PropertyFlags & CPF_DisableEditOnInstance) == CPF_DisableEditOnInstance);
-	const bool bIsReadOnly = ((InVariableDescription.PropertyFlags & CPF_BlueprintReadOnly) == CPF_BlueprintReadOnly);
-	return GetExternalVariableFromPinType(InVariableDescription.VarName, InVariableDescription.VarType, bIsPublic, bIsReadOnly);
-}
-
 URigVM* UControlRig::GetSnapshotVM(bool bCreateIfNeeded)
 {
 #if WITH_EDITOR
diff --git a/Engine/Plugins/Animation/ControlRig/Source/ControlRig/Public/ControlRig.h b/Engine/Plugins/Animation/ControlRig/Source/ControlRig/Public/ControlRig.h
--- a/Engine/Plugins/Animation/ControlRig/Source/ControlRig/Public/ControlRig.h
+++ b/Engine/Plugins/Animation/ControlRig/Source/ControlRig/Public/ControlRig.h
@@ -23,6 +23,7 @@
 
 #if WITH_EDITOR
 #include "RigVMModel/RigVMPin.h"
+#include "RigVMTypeUtils.h"
 #endif
 
 #if WITH_EDITOR
@@ -629,15 +630,6 @@ private:
 	void HandleInteractionRigExecuted(UControlRig* Subject, EControlRigState State, const FName& EventName);
 	void HandleInteractionRigControlSelected(UControlRig* Subject, FRigControlElement* InControl, bool bSelected, bool bInverted);
 
-#if WITH_EDITOR
-
-public:
-
-	static FEdGraphPinType GetPinTypeFromExternalVariable(const FRigVMExternalVariable& InExternalVariable);
-	static FRigVMExternalVariable GetExternalVariableFromPinType(const FName& InName, const FEdGraphPinType& InPinType, bool bInPublic = false, bool bInReadonly = false);
-	static FRigVMExternalVariable GetExternalVariableFromDescription(const FBPVariableDescription& InVariableDescription);
-
-#endif
 
 protected:
 	bool bRequiresInitExecution;
diff --git a/Engine/Plugins/Animation/ControlRig/Source/ControlRigDeveloper/Private/AnimGraphNode_ControlRig.cpp b/Engine/Plugins/Animation/ControlRig/Source/ControlRigDeveloper/Private/AnimGraphNode_ControlRig.cpp
--- a/Engine/Plugins/Animation/ControlRig/Source/ControlRigDeveloper/Private/AnimGraphNode_ControlRig.cpp
+++ b/Engine/Plugins/Animation/ControlRig/Source/ControlRigDeveloper/Private/AnimGraphNode_ControlRig.cpp
@@ -130,7 +130,7 @@ void UAnimGraphNode_ControlRig::CreateCustomPins(TArray<UEdGraphPin*>* OldPins)
 		{
 			FRigVMExternalVariable& Variable = VariablePair.Value;
 
-			FEdGraphPinType PinType = UControlRig::GetPinTypeFromExternalVariable(Variable);
+			FEdGraphPinType PinType = RigVMTypeUtils::PinTypeFromExternalVariable(Variable);
 			if (!PinType.PinCategory.IsValid())
 			{
 				continue;
diff --git a/Engine/Plugins/Animation/ControlRig/Source/ControlRigDeveloper/Private/ControlRigBlueprint.cpp b/Engine/Plugins/Animation/ControlRig/Source/ControlRigDeveloper/Private/ControlRigBlueprint.cpp
--- a/Engine/Plugins/Animation/ControlRig/Source/ControlRigDeveloper/Private/ControlRigBlueprint.cpp
+++ b/Engine/Plugins/Animation/ControlRig/Source/ControlRigDeveloper/Private/ControlRigBlueprint.cpp
@@ -26,6 +26,7 @@
 #include "Misc/CoreDelegates.h"
 #include "AssetRegistryModule.h"
 #include "RigVMPythonUtils.h"
+#include "RigVMTypeUtils.h"
 
 #if WITH_EDITOR
 #include "IControlRigEditorModule.h"
@@ -55,7 +56,7 @@ FEdGraphPinType FControlRigPublicFunctionArg::GetPinType() const
 		Variable.TypeObject = URigVMPin::FindObjectFromCPPTypeObjectPath(CPPTypeObjectPath.ToString());
 	}
 
-	return UControlRig::GetPinTypeFromExternalVariable(Variable);
+	return RigVMTypeUtils::PinTypeFromExternalVariable(Variable);
 }
 
 bool FControlRigPublicFunctionData::IsMutable() const
@@ -1479,12 +1480,25 @@ TArray<FString> UControlRigBlueprint::GeneratePythonCommands(const FString InNew
 	// Add variables
 	for (const FBPVariableDescription& Variable : NewVariables)
 	{
-		const FRigVMExternalVariable ExternalVariable = UControlRig::GetExternalVariableFromDescription(Variable);
-
+		const FRigVMExternalVariable ExternalVariable = RigVMTypeUtils::ExternalVariableFromBPVariableDescription(Variable);
+		FString CPPType;
+		UObject* CPPTypeObject = nullptr;
+		RigVMTypeUtils::CPPTypeFromExternalVariable(ExternalVariable, CPPType, &CPPTypeObject);
+		if (CPPTypeObject)
+		{
+			if (ExternalVariable.bIsArray)
+			{
+				CPPType = RigVMTypeUtils::ArrayTypeFromBaseType(CPPTypeObject->GetPathName());
+			}
+			else
+			{
+				CPPType = CPPTypeObject->GetPathName();
+			}
+		}
 		// FName AddMemberVariable(const FName& InName, const FString& InCPPType, bool bIsPublic = false, bool bIsReadOnly = false, FString InDefaultValue = TEXT(""));
 		Commands.Add(FString::Printf(TEXT("blueprint.add_member_variable('%s', '%s', %s, %s)"),
 					*ExternalVariable.Name.ToString(),
-					ExternalVariable.TypeObject ? *ExternalVariable.TypeObject->GetPathName() : *ExternalVariable.TypeName.ToString(),
+					*CPPType,
 					ExternalVariable.bIsPublic ? TEXT("True") : TEXT("False"),
 					ExternalVariable.bIsReadOnly ? TEXT("True") : TEXT("False")));	
 	}
@@ -2052,67 +2066,7 @@ TArray<UStruct*> UControlRigBlueprint::GetAvailableRigUnits()
 
 FName UControlRigBlueprint::AddMemberVariable(const FName& InName, const FString& InCPPType, bool bIsPublic, bool bIsReadOnly, FString InDefaultValue)
 {
-	FRigVMExternalVariable Variable;
-	Variable.Name = InName;
-	Variable.bIsPublic = bIsPublic;
-	Variable.bIsReadOnly = bIsReadOnly;
-
-	FString CPPType = InCPPType;
-	if (CPPType.StartsWith(TEXT("TMap<")))
-	{
-		UE_LOG(LogControlRigDeveloper, Warning, TEXT("TMap Variables are not supported."));
-		return NAME_None;
-	}
-
-	Variable.bIsArray = RigVMUtilities::IsArrayType(CPPType);
-	if (Variable.bIsArray)
-	{
-		CPPType = CPPType.RightChop(7).LeftChop(1);
-	}
-
-	if (CPPType == TEXT("bool"))
-	{
-		Variable.TypeName = *CPPType;
-		Variable.Size = sizeof(bool);
-	}
-	else if (CPPType == TEXT("float"))
-	{
-		Variable.TypeName = *CPPType;
-		Variable.Size = sizeof(float);
-	}
-	else if (CPPType == TEXT("double"))
-	{
-		Variable.TypeName = *CPPType;
-		Variable.Size = sizeof(double);
-	}
-	else if (CPPType == TEXT("int32"))
-	{
-		Variable.TypeName = *CPPType;
-		Variable.Size = sizeof(int32);
-	}
-	else if (CPPType == TEXT("FString"))
-	{
-		Variable.TypeName = *CPPType;
-		Variable.Size = sizeof(FString);
-	}
-	else if (CPPType == TEXT("FName"))
-	{
-		Variable.TypeName = *CPPType;
-		Variable.Size = sizeof(FName);
-	}
-	else if(UScriptStruct* ScriptStruct = URigVMPin::FindObjectFromCPPTypeObjectPath<UScriptStruct>(CPPType))
-	{
-		Variable.TypeName = *ScriptStruct->GetStructCPPName();
-		Variable.TypeObject = ScriptStruct;
-		Variable.Size = ScriptStruct->GetStructureSize();
-	}
-	else if (UEnum* Enum= URigVMPin::FindObjectFromCPPTypeObjectPath<UEnum>(CPPType))
-	{
-		Variable.TypeName = *Enum->CppType;
-		Variable.TypeObject = Enum;
-		Variable.Size = Enum->GetResourceSizeBytes(EResourceSizeMode::EstimatedTotal);
-	}
-
+	FRigVMExternalVariable Variable = RigVMTypeUtils::ExternalVariableFromCPPType(InName, InCPPType, bIsPublic, bIsReadOnly);
 	FName Result = AddCRMemberVariableFromExternal(Variable, InDefaultValue);
 	if (!Result.IsNone())
 	{
@@ -2173,7 +2127,7 @@ bool UControlRigBlueprint::ChangeMemberVariableType(const FName& InName, const F
 		return false;
 	}
 
-	Variable.bIsArray = RigVMUtilities::IsArrayType(CPPType);
+	Variable.bIsArray = RigVMTypeUtils::IsArrayType(CPPType);
 	if (Variable.bIsArray)
 	{
 		CPPType = CPPType.RightChop(7).LeftChop(1);
@@ -2222,7 +2176,7 @@ bool UControlRigBlueprint::ChangeMemberVariableType(const FName& InName, const F
 		Variable.Size = Enum->GetResourceSizeBytes(EResourceSizeMode::EstimatedTotal);
 	}
 
-	FEdGraphPinType PinType = UControlRig::GetPinTypeFromExternalVariable(Variable);
+	FEdGraphPinType PinType = RigVMTypeUtils::PinTypeFromExternalVariable(Variable);
 	if (!PinType.PinCategory.IsValid())
 	{
 		return false;
@@ -3278,7 +3232,7 @@ void UControlRigBlueprint::CreateMemberVariablesOnLoad()
 					continue;
 				}
 
-				FEdGraphPinType PinType = UControlRig::GetPinTypeFromExternalVariable(Description.ToExternalVariable());
+				FEdGraphPinType PinType = RigVMTypeUtils::PinTypeFromExternalVariable(Description.ToExternalVariable());
 				if (!PinType.PinCategory.IsValid())
 				{
 					continue;
@@ -3309,7 +3263,7 @@ void UControlRigBlueprint::CreateMemberVariablesOnLoad()
 					continue;
 				}
 
-				FEdGraphPinType PinType = UControlRig::GetPinTypeFromExternalVariable(Description.ToExternalVariable());
+				FEdGraphPinType PinType = RigVMTypeUtils::PinTypeFromExternalVariable(Description.ToExternalVariable());
 				if (!PinType.PinCategory.IsValid())
 				{
 					continue;
@@ -3402,7 +3356,7 @@ int32 UControlRigBlueprint::AddCRMemberVariable(UControlRigBlueprint* InBlueprin
 
 FName UControlRigBlueprint::AddCRMemberVariableFromExternal(FRigVMExternalVariable InVariableToCreate, FString InDefaultValue)
 {
-	FEdGraphPinType PinType = UControlRig::GetPinTypeFromExternalVariable(InVariableToCreate);
+	FEdGraphPinType PinType = RigVMTypeUtils::PinTypeFromExternalVariable(InVariableToCreate);
 	if (!PinType.PinCategory.IsValid())
 	{
 		return NAME_None;
@@ -3875,12 +3829,25 @@ void UControlRigBlueprint::OnVariableAdded(const FName& InVarName)
 		}
 	}
 
-	FRigVMExternalVariable ExternalVariable = UControlRig::GetExternalVariableFromDescription(Variable);
-
-	RigVMPythonUtils::Print(GetFName().ToString(),
+	const FRigVMExternalVariable ExternalVariable = RigVMTypeUtils::ExternalVariableFromBPVariableDescription(Variable);
+    FString CPPType;
+    UObject* CPPTypeObject = nullptr;
+    RigVMTypeUtils::CPPTypeFromExternalVariable(ExternalVariable, CPPType, &CPPTypeObject);
+	if (CPPTypeObject)
+	{
+		if (ExternalVariable.bIsArray)
+		{
+			CPPType = RigVMTypeUtils::ArrayTypeFromBaseType(CPPTypeObject->GetPathName());
+		}
+		else
+		{
+			CPPType = CPPTypeObject->GetPathName();
+		}
+	}
+    RigVMPythonUtils::Print(GetFName().ToString(),
 		FString::Printf(TEXT("blueprint.add_member_variable('%s', '%s', %s, %s, '%s')"),
 			*InVarName.ToString(),
-			*ExternalVariable.TypeName.ToString(),
+			*CPPType,
 			(ExternalVariable.bIsPublic) ? TEXT("False") : TEXT("True"), 
 			(ExternalVariable.bIsReadOnly) ? TEXT("True") : TEXT("False"), 
 			*Variable.DefaultValue)); 
@@ -3938,8 +3905,8 @@ void UControlRigBlueprint::OnVariableRenamed(const FName& InOldVarName, const FN
 void UControlRigBlueprint::OnVariableTypeChanged(const FName& InVarName, FEdGraphPinType InOldPinType, FEdGraphPinType InNewPinType)
 {
 	FString CPPType;
-	UObject* CPPTypeObject;
-	FRigVMGraphVariableDescription::CPPTypeFromPinType(InNewPinType, CPPType, CPPTypeObject);
+	UObject* CPPTypeObject = nullptr;
+	RigVMTypeUtils::CPPTypeFromPinType(InNewPinType, CPPType, &CPPTypeObject);
 	
 	TArray<URigVMGraph*> AllGraphs = GetAllModels();
 	for (URigVMGraph* Graph : AllGraphs)
diff --git a/Engine/Plugins/Animation/ControlRig/Source/ControlRigDeveloper/Private/Graph/ControlRigGraphSchema.cpp b/Engine/Plugins/Animation/ControlRig/Source/ControlRigDeveloper/Private/Graph/ControlRigGraphSchema.cpp
--- a/Engine/Plugins/Animation/ControlRig/Source/ControlRigDeveloper/Private/Graph/ControlRigGraphSchema.cpp
+++ b/Engine/Plugins/Animation/ControlRig/Source/ControlRigDeveloper/Private/Graph/ControlRigGraphSchema.cpp
@@ -120,7 +120,7 @@ void FControlRigGraphSchemaAction_LocalVar::ChangeVariableType(const FEdGraphPin
 	{
 		FString NewCPPType;
 		UObject* NewCPPTypeObject = nullptr;
-		FRigVMGraphVariableDescription::CPPTypeFromPinType(NewPinType, NewCPPType, NewCPPTypeObject);
+		RigVMTypeUtils::CPPTypeFromPinType(NewPinType, NewCPPType, &NewCPPTypeObject);
 		Graph->GetController()->SetLocalVariableType(GetVariableName(), NewCPPType, NewCPPTypeObject, true, true);
 	}
 }
@@ -1292,7 +1292,7 @@ UEdGraphPin* UControlRigGraphSchema::DropPinOnNode(UEdGraphNode* InTargetNode, c
 				{
 					ensure(!Model->IsTopLevelGraph());
 
-					FRigVMExternalVariable ExternalVar = UControlRig::GetExternalVariableFromPinType(InSourcePinName, InSourcePinType);
+					FRigVMExternalVariable ExternalVar = RigVMTypeUtils::ExternalVariableFromPinType(InSourcePinName, InSourcePinType);
 					if (ExternalVar.IsValid(true /* allow null memory */))
 					{
 						if (URigVMController* Controller = RigBlueprint->GetController(Model))
diff --git a/Engine/Plugins/Animation/ControlRig/Source/ControlRigEditor/Private/ControlRigGraphDetails.cpp b/Engine/Plugins/Animation/ControlRig/Source/ControlRigEditor/Private/ControlRigGraphDetails.cpp
--- a/Engine/Plugins/Animation/ControlRig/Source/ControlRigEditor/Private/ControlRigGraphDetails.cpp
+++ b/Engine/Plugins/Animation/ControlRig/Source/ControlRigEditor/Private/ControlRigGraphDetails.cpp
@@ -402,7 +402,7 @@ void FControlRigArgumentLayout::PinInfoChanged(const FEdGraphPinType& PinType)
 		{
 			if (URigVMController* Controller = Blueprint->GetController(LibraryNode->GetContainedGraph()))
 			{
-				FRigVMExternalVariable ExternalVariable = UControlRig::GetExternalVariableFromPinType(Pin->GetFName(), PinType, true, false);
+				FRigVMExternalVariable ExternalVariable = RigVMTypeUtils::ExternalVariableFromPinType(Pin->GetFName(), PinType, true, false);
 				if (!ExternalVariable.IsValid(true /* allow nullptr memory */))
 				{
 					return;
diff --git a/Engine/Plugins/Animation/ControlRig/Source/ControlRigEditor/Private/Graph/NodeSpawners/ControlRigVariableNodeSpawner.cpp b/Engine/Plugins/Animation/ControlRig/Source/ControlRigEditor/Private/Graph/NodeSpawners/ControlRigVariableNodeSpawner.cpp
--- a/Engine/Plugins/Animation/ControlRig/Source/ControlRigEditor/Private/Graph/NodeSpawners/ControlRigVariableNodeSpawner.cpp
+++ b/Engine/Plugins/Animation/ControlRig/Source/ControlRigEditor/Private/Graph/NodeSpawners/ControlRigVariableNodeSpawner.cpp
@@ -37,7 +37,7 @@ UControlRigVariableNodeSpawner* UControlRigVariableNodeSpawner::CreateFromExtern
 	MenuSignature.Category = InCategory;
 	MenuSignature.Keywords = FText::FromString(TEXT("Variable"));
 
-	FEdGraphPinType PinType = UControlRig::GetPinTypeFromExternalVariable(InExternalVariable);
+	FEdGraphPinType PinType = RigVMTypeUtils::PinTypeFromExternalVariable(InExternalVariable);
 	MenuSignature.Icon = UK2Node_Variable::GetVarIconFromPinType(PinType, MenuSignature.IconTint);
 
 	return NodeSpawner;
diff --git a/Engine/Plugins/Animation/ControlRig/Source/ControlRigEditor/Private/Graph/SControlRigGraphPinVariableBinding.cpp b/Engine/Plugins/Animation/ControlRig/Source/ControlRigEditor/Private/Graph/SControlRigGraphPinVariableBinding.cpp
--- a/Engine/Plugins/Animation/ControlRig/Source/ControlRigEditor/Private/Graph/SControlRigGraphPinVariableBinding.cpp
+++ b/Engine/Plugins/Animation/ControlRig/Source/ControlRigEditor/Private/Graph/SControlRigGraphPinVariableBinding.cpp
@@ -128,7 +128,7 @@ FLinearColor SControlRigVariableBinding::GetBindingColor() const
 
 				if (ExternalVariable.Name == BoundVariable)
 				{
-					const FEdGraphPinType PinType = UControlRig::GetPinTypeFromExternalVariable(ExternalVariable);
+					const FEdGraphPinType PinType = RigVMTypeUtils::PinTypeFromExternalVariable(ExternalVariable);
 					return Schema->GetPinTypeColor(PinType);
 				}
 			}
@@ -289,7 +289,7 @@ void SControlRigVariableBinding::FillLocalVariableMenu(FMenuBuilder& MenuBuilder
 				continue;
 			}
 			
-			const FEdGraphPinType PinType = UControlRig::GetPinTypeFromExternalVariable(ExternalVariable);
+			const FEdGraphPinType PinType = RigVMTypeUtils::PinTypeFromExternalVariable(ExternalVariable);
 
 			MenuBuilder.AddMenuEntry(
 				FUIAction(FExecuteAction::CreateSP(this, &SControlRigVariableBinding::HandleBindToLocalVariable, LocalVariable)),
diff --git a/Engine/Source/Developer/RigVMDeveloper/Private/RigVMCompiler/RigVMCompiler.cpp b/Engine/Source/Developer/RigVMDeveloper/Private/RigVMCompiler/RigVMCompiler.cpp
--- a/Engine/Source/Developer/RigVMDeveloper/Private/RigVMCompiler/RigVMCompiler.cpp
+++ b/Engine/Source/Developer/RigVMDeveloper/Private/RigVMCompiler/RigVMCompiler.cpp
@@ -6,6 +6,7 @@
 #include "RigVMDeveloperModule.h"
 #include "UObject/PropertyPortFlags.h"
 #include "Stats/StatsHierarchical.h"
+#include "RigVMTypeUtils.h"
 
 class FRigVMCompilerImportErrorContext : public FOutputDevice
 {
diff --git a/Engine/Source/Developer/RigVMDeveloper/Private/RigVMModel/Nodes/RigVMArrayNode.cpp b/Engine/Source/Developer/RigVMDeveloper/Private/RigVMModel/Nodes/RigVMArrayNode.cpp
--- a/Engine/Source/Developer/RigVMDeveloper/Private/RigVMModel/Nodes/RigVMArrayNode.cpp
+++ b/Engine/Source/Developer/RigVMDeveloper/Private/RigVMModel/Nodes/RigVMArrayNode.cpp
@@ -398,7 +398,12 @@ FString URigVMArrayNode::GetCPPType() const
 	{
 		return FString();
 	}
-	return ArrayPin->GetCPPType();
+	const FString& CPPType = ArrayPin->GetCPPType();
+	if (RigVMTypeUtils::IsArrayType(CPPType))
+	{
+		return CPPType.Mid(7, CPPType.Len() - 8);
+	}
+	return CPPType;
 }
 
 UObject* URigVMArrayNode::GetCPPTypeObject() const
diff --git a/Engine/Source/Developer/RigVMDeveloper/Private/RigVMModel/RigVMController.cpp b/Engine/Source/Developer/RigVMDeveloper/Private/RigVMModel/RigVMController.cpp
--- a/Engine/Source/Developer/RigVMDeveloper/Private/RigVMModel/RigVMController.cpp
+++ b/Engine/Source/Developer/RigVMDeveloper/Private/RigVMModel/RigVMController.cpp
@@ -16,6 +16,7 @@
 #include "Misc/CoreMisc.h"
 #include "Algo/Sort.h"
 #include "RigVMPythonUtils.h"
+#include "RigVMTypeUtils.h"
 
 #if WITH_EDITOR
 #include "Exporters/Exporter.h"
@@ -766,6 +767,16 @@ TArray<FString> URigVMController::GetAddNodePythonCommands(URigVMNode* Node) con
 							*PinPath,
 							*Pin->GetDefaultValue()));
 			}
+
+			if (!Pin->GetBoundVariablePath().IsEmpty())
+			{
+				const FString PinPath = GetSanitizedPinPath(Pin->GetPinPath());
+
+				Commands.Add(FString::Printf(TEXT("blueprint.get_controller_by_name('%s').bind_pin_to_variable('%s', '%s')"),
+							*GraphName,
+							*PinPath,
+							*Pin->GetBoundVariablePath()));
+			}
 		}
 	}
 
diff --git a/Engine/Source/Developer/RigVMDeveloper/Private/RigVMModel/RigVMGraph.cpp b/Engine/Source/Developer/RigVMDeveloper/Private/RigVMModel/RigVMGraph.cpp
--- a/Engine/Source/Developer/RigVMDeveloper/Private/RigVMModel/RigVMGraph.cpp
+++ b/Engine/Source/Developer/RigVMDeveloper/Private/RigVMModel/RigVMGraph.cpp
@@ -7,6 +7,7 @@
 #include "RigVMModel/Nodes/RigVMFunctionEntryNode.h"
 #include "RigVMModel/Nodes/RigVMFunctionReturnNode.h"
 #include "UObject/Package.h"
+#include "RigVMTypeUtils.h"
 
 URigVMGraph::URigVMGraph()
 : DiagnosticsAST(nullptr)
diff --git a/Engine/Source/Developer/RigVMDeveloper/Private/RigVMModel/RigVMPin.cpp b/Engine/Source/Developer/RigVMDeveloper/Private/RigVMModel/RigVMPin.cpp
--- a/Engine/Source/Developer/RigVMDeveloper/Private/RigVMModel/RigVMPin.cpp
+++ b/Engine/Source/Developer/RigVMDeveloper/Private/RigVMModel/RigVMPin.cpp
@@ -1149,7 +1149,7 @@ FRigVMExternalVariable URigVMPin::ToExternalVariable() const
 	}
 	ExternalVariable.Name = *VariableName;
 
-	if (RigVMUtilities::IsArrayType(CPPType))
+	if (RigVMTypeUtils::IsArrayType(CPPType))
 	{
 		ExternalVariable.bIsArray = true;
 		ExternalVariable.TypeName = *CPPType.Mid(7, CPPType.Len() - 8);
diff --git a/Engine/Source/Developer/RigVMDeveloper/Private/RigVMTypeUtils.cpp b/Engine/Source/Developer/RigVMDeveloper/Private/RigVMTypeUtils.cpp
new file mode 100644
--- /dev/null
+++ b/Engine/Source/Developer/RigVMDeveloper/Private/RigVMTypeUtils.cpp
@@ -0,0 +1,429 @@
+﻿// Copyright Epic Games, Inc. All Rights Reserved.
+
+#include "RigVMTypeUtils.h"
+#include "RigVMModel/RigVMVariableDescription.h"
+
+#if WITH_EDITOR
+FRigVMExternalVariable RigVMTypeUtils::ExternalVariableFromBPVariableDescription(
+	const FBPVariableDescription& InVariableDescription)
+{
+	const bool bIsPublic = !((InVariableDescription.PropertyFlags & CPF_DisableEditOnInstance) == CPF_DisableEditOnInstance);
+	const bool bIsReadOnly = ((InVariableDescription.PropertyFlags & CPF_BlueprintReadOnly) == CPF_BlueprintReadOnly);
+	return ExternalVariableFromPinType(InVariableDescription.VarName, InVariableDescription.VarType, bIsPublic, bIsReadOnly);
+}
+
+FRigVMExternalVariable RigVMTypeUtils::ExternalVariableFromRigVMVariableDescription(const FRigVMGraphVariableDescription& InVariableDescription)
+{	
+	FRigVMExternalVariable ExternalVariable;
+	ExternalVariable.Name = InVariableDescription.Name;
+
+	if (InVariableDescription.CPPType.StartsWith(TEXT("TArray<")))
+	{
+		ExternalVariable.bIsArray = true;
+		ExternalVariable.TypeName = *InVariableDescription.CPPType.Mid(7, InVariableDescription.CPPType.Len() - 8);
+		ExternalVariable.TypeObject = InVariableDescription.CPPTypeObject;
+	}
+	else
+	{
+		ExternalVariable.bIsArray = false;
+		ExternalVariable.TypeName = *InVariableDescription.CPPType;
+		ExternalVariable.TypeObject = InVariableDescription.CPPTypeObject;
+	}
+
+	ExternalVariable.bIsPublic = false;
+	ExternalVariable.bIsReadOnly = false;
+	ExternalVariable.Memory = nullptr;
+	return ExternalVariable;
+
+}
+
+FRigVMExternalVariable RigVMTypeUtils::ExternalVariableFromPinType(const FName& InName, const FEdGraphPinType& InPinType, bool bInPublic, bool bInReadonly)
+{
+	FRigVMExternalVariable ExternalVariable;
+	ExternalVariable.Name = InName;
+	ExternalVariable.bIsPublic = bInPublic;
+	ExternalVariable.bIsReadOnly = bInReadonly;
+
+	if (InPinType.ContainerType == EPinContainerType::None)
+	{
+		ExternalVariable.bIsArray = false;
+	}
+	else if (InPinType.ContainerType == EPinContainerType::Array)
+	{
+		ExternalVariable.bIsArray = true;
+	}
+	else
+	{
+		return FRigVMExternalVariable();
+	}
+
+	if (InPinType.PinCategory == UEdGraphSchema_K2::PC_Boolean)
+	{
+		ExternalVariable.TypeName = TEXT("bool");
+		ExternalVariable.Size = sizeof(bool);
+	}
+	else if (InPinType.PinCategory == UEdGraphSchema_K2::PC_Int)
+	{
+		ExternalVariable.TypeName = TEXT("int32");
+		ExternalVariable.Size = sizeof(int32);
+	}
+	else if (InPinType.PinCategory == UEdGraphSchema_K2::PC_Enum ||
+		InPinType.PinCategory == UEdGraphSchema_K2::PC_Byte)
+	{
+		if (UEnum* Enum = Cast<UEnum>(InPinType.PinSubCategoryObject))
+		{
+			ExternalVariable.TypeName = Enum->GetFName();
+			ExternalVariable.TypeObject = Enum;
+		}
+		else
+		{
+			ExternalVariable.TypeName = TEXT("uint8");
+		}
+		ExternalVariable.Size = sizeof(uint8);
+	}
+	else if (InPinType.PinCategory == UEdGraphSchema_K2::PC_Float)
+	{
+		ExternalVariable.TypeName = TEXT("float");
+		ExternalVariable.Size = sizeof(float);
+	}
+	else if (InPinType.PinCategory == UEdGraphSchema_K2::PC_Double)
+	{
+		ExternalVariable.TypeName = TEXT("double");
+		ExternalVariable.Size = sizeof(double);
+	}
+	else if (InPinType.PinCategory == UEdGraphSchema_K2::PC_Name)
+	{
+		ExternalVariable.TypeName = TEXT("FName");
+		ExternalVariable.Size = sizeof(FName);
+	}
+	else if (InPinType.PinCategory == UEdGraphSchema_K2::PC_String)
+	{
+		ExternalVariable.TypeName = TEXT("FString");
+		ExternalVariable.Size = sizeof(FString);
+	}
+	else if (InPinType.PinCategory == UEdGraphSchema_K2::PC_Struct)
+	{
+		if (UScriptStruct* Struct = Cast<UScriptStruct>(InPinType.PinSubCategoryObject))
+		{
+			ExternalVariable.TypeName = *Struct->GetStructCPPName();
+			ExternalVariable.TypeObject = Struct;
+			ExternalVariable.Size = Struct->GetStructureSize();
+		}
+	}
+
+	return ExternalVariable;
+}
+
+FRigVMExternalVariable RigVMTypeUtils::ExternalVariableFromCPPType(const FName& InName, const FString& InCPPType,
+	bool bInPublic, bool bInReadonly)
+{
+	FRigVMExternalVariable Variable;
+	if (InCPPType.StartsWith(TEXT("TMap<")))
+	{
+		return Variable;
+	}
+	
+	Variable.Name = InName;
+	Variable.bIsPublic = bInPublic;
+	Variable.bIsReadOnly = bInReadonly;
+
+	FString CPPType = InCPPType;
+	Variable.bIsArray = RigVMTypeUtils::IsArrayType(CPPType);
+	if (Variable.bIsArray)
+	{
+		CPPType = CPPType.RightChop(7).LeftChop(1);
+	}
+
+	if (CPPType == TEXT("bool"))
+	{
+		Variable.TypeName = *CPPType;
+		Variable.Size = sizeof(bool);
+	}
+	else if (CPPType == TEXT("float"))
+	{
+		Variable.TypeName = *CPPType;
+		Variable.Size = sizeof(float);
+	}
+	else if (CPPType == TEXT("double"))
+	{
+		Variable.TypeName = *CPPType;
+		Variable.Size = sizeof(double);
+	}
+	else if (CPPType == TEXT("int32"))
+	{
+		Variable.TypeName = *CPPType;
+		Variable.Size = sizeof(int32);
+	}
+	else if (CPPType == TEXT("FString"))
+	{
+		Variable.TypeName = *CPPType;
+		Variable.Size = sizeof(FString);
+	}
+	else if (CPPType == TEXT("FName"))
+	{
+		Variable.TypeName = *CPPType;
+		Variable.Size = sizeof(FName);
+	}
+	else if(UScriptStruct* ScriptStruct = URigVMPin::FindObjectFromCPPTypeObjectPath<UScriptStruct>(CPPType))
+	{
+		Variable.TypeName = *ScriptStruct->GetStructCPPName();
+		Variable.TypeObject = ScriptStruct;
+		Variable.Size = ScriptStruct->GetStructureSize();
+	}
+	else if (UEnum* Enum= URigVMPin::FindObjectFromCPPTypeObjectPath<UEnum>(CPPType))
+	{
+		Variable.TypeName = *Enum->CppType;
+		Variable.TypeObject = Enum;
+		Variable.Size = Enum->GetResourceSizeBytes(EResourceSizeMode::EstimatedTotal);
+	}
+	else
+	{
+		check(false);
+	}
+
+	return Variable;
+}
+
+FEdGraphPinType RigVMTypeUtils::PinTypeFromExternalVariable(const FRigVMExternalVariable& InExternalVariable)
+{
+	FEdGraphPinType PinType;
+	PinType.ResetToDefaults();
+	PinType.PinCategory = NAME_None;
+
+	if (InExternalVariable.TypeName == TEXT("bool"))
+	{
+		PinType.PinCategory = UEdGraphSchema_K2::PC_Boolean;
+	}
+	else if (InExternalVariable.TypeName == TEXT("int32"))
+	{
+		PinType.PinCategory = UEdGraphSchema_K2::PC_Int;
+	}
+	else if (InExternalVariable.TypeName == TEXT("float"))
+	{
+		PinType.PinCategory = UEdGraphSchema_K2::PC_Float;
+	}
+	else if (InExternalVariable.TypeName == TEXT("double"))
+	{
+		PinType.PinCategory = UEdGraphSchema_K2::PC_Double;
+	}
+	else if (InExternalVariable.TypeName == TEXT("FName"))
+	{
+		PinType.PinCategory = UEdGraphSchema_K2::PC_Name;
+	}
+	else if (InExternalVariable.TypeName == TEXT("FString"))
+	{
+		PinType.PinCategory = UEdGraphSchema_K2::PC_String;
+	}
+	else if (Cast<UScriptStruct>(InExternalVariable.TypeObject))
+	{
+		PinType.PinCategory = UEdGraphSchema_K2::PC_Struct;
+		PinType.PinSubCategoryObject = InExternalVariable.TypeObject;
+	}
+	else if (Cast<UEnum>(InExternalVariable.TypeObject))
+	{
+		PinType.PinCategory = UEdGraphSchema_K2::PC_Byte;
+		PinType.PinSubCategoryObject = InExternalVariable.TypeObject;
+	}
+
+	if (InExternalVariable.bIsArray)
+	{
+		PinType.ContainerType = EPinContainerType::Array;
+	}
+	else
+	{
+		PinType.ContainerType = EPinContainerType::None;
+	}
+
+	return PinType;
+}
+
+FEdGraphPinType RigVMTypeUtils::PinTypeFromRigVMVariableDescription(
+	const FRigVMGraphVariableDescription& InVariableDescription)
+{
+	FEdGraphPinType PinType;
+	PinType.ResetToDefaults();
+	PinType.PinCategory = NAME_None;
+
+	FString CurrentCPPType = InVariableDescription.CPPType;
+	if (CurrentCPPType.StartsWith(TEXT("TArray<")))
+	{
+		PinType.ContainerType = EPinContainerType::Array;
+		CurrentCPPType.RemoveFromStart(TEXT("TArray<"));
+		CurrentCPPType.RemoveFromEnd(TEXT(">"));
+	}
+	else
+	{
+		PinType.ContainerType = EPinContainerType::None;
+	}
+
+	if (CurrentCPPType == TEXT("bool"))
+	{
+		PinType.PinCategory = UEdGraphSchema_K2::PC_Boolean;
+	}
+	else if (CurrentCPPType == TEXT("int32"))
+	{
+		PinType.PinCategory = UEdGraphSchema_K2::PC_Int;
+	}
+	else if (CurrentCPPType == TEXT("float"))
+	{
+		PinType.PinCategory = UEdGraphSchema_K2::PC_Float;
+	}
+	else if (CurrentCPPType == TEXT("double"))
+	{
+		PinType.PinCategory = UEdGraphSchema_K2::PC_Double;
+	}
+	else if (CurrentCPPType == TEXT("FName"))
+	{
+		PinType.PinCategory = UEdGraphSchema_K2::PC_Name;
+	}
+	else if (CurrentCPPType == TEXT("FString"))
+	{
+		PinType.PinCategory = UEdGraphSchema_K2::PC_String;
+	}
+	else if (Cast<UScriptStruct>(InVariableDescription.CPPTypeObject))
+	{
+		PinType.PinCategory = UEdGraphSchema_K2::PC_Struct;
+		PinType.PinSubCategoryObject = InVariableDescription.CPPTypeObject;
+	}
+	else if (Cast<UEnum>(InVariableDescription.CPPTypeObject))
+	{
+		PinType.PinCategory = UEdGraphSchema_K2::PC_Byte;
+		PinType.PinSubCategoryObject = InVariableDescription.CPPTypeObject;
+	}
+	else if (Cast<UClass>(InVariableDescription.CPPTypeObject))
+	{
+		PinType.PinCategory = UEdGraphSchema_K2::PC_Object;
+		PinType.PinSubCategoryObject = InVariableDescription.CPPTypeObject;
+	}
+
+	return PinType;
+}
+
+bool RigVMTypeUtils::CPPTypeFromPinType(const FEdGraphPinType& InPinType, FString& OutCPPType, UObject** OutCPPTypeObject)
+{
+	FString Prefix = "";
+	FString Suffix = "";
+	if (InPinType.ContainerType == EPinContainerType::Array)
+	{
+		Prefix = TEXT("TArray<");
+		Suffix = TEXT(">");
+	}
+
+	*OutCPPTypeObject = nullptr;
+	if (InPinType.PinCategory == UEdGraphSchema_K2::PC_Boolean)
+	{
+		OutCPPType = Prefix + TEXT("bool") + Suffix;
+	}
+	else if (InPinType.PinCategory == UEdGraphSchema_K2::PC_Int)
+	{
+		OutCPPType = Prefix + TEXT("int32") + Suffix;
+	}
+	else if (InPinType.PinCategory == UEdGraphSchema_K2::PC_Float)
+	{
+		OutCPPType = Prefix + TEXT("float") + Suffix;
+	}
+	else if (InPinType.PinCategory == UEdGraphSchema_K2::PC_Double)
+	{
+		OutCPPType = Prefix + TEXT("double") + Suffix;
+	}
+	else if (InPinType.PinCategory == UEdGraphSchema_K2::PC_Name)
+	{
+		OutCPPType = Prefix + TEXT("FName") + Suffix;
+	}
+	else if (InPinType.PinCategory == UEdGraphSchema_K2::PC_String)
+	{
+		OutCPPType = Prefix + TEXT("FString") + Suffix;
+	}
+	else if (InPinType.PinCategory == UEdGraphSchema_K2::PC_Struct)
+	{
+		if (UScriptStruct* Struct = Cast<UScriptStruct>(InPinType.PinSubCategoryObject))
+		{
+			OutCPPType = Prefix + *Struct->GetStructCPPName() + Suffix;
+			*OutCPPTypeObject = Struct;
+		}
+	}
+	else if (InPinType.PinCategory == UEdGraphSchema_K2::PC_Object ||
+		InPinType.PinCategory == UEdGraphSchema_K2::PC_SoftObject ||
+		InPinType.PinCategory == UEdGraphSchema_K2::AllObjectTypes)
+	{
+		if (UClass* Class = Cast<UClass>(InPinType.PinSubCategoryObject))
+		{
+			OutCPPType = Prefix + FString::Printf(TEXT("TObjectPtr<%s%s>"), Class->GetPrefixCPP(), *Class->GetName()) + Suffix;
+			*OutCPPTypeObject = Class;
+		}
+	}
+	else if (InPinType.PinCategory == UEdGraphSchema_K2::PC_Byte ||
+		InPinType.PinCategory == UEdGraphSchema_K2::PC_Enum)
+	{
+		if (UEnum* Enum = Cast<UEnum>(InPinType.PinSubCategoryObject))
+		{
+			OutCPPType = Prefix + Enum->GetFName().ToString() + Suffix;
+			*OutCPPTypeObject = Enum;
+		}
+		else
+		{
+			OutCPPType = Prefix + TEXT("uint8") + Suffix;
+		}
+	}
+	else
+	{
+		return false;
+	}
+		
+	return true;
+}
+
+bool RigVMTypeUtils::CPPTypeFromExternalVariable(const FRigVMExternalVariable& InExternalVariable, FString& OutCPPType, UObject** OutCPPTypeObject)
+{
+	FString Prefix = "";
+	FString Suffix = "";
+	if (InExternalVariable.bIsArray)
+	{
+		Prefix = TEXT("TArray<");
+		Suffix = TEXT(">");
+	}
+
+	*OutCPPTypeObject = nullptr;
+	if (InExternalVariable.TypeName == TEXT("bool"))
+	{
+		OutCPPType = Prefix + TEXT("bool") + Suffix;
+	}
+	else if (InExternalVariable.TypeName == TEXT("int32"))
+	{
+		OutCPPType = Prefix + TEXT("int32") + Suffix;
+	}
+	else if (InExternalVariable.TypeName == TEXT("float"))
+	{
+		OutCPPType = Prefix + TEXT("float") + Suffix;
+	}
+	else if (InExternalVariable.TypeName == TEXT("double"))
+	{
+		OutCPPType = Prefix + TEXT("double") + Suffix;
+	}
+	else if (InExternalVariable.TypeName == TEXT("FName"))
+	{
+		OutCPPType = Prefix + TEXT("FName") + Suffix;
+	}
+	else if (InExternalVariable.TypeName == TEXT("FString"))
+	{
+		OutCPPType = Prefix + TEXT("FString") + Suffix;
+	}
+	else if (UScriptStruct* Struct = Cast<UScriptStruct>(InExternalVariable.TypeObject))
+	{
+		OutCPPType = Prefix + *Struct->GetStructCPPName() + Suffix;
+		*OutCPPTypeObject = Struct;	
+	}
+	else if (UEnum* Enum = Cast<UEnum>(InExternalVariable.TypeObject))
+	{
+		OutCPPType = Prefix + Enum->GetFName().ToString() + Suffix;
+		*OutCPPTypeObject = Enum;
+	}	
+	else
+	{
+		check(false);
+		return false;
+	}
+		
+	return true;
+}
+#endif
diff --git a/Engine/Source/Developer/RigVMDeveloper/Public/RigVMModel/Nodes/RigVMArrayNode.h b/Engine/Source/Developer/RigVMDeveloper/Public/RigVMModel/Nodes/RigVMArrayNode.h
--- a/Engine/Source/Developer/RigVMDeveloper/Public/RigVMModel/Nodes/RigVMArrayNode.h
+++ b/Engine/Source/Developer/RigVMDeveloper/Public/RigVMModel/Nodes/RigVMArrayNode.h
@@ -36,7 +36,7 @@ public:
 	UFUNCTION(BlueprintCallable, Category = RigVMArrayNode)
 	ERigVMOpCode GetOpCode() const;
 
-	// Returns the C++ data type of the array
+	// Returns the C++ data type of the element
 	UFUNCTION(BlueprintCallable, Category = RigVMArrayNode)
 	FString GetCPPType() const;
 
diff --git a/Engine/Source/Developer/RigVMDeveloper/Public/RigVMModel/Nodes/RigVMParameterNode.h b/Engine/Source/Developer/RigVMDeveloper/Public/RigVMModel/Nodes/RigVMParameterNode.h
--- a/Engine/Source/Developer/RigVMDeveloper/Public/RigVMModel/Nodes/RigVMParameterNode.h
+++ b/Engine/Source/Developer/RigVMDeveloper/Public/RigVMModel/Nodes/RigVMParameterNode.h
@@ -4,7 +4,7 @@
 
 #include "RigVMCore/RigVMExternalVariable.h"
 #include "RigVMModel/RigVMNode.h"
-#include "RigVMCore/RigVMUtilities.h"
+#include "RigVMTypeUtils.h"
 #include "RigVMParameterNode.generated.h"
 
 /**
@@ -51,7 +51,7 @@ public:
 		FRigVMExternalVariable ExternalVariable;
 		ExternalVariable.Name = Name;
 
-		if (RigVMUtilities::IsArrayType(CPPType))
+		if (RigVMTypeUtils::IsArrayType(CPPType))
 		{
 			ExternalVariable.bIsArray = true;
 			ExternalVariable.TypeName = *CPPType.Mid(7, CPPType.Len() - 8);
diff --git a/Engine/Source/Developer/RigVMDeveloper/Public/RigVMModel/RigVMVariableDescription.h b/Engine/Source/Developer/RigVMDeveloper/Public/RigVMModel/RigVMVariableDescription.h
--- a/Engine/Source/Developer/RigVMDeveloper/Public/RigVMModel/RigVMVariableDescription.h
+++ b/Engine/Source/Developer/RigVMDeveloper/Public/RigVMModel/RigVMVariableDescription.h
@@ -5,6 +5,7 @@
 #include "RigVMCore/RigVMExternalVariable.h"
 #include "RigVMModel/RigVMNode.h"
 #include "EdGraphSchema_K2.h"
+#include "RigVMTypeUtils.h"
 #include "RigVMVariableDescription.generated.h"
 
 /**
@@ -44,167 +45,20 @@ public:
 	// Returns nullptr external variable matching this description
 	FORCEINLINE FRigVMExternalVariable ToExternalVariable() const
 	{
-		FRigVMExternalVariable ExternalVariable;
-		ExternalVariable.Name = Name;
-
-		if (CPPType.StartsWith(TEXT("TArray<")))
-		{
-			ExternalVariable.bIsArray = true;
-			ExternalVariable.TypeName = *CPPType.Mid(7, CPPType.Len() - 8);
-			ExternalVariable.TypeObject = CPPTypeObject;
-		}
-		else
-		{
-			ExternalVariable.bIsArray = false;
-			ExternalVariable.TypeName = *CPPType;
-			ExternalVariable.TypeObject = CPPTypeObject;
-		}
-
-		ExternalVariable.bIsPublic = false;
-		ExternalVariable.bIsReadOnly = false;
-		ExternalVariable.Memory = nullptr;
-		return ExternalVariable;
+		return RigVMTypeUtils::ExternalVariableFromRigVMVariableDescription(*this);
 	}
 
 	FORCEINLINE FEdGraphPinType ToPinType() const
 	{
-		FEdGraphPinType PinType;
-		PinType.ResetToDefaults();
-		PinType.PinCategory = NAME_None;
-
-		FString CurrentCPPType = CPPType;
-		if (CurrentCPPType.StartsWith(TEXT("TArray<")))
-		{
-			PinType.ContainerType = EPinContainerType::Array;
-			CurrentCPPType.RemoveFromStart(TEXT("TArray<"));
-			CurrentCPPType.RemoveFromEnd(TEXT(">"));
-		}
-		else
-		{
-			PinType.ContainerType = EPinContainerType::None;
-		}
-
-		if (CurrentCPPType == TEXT("bool"))
-		{
-			PinType.PinCategory = UEdGraphSchema_K2::PC_Boolean;
-		}
-		else if (CurrentCPPType == TEXT("int32"))
-		{
-			PinType.PinCategory = UEdGraphSchema_K2::PC_Int;
-		}
-		else if (CurrentCPPType == TEXT("float"))
-		{
-			PinType.PinCategory = UEdGraphSchema_K2::PC_Float;
-		}
-		else if (CurrentCPPType == TEXT("double"))
-		{
-			PinType.PinCategory = UEdGraphSchema_K2::PC_Double;
-		}
-		else if (CurrentCPPType == TEXT("FName"))
-		{
-			PinType.PinCategory = UEdGraphSchema_K2::PC_Name;
-		}
-		else if (CurrentCPPType == TEXT("FString"))
-		{
-			PinType.PinCategory = UEdGraphSchema_K2::PC_String;
-		}
-		else if (Cast<UScriptStruct>(CPPTypeObject))
-		{
-			PinType.PinCategory = UEdGraphSchema_K2::PC_Struct;
-			PinType.PinSubCategoryObject = CPPTypeObject;
-		}
-		else if (Cast<UEnum>(CPPTypeObject))
-		{
-			PinType.PinCategory = UEdGraphSchema_K2::PC_Byte;
-			PinType.PinSubCategoryObject = CPPTypeObject;
-		}
-		else if (Cast<UClass>(CPPTypeObject))
-		{
-			PinType.PinCategory = UEdGraphSchema_K2::PC_Object;
-			PinType.PinSubCategoryObject = CPPTypeObject;
-		}
-
-		return PinType;
-	}
-
-	template <typename T>
-	static bool CPPTypeFromPinType(const FEdGraphPinType& InPinType, FString& OutCPPType, T& OutCPPTypeObject)
-	{
-		FString Prefix = "";
-		FString Suffix = "";
-		if (InPinType.ContainerType == EPinContainerType::Array)
-		{
-			Prefix = TEXT("TArray<");
-			Suffix = TEXT(">");
-		}
-
-		OutCPPTypeObject = nullptr;
-		if (InPinType.PinCategory == UEdGraphSchema_K2::PC_Boolean)
-		{
-			OutCPPType = Prefix + TEXT("bool") + Suffix;
-		}
-		else if (InPinType.PinCategory == UEdGraphSchema_K2::PC_Int)
-		{
-			OutCPPType = Prefix + TEXT("int32") + Suffix;
-		}
-		else if (InPinType.PinCategory == UEdGraphSchema_K2::PC_Float)
-		{
-			OutCPPType = Prefix + TEXT("float") + Suffix;
-		}
-		else if (InPinType.PinCategory == UEdGraphSchema_K2::PC_Double)
-		{
-			OutCPPType = Prefix + TEXT("double") + Suffix;
-		}
-		else if (InPinType.PinCategory == UEdGraphSchema_K2::PC_Name)
-		{
-			OutCPPType = Prefix + TEXT("FName") + Suffix;
-		}
-		else if (InPinType.PinCategory == UEdGraphSchema_K2::PC_String)
-		{
-			OutCPPType = Prefix + TEXT("FString") + Suffix;
-		}
-		else if (InPinType.PinCategory == UEdGraphSchema_K2::PC_Struct)
-		{
-			if (UScriptStruct* Struct = Cast<UScriptStruct>(InPinType.PinSubCategoryObject))
-			{
-				OutCPPType = Prefix + *Struct->GetStructCPPName() + Suffix;
-				OutCPPTypeObject = Struct;
-			}
-		}
-		else if (InPinType.PinCategory == UEdGraphSchema_K2::PC_Object ||
-			InPinType.PinCategory == UEdGraphSchema_K2::PC_SoftObject ||
-			InPinType.PinCategory == UEdGraphSchema_K2::AllObjectTypes)
-		{
-			if (UClass* Class = Cast<UClass>(InPinType.PinSubCategoryObject))
-			{
-				OutCPPType = Prefix + FString::Printf(TEXT("TObjectPtr<%s%s>"), Class->GetPrefixCPP(), *Class->GetName()) + Suffix;
-				OutCPPTypeObject = Class;
-			}
-		}
-		else if (InPinType.PinCategory == UEdGraphSchema_K2::PC_Byte ||
-			InPinType.PinCategory == UEdGraphSchema_K2::PC_Enum)
-		{
-			if (UEnum* Enum = Cast<UEnum>(InPinType.PinSubCategoryObject))
-			{
-				OutCPPType = Prefix + Enum->GetFName().ToString() + Suffix;
-				OutCPPTypeObject = Enum;
-			}
-			else
-			{
-				OutCPPType = Prefix + TEXT("uint8") + Suffix;
-			}
-		}
-		else
-		{
-			return false;
-		}
-		
-		return true;
+		return RigVMTypeUtils::PinTypeFromRigVMVariableDescription(*this);
 	}
 
 	FORCEINLINE bool ChangeType(const FEdGraphPinType& PinType)
 	{
-		return CPPTypeFromPinType(PinType, CPPType, CPPTypeObject);
+		UObject* Object = nullptr;
+		const bool bSuccess = RigVMTypeUtils::CPPTypeFromPinType(PinType, CPPType, &Object);
+		CPPTypeObject = Object;
+		return bSuccess;
 	}
 	
 };
diff --git a/Engine/Source/Developer/RigVMDeveloper/Public/RigVMTypeUtils.h b/Engine/Source/Developer/RigVMDeveloper/Public/RigVMTypeUtils.h
new file mode 100644
--- /dev/null
+++ b/Engine/Source/Developer/RigVMDeveloper/Public/RigVMTypeUtils.h
@@ -0,0 +1,44 @@
+﻿// Copyright Epic Games, Inc. All Rights Reserved.
+
+#pragma once
+
+#include "RigVMDeveloperModule.h"
+#include "RigVMCore/RigVMExternalVariable.h"
+#include "EdGraphSchema_K2.h"
+
+struct FRigVMGraphVariableDescription;
+
+namespace RigVMTypeUtils
+{
+	// Returns true if the type specified is an array
+	FORCEINLINE bool IsArrayType(const FString& InCPPType)
+	{
+		return InCPPType.StartsWith(TEXT("TArray<"));
+	}
+
+	FORCEINLINE FString ArrayTypeFromBaseType(const FString& InCPPType)
+	{
+		return "TArray<" + InCPPType + ">";
+	}
+
+	
+#if WITH_EDITOR
+	RIGVMDEVELOPER_API FRigVMExternalVariable ExternalVariableFromBPVariableDescription(const FBPVariableDescription& InVariableDescription);
+
+	RIGVMDEVELOPER_API FRigVMExternalVariable ExternalVariableFromRigVMVariableDescription(const FRigVMGraphVariableDescription& InVariableDescription);
+	
+	RIGVMDEVELOPER_API FRigVMExternalVariable ExternalVariableFromPinType(const FName& InName, const FEdGraphPinType& InPinType, bool bInPublic = false, bool bInReadonly = false);
+
+	RIGVMDEVELOPER_API FRigVMExternalVariable ExternalVariableFromCPPType(const FName& InName, const FString& InCPPType, bool bInPublic = false, bool bInReadonly = false);
+
+	RIGVMDEVELOPER_API FEdGraphPinType PinTypeFromExternalVariable(const FRigVMExternalVariable& InExternalVariable);
+
+	RIGVMDEVELOPER_API FEdGraphPinType PinTypeFromRigVMVariableDescription(const FRigVMGraphVariableDescription& InVariableDescription);
+
+	RIGVMDEVELOPER_API bool CPPTypeFromPinType(const FEdGraphPinType& InPinType, FString& OutCPPType, UObject** OutCPPTypeObject);
+
+	RIGVMDEVELOPER_API bool CPPTypeFromExternalVariable(const FRigVMExternalVariable& InExternalVariable, FString& OutCPPType, UObject** OutCPPTypeObject);
+
+	
+#endif
+}
diff --git a/Engine/Source/Developer/RigVMDeveloper/RigVMDeveloper.Build.cs b/Engine/Source/Developer/RigVMDeveloper/RigVMDeveloper.Build.cs
--- a/Engine/Source/Developer/RigVMDeveloper/RigVMDeveloper.Build.cs
+++ b/Engine/Source/Developer/RigVMDeveloper/RigVMDeveloper.Build.cs
@@ -27,6 +27,7 @@ public class RigVMDeveloper : ModuleRules
 					"SlateCore",
 					"EditorStyle",
 					"MessageLog",
+					"BlueprintGraph"
 				}
 			);
         }
diff --git a/Engine/Source/Runtime/RigVM/Private/RigVMCore/RigVMMemoryDeprecated.cpp b/Engine/Source/Runtime/RigVM/Private/RigVMCore/RigVMMemoryDeprecated.cpp
--- a/Engine/Source/Runtime/RigVM/Private/RigVMCore/RigVMMemoryDeprecated.cpp
+++ b/Engine/Source/Runtime/RigVM/Private/RigVMCore/RigVMMemoryDeprecated.cpp
@@ -1,7 +1,6 @@
 // Copyright Epic Games, Inc. All Rights Reserved.
 
 #include "RigVMCore/RigVMMemoryDeprecated.h"
-#include "RigVMCore/RigVMUtilities.h"
 #include "UObject/AnimObjectVersion.h"
 #include "UObject/ReleaseObjectVersion.h"
 #include "UObject/PropertyPortFlags.h"
@@ -1691,7 +1690,7 @@ void FRigVMMemoryContainer::SetRegisterValueFromString(const FRigVMOperand& InOp
 	}
 
 	FString CPPType = InCPPType;
-	if (RigVMUtilities::IsArrayType(CPPType))
+	if (RigVMTypeUtils::IsArrayType(CPPType))
 	{
 		CPPType = CPPType.Mid(7, CPPType.Len() - 8);
 	}
diff --git a/Engine/Source/Runtime/RigVM/Public/RigVMCore/RigVMUtilities.h b/Engine/Source/Runtime/RigVM/Public/RigVMCore/RigVMUtilities.h
deleted file mode 100644
