RigVM: Optimized FindProxiesWithSharedOperand

diff --git a/Engine/Source/Developer/RigVMDeveloper/Private/RigVMCompiler/RigVMCompiler.cpp b/Engine/Source/Developer/RigVMDeveloper/Private/RigVMCompiler/RigVMCompiler.cpp
--- a/Engine/Source/Developer/RigVMDeveloper/Private/RigVMCompiler/RigVMCompiler.cpp
+++ b/Engine/Source/Developer/RigVMDeveloper/Private/RigVMCompiler/RigVMCompiler.cpp
@@ -306,6 +306,18 @@ bool URigVMCompiler::Compile(URigVMGraph* InGraph, URigVMController* InControlle
 	WorkData.PinPathToOperand = OutOperands;
 	WorkData.RigVMUserData = UserData[0];
 	WorkData.bSetupMemory = true;
+	WorkData.ProxySources = &AST->SharedOperandPins;
+
+	// tbd: do we need this only when we have no pins?
+	//if(!WorkData.WatchedPins.IsEmpty())
+	{
+		// create the inverse map for the proxies
+		WorkData.ProxyTargets.Reserve(WorkData.ProxySources->Num());
+		for(const TPair<FRigVMASTProxy,FRigVMASTProxy>& Pair : *WorkData.ProxySources)
+		{
+			WorkData.ProxyTargets.FindOrAdd(Pair.Value).Add(Pair.Key);
+		}
+	}
 
 	UE_LOG_RIGVMMEMORY(TEXT("RigVMCompiler: Begin '%s'..."), *InGraph->GetPathName());
 
@@ -2901,20 +2913,22 @@ FRigVMOperand URigVMCompiler::FindOrAddRegister(const FRigVMVarExprAST* InVarExp
 	ensure(Operand.IsValid());
 
 	// Get all possible pins that lead to the same operand
-
 	if(Settings.ASTSettings.bFoldAssignments)
 	{
-		const TMap<FRigVMASTProxy, FRigVMASTProxy>& SharedOperandPins = InVarExpr->GetParser()->SharedOperandPins;
-		TArray<FRigVMASTProxy> PinProxies = FindProxiesWithSharedOperand(InVarExpr);
-		ensure(!PinProxies.IsEmpty());
-
-		for (const FRigVMASTProxy& Proxy : PinProxies)
+		// tbd: this functionality is only needed when there is a watch anywhere?
+		//if(!WorkData.WatchedPins.IsEmpty())
 		{
-			if (URigVMPin* VirtualPin = Cast<URigVMPin>(Proxy.GetSubject()))
+			const FRigVMCompilerWorkData::FRigVMASTProxyArray& PinProxies = FindProxiesWithSharedOperand(InVarExpr, WorkData);
+			ensure(!PinProxies.IsEmpty());
+
+			for (const FRigVMASTProxy& Proxy : PinProxies)
 			{
-				FString VirtualPinHash = GetPinHash(VirtualPin, InVarExpr, bIsDebugValue);
-				WorkData.PinPathToOperand->Add(VirtualPinHash, Operand);
-			}	
+				if (URigVMPin* VirtualPin = Cast<URigVMPin>(Proxy.GetSubject()))
+				{
+					FString VirtualPinHash = GetPinHash(VirtualPin, InVarExpr, bIsDebugValue);
+					WorkData.PinPathToOperand->Add(VirtualPinHash, Operand);
+				}	
+			}
 		}
 	}
 	else
@@ -2934,18 +2948,27 @@ FRigVMOperand URigVMCompiler::FindOrAddRegister(const FRigVMVarExprAST* InVarExp
 	return Operand;
 }
 
-TArray<FRigVMASTProxy> URigVMCompiler::FindProxiesWithSharedOperand(const FRigVMVarExprAST* InVarExpr)
+const FRigVMCompilerWorkData::FRigVMASTProxyArray& URigVMCompiler::FindProxiesWithSharedOperand(const FRigVMVarExprAST* InVarExpr, FRigVMCompilerWorkData& WorkData)
 {
-	TArray<FRigVMASTProxy> PinProxies, PinProxiesToProcess;
-	const TMap<FRigVMASTProxy, FRigVMASTProxy>& SharedOperandPins = InVarExpr->GetParser()->SharedOperandPins;
-	PinProxiesToProcess.Add(InVarExpr->GetProxy());
-	while(!PinProxiesToProcess.IsEmpty())
+	const FRigVMASTProxy& InProxy = InVarExpr->GetProxy();
+	if(const FRigVMCompilerWorkData::FRigVMASTProxyArray* ExistingArray = WorkData.CachedProxiesWithSharedOperand.Find(InProxy))
 	{
-		FRigVMASTProxy Proxy = PinProxiesToProcess.Pop();
+		return *ExistingArray;
+	}
+	
+	FRigVMCompilerWorkData::FRigVMASTProxyArray PinProxies, PinProxiesToProcess;
+	const FRigVMCompilerWorkData::FRigVMASTProxySourceMap& ProxySources = *WorkData.ProxySources;
+	const FRigVMCompilerWorkData::FRigVMASTProxyTargetsMap& ProxyTargets = WorkData.ProxyTargets;
+
+	PinProxiesToProcess.Add(InProxy);
 
-		if (Proxy.IsValid())
+	for(int32 ProxyIndex = 0; ProxyIndex < PinProxiesToProcess.Num(); ProxyIndex++)
+	{
+		const FRigVMASTProxy& CurrentProxy = PinProxiesToProcess[ProxyIndex];
+
+		if (CurrentProxy.IsValid())
 		{
-			if (URigVMPin* Pin = Cast<URigVMPin>(Proxy.GetSubject()))
+			if (URigVMPin* Pin = Cast<URigVMPin>(CurrentProxy.GetSubject()))
 			{
 				if (Pin->GetNode()->IsA<URigVMVariableNode>() || Pin->GetNode()->IsA<URigVMParameterNode>())
 				{
@@ -2955,29 +2978,29 @@ TArray<FRigVMASTProxy> URigVMCompiler::FindProxiesWithSharedOperand(const FRigVM
 					}
 				}
 			}
-			
-			PinProxies.Add(Proxy);
+			PinProxies.Add(CurrentProxy);
 		}
 
-		for (TPair<FRigVMASTProxy, FRigVMASTProxy> Pair : SharedOperandPins)
+		if(const FRigVMASTProxy* SourceProxy = ProxySources.Find(CurrentProxy))
 		{
-			FRigVMASTProxy* Other = nullptr;
-			if (PinProxies.Contains(Pair.Key))
-			{
-				Other = &Pair.Value;
-			}
-			else if (PinProxies.Contains(Pair.Value))
+			if(SourceProxy->IsValid())
 			{
-				Other = &Pair.Key;
+				if (!PinProxies.Contains(*SourceProxy) && !PinProxiesToProcess.Contains(*SourceProxy))
+				{
+					PinProxiesToProcess.Add(*SourceProxy);
+				}
 			}
+		}
 
-			if (Other)
+		if(const FRigVMCompilerWorkData::FRigVMASTProxyArray* TargetProxies = WorkData.ProxyTargets.Find(CurrentProxy))
+		{
+			for(const FRigVMASTProxy& TargetProxy : *TargetProxies)
 			{
-				if (Other->IsValid())
+				if(TargetProxy.IsValid())
 				{
-					if (!PinProxies.Contains(*Other) && !PinProxiesToProcess.Contains(*Other))
+					if (!PinProxies.Contains(TargetProxy) && !PinProxiesToProcess.Contains(TargetProxy))
 					{
-						PinProxiesToProcess.Add(*Other);
+						PinProxiesToProcess.Add(TargetProxy);
 					}
 				}
 			}
@@ -2989,7 +3012,17 @@ TArray<FRigVMASTProxy> URigVMCompiler::FindProxiesWithSharedOperand(const FRigVM
 		PinProxies.Add(InVarExpr->GetProxy());
 	}
 
-	return PinProxies;
+	// store the cache for all other proxies within this group
+	for(const FRigVMASTProxy& CurrentProxy : PinProxies)
+	{
+		if(CurrentProxy != InProxy)
+		{
+			WorkData.CachedProxiesWithSharedOperand.Add(CurrentProxy, PinProxies);
+		}
+	}
+
+	// finally store and return the cache the the input proxy
+	return WorkData.CachedProxiesWithSharedOperand.Add(InProxy, PinProxies);
 }
 
 bool URigVMCompiler::ValidateNode(URigVMNode* InNode)
diff --git a/Engine/Source/Developer/RigVMDeveloper/Public/RigVMCompiler/RigVMCompiler.h b/Engine/Source/Developer/RigVMDeveloper/Public/RigVMCompiler/RigVMCompiler.h
--- a/Engine/Source/Developer/RigVMDeveloper/Public/RigVMCompiler/RigVMCompiler.h
+++ b/Engine/Source/Developer/RigVMDeveloper/Public/RigVMCompiler/RigVMCompiler.h
@@ -92,6 +92,14 @@ public:
 	TMap<int32, FRigVMOperand> IntegerLiterals;
 	FRigVMOperand ComparisonOperand;
 
+	using FRigVMASTProxyArray = TArray<FRigVMASTProxy, TInlineAllocator<3>>; 
+	using FRigVMASTProxySourceMap = TMap<FRigVMASTProxy, FRigVMASTProxy>;
+	using FRigVMASTProxyTargetsMap =
+		TMap<FRigVMASTProxy, FRigVMASTProxyArray>;
+	TMap<FRigVMASTProxy, FRigVMASTProxyArray> CachedProxiesWithSharedOperand;
+	const FRigVMASTProxySourceMap* ProxySources;
+	FRigVMASTProxyTargetsMap ProxyTargets;
+
 #if !UE_RIGVM_UCLASS_BASED_STORAGE_DISABLED
 
 	TArray<URigVMPin*> WatchedPins;
@@ -173,7 +181,7 @@ private:
 	void InitializeLocalVariables(const FRigVMExprAST* InExpr, FRigVMCompilerWorkData& WorkData);
 
 	FRigVMOperand FindOrAddRegister(const FRigVMVarExprAST* InVarExpr, FRigVMCompilerWorkData& WorkData, bool bIsDebugValue = false);
-	TArray<FRigVMASTProxy> FindProxiesWithSharedOperand(const FRigVMVarExprAST* InVarExpr);
+	const FRigVMCompilerWorkData::FRigVMASTProxyArray& FindProxiesWithSharedOperand(const FRigVMVarExprAST* InVarExpr, FRigVMCompilerWorkData& WorkData);
 
 	bool ValidateNode(URigVMNode* InNode);
 	
