Control Rig: Print python context commands when editor opens

diff --git a/Engine/Plugins/Animation/ControlRig/Source/ControlRigEditor/Private/Editor/ControlRigEditor.cpp b/Engine/Plugins/Animation/ControlRig/Source/ControlRigEditor/Private/Editor/ControlRigEditor.cpp
--- a/Engine/Plugins/Animation/ControlRig/Source/ControlRigEditor/Private/Editor/ControlRigEditor.cpp
+++ b/Engine/Plugins/Animation/ControlRig/Source/ControlRigEditor/Private/Editor/ControlRigEditor.cpp
@@ -537,6 +537,11 @@ void FControlRigEditor::InitControlRigEditor(const EToolkitMode::Type Mode, cons
 	}, TStatId(), NULL, ENamedThreads::GameThread);
 	
 	CreateRigHierarchyToGraphDragAndDropMenu();
+
+#if WITH_EDITOR
+	FString BlueprintName = InControlRigBlueprint->GetPathName();
+	RigVMPythonUtils::PrintPythonContext(BlueprintName);
+#endif
 }
 
 void FControlRigEditor::BindCommands()
diff --git a/Engine/Source/Developer/RigVMDeveloper/Private/RigVMPythonUtils.cpp b/Engine/Source/Developer/RigVMDeveloper/Private/RigVMPythonUtils.cpp
--- a/Engine/Source/Developer/RigVMDeveloper/Private/RigVMPythonUtils.cpp
+++ b/Engine/Source/Developer/RigVMDeveloper/Private/RigVMPythonUtils.cpp
@@ -84,4 +84,27 @@ void RigVMPythonUtils::Print(const FString& BlueprintTitle, const FString& InMes
 	PythonLog->AddMessage(Token, false);
 }
 
+void RigVMPythonUtils::PrintPythonContext(const FString& InBlueprintPath)
+{
+	FString BlueprintName = InBlueprintPath;
+	int32 DotIndex = BlueprintName.Find(TEXT("."), ESearchCase::IgnoreCase, ESearchDir::FromEnd);
+	if (DotIndex != INDEX_NONE)
+	{
+		BlueprintName = BlueprintName.Right(BlueprintName.Len() - DotIndex - 1);
+	}
+		
+	TArray<FString> PyCommands = {
+		TEXT("import unreal"),
+		FString::Printf(TEXT("blueprint = unreal.load_object(name = '%s', outer = None)"), *InBlueprintPath),
+		TEXT("library = blueprint.get_local_function_library()"),
+		TEXT("library_controller = blueprint.get_controller(library)"),
+		TEXT("hierarchy = blueprint.hierarchy"),
+		TEXT("hierarchy_controller = hierarchy.get_controller()")};
+
+	for (FString& Command : PyCommands)
+	{
+		RigVMPythonUtils::Print(BlueprintName, Command);
+	}
+}
+
 #endif
diff --git a/Engine/Source/Developer/RigVMDeveloper/Public/RigVMPythonUtils.h b/Engine/Source/Developer/RigVMDeveloper/Public/RigVMPythonUtils.h
--- a/Engine/Source/Developer/RigVMDeveloper/Public/RigVMPythonUtils.h
+++ b/Engine/Source/Developer/RigVMDeveloper/Public/RigVMPythonUtils.h
@@ -17,5 +17,7 @@ namespace RigVMPythonUtils
 #if WITH_EDITOR
 	RIGVMDEVELOPER_API void Print(const FString& BlueprintTitle, const FString& InMessage);
 
+	RIGVMDEVELOPER_API void PrintPythonContext(const FString& InBlueprintName);
+
 #endif
 }
