Control Rig: Fix non-unity build.

#rnx

diff --git a/Engine/Source/Developer/RigVMDeveloper/Private/RigVMTypeUtils.cpp b/Engine/Source/Developer/RigVMDeveloper/Private/RigVMTypeUtils.cpp
--- a/Engine/Source/Developer/RigVMDeveloper/Private/RigVMTypeUtils.cpp
+++ b/Engine/Source/Developer/RigVMDeveloper/Private/RigVMTypeUtils.cpp
@@ -4,6 +4,9 @@
 #include "RigVMModel/RigVMVariableDescription.h"
 
 #if WITH_EDITOR
+
+#include "Engine/Blueprint.h"
+
 FRigVMExternalVariable RigVMTypeUtils::ExternalVariableFromBPVariableDescription(
 	const FBPVariableDescription& InVariableDescription)
 {
