Control Rig: Fix python script when array nodes are present

diff --git a/Engine/Plugins/Animation/ControlRig/Source/ControlRigDeveloper/Private/ControlRigBlueprint.cpp b/Engine/Plugins/Animation/ControlRig/Source/ControlRigDeveloper/Private/ControlRigBlueprint.cpp
--- a/Engine/Plugins/Animation/ControlRig/Source/ControlRigDeveloper/Private/ControlRigBlueprint.cpp
+++ b/Engine/Plugins/Animation/ControlRig/Source/ControlRigDeveloper/Private/ControlRigBlueprint.cpp
@@ -2131,7 +2131,7 @@ bool UControlRigBlueprint::ChangeMemberVariableType(const FName& InName, const F
 	Variable.bIsArray = RigVMTypeUtils::IsArrayType(CPPType);
 	if (Variable.bIsArray)
 	{
-		CPPType = CPPType.RightChop(7).LeftChop(1);
+		CPPType = RigVMTypeUtils::BaseTypeFromArrayType(CPPType);
 	}
 
 	if (CPPType == TEXT("bool"))
diff --git a/Engine/Source/Developer/RigVMDeveloper/Private/RigVMModel/Nodes/RigVMArrayNode.cpp b/Engine/Source/Developer/RigVMDeveloper/Private/RigVMModel/Nodes/RigVMArrayNode.cpp
--- a/Engine/Source/Developer/RigVMDeveloper/Private/RigVMModel/Nodes/RigVMArrayNode.cpp
+++ b/Engine/Source/Developer/RigVMDeveloper/Private/RigVMModel/Nodes/RigVMArrayNode.cpp
@@ -398,12 +398,7 @@ FString URigVMArrayNode::GetCPPType() const
 	{
 		return FString();
 	}
-	const FString& CPPType = ArrayPin->GetCPPType();
-	if (RigVMTypeUtils::IsArrayType(CPPType))
-	{
-		return CPPType.Mid(7, CPPType.Len() - 8);
-	}
-	return CPPType;
+	return RigVMTypeUtils::BaseTypeFromArrayType(ArrayPin->GetCPPType());
 }
 
 UObject* URigVMArrayNode::GetCPPTypeObject() const
diff --git a/Engine/Source/Developer/RigVMDeveloper/Private/RigVMModel/RigVMController.cpp b/Engine/Source/Developer/RigVMDeveloper/Private/RigVMModel/RigVMController.cpp
--- a/Engine/Source/Developer/RigVMDeveloper/Private/RigVMModel/RigVMController.cpp
+++ b/Engine/Source/Developer/RigVMDeveloper/Private/RigVMModel/RigVMController.cpp
@@ -1006,7 +1006,7 @@ URigVMVariableNode* URigVMController::AddVariableNode(const FName& InVariableNam
 		InCPPTypeObject = URigVMPin::FindObjectFromCPPTypeObjectPath<UObject>(InCPPType);
 	}
 
-	FString CPPType = PostProcessCPPType(InCPPType, InCPPTypeObject, InCPPType.StartsWith(TArrayPrefix));
+	FString CPPType = PostProcessCPPType(InCPPType, InCPPTypeObject);
 	
 	FString Name = GetValidNodeName(InNodeName.IsEmpty() ? FString(TEXT("VariableNode")) : InNodeName);
 	URigVMVariableNode* Node = NewObject<URigVMVariableNode>(Graph, *Name);
@@ -1637,7 +1637,7 @@ URigVMParameterNode* URigVMController::AddParameterNode(const FName& InParameter
 
 	if (DefaultValuePin)
 	{
-		DefaultValuePin->CPPType = PostProcessCPPType(InCPPType, InCPPTypeObject, InCPPType.StartsWith(TArrayPrefix));
+		DefaultValuePin->CPPType = PostProcessCPPType(InCPPType, InCPPTypeObject);
 		DefaultValuePin->CPPTypeObject = InCPPTypeObject;
 		if(DefaultValuePin->CPPTypeObject)
 		{
@@ -1645,7 +1645,7 @@ URigVMParameterNode* URigVMController::AddParameterNode(const FName& InParameter
 		}
 	}
 	
-	ValuePin->CPPType = PostProcessCPPType(InCPPType, InCPPTypeObject, InCPPType.StartsWith(TArrayPrefix));
+	ValuePin->CPPType = PostProcessCPPType(InCPPType, InCPPTypeObject);
 	ValuePin->CPPTypeObject = InCPPTypeObject;
 	if(ValuePin->CPPTypeObject)
 	{
@@ -7664,7 +7664,7 @@ FName URigVMController::AddExposedPin(const FName& InPinName, ERigVMPinDirection
 	}, false, true);
 
 	URigVMPin* Pin = NewObject<URigVMPin>(LibraryNode, PinName);
-	Pin->CPPType = PostProcessCPPType(InCPPType, CPPTypeObject, InCPPType.StartsWith(TArrayPrefix));
+	Pin->CPPType = PostProcessCPPType(InCPPType, CPPTypeObject);
 	Pin->CPPTypeObjectPath = InCPPTypeObjectPath;
 	Pin->bIsConstant = false;
 	Pin->Direction = InDirection;
@@ -9553,7 +9553,7 @@ URigVMIfNode* URigVMController::AddIfNode(const FString& InCPPType, const FName&
 		}
 	}
 
-	FString CPPType = PostProcessCPPType(InCPPType, CPPTypeObject, InCPPType.StartsWith(TArrayPrefix));
+	FString CPPType = PostProcessCPPType(InCPPType, CPPTypeObject);
 
 	FString DefaultValue;
 	if(UScriptStruct* ScriptStruct = Cast<UScriptStruct>(CPPTypeObject))
@@ -9669,7 +9669,7 @@ URigVMSelectNode* URigVMController::AddSelectNode(const FString& InCPPType, cons
 		}
 	}
 
-	FString CPPType = PostProcessCPPType(InCPPType, CPPTypeObject, InCPPType.StartsWith(TArrayPrefix));
+	FString CPPType = PostProcessCPPType(InCPPType, CPPTypeObject);
 
 	FString DefaultValue;
 	if (UScriptStruct* ScriptStruct = Cast<UScriptStruct>(CPPTypeObject))
@@ -9939,7 +9939,7 @@ URigVMArrayNode* URigVMController::AddArrayNode(ERigVMOpCode InOpCode, const FSt
 		InCPPTypeObject = URigVMPin::FindObjectFromCPPTypeObjectPath<UObject>(InCPPType);
 	}
 
-	FString CPPType = PostProcessCPPType(InCPPType, InCPPTypeObject, InCPPType.StartsWith(TArrayPrefix));
+	FString CPPType = PostProcessCPPType(InCPPType, InCPPTypeObject);
 	
 	const FString Name = GetValidNodeName(InNodeName.IsEmpty() ? FString(TEXT("ArrayNode")) : InNodeName);
 	URigVMArrayNode* Node = NewObject<URigVMArrayNode>(Graph, *Name);
@@ -12041,7 +12041,7 @@ void URigVMController::PostProcessDefaultValue(URigVMPin* Pin, FString& OutDefau
 	}
 }
 
-FString URigVMController::PostProcessCPPType(const FString& InCPPType, UObject* InCPPTypeObject, bool bIsArray)
+FString URigVMController::PostProcessCPPType(const FString& InCPPType, UObject* InCPPTypeObject)
 {
 	FString CPPType = InCPPType;
 	
@@ -12058,12 +12058,14 @@ FString URigVMController::PostProcessCPPType(const FString& InCPPType, UObject*
 		CPPType = Enum->CppType;
 	}
 
-	if(bIsArray)
+	if(CPPType != InCPPType)
 	{
-		if(!CPPType.StartsWith(TArrayPrefix))
+		FString TemplateType = InCPPType;
+		while (RigVMTypeUtils::IsArrayType(TemplateType))
 		{
-			CPPType = FString::Printf(TArrayTemplate, *CPPType);
-		}
+			CPPType = RigVMTypeUtils::ArrayTypeFromBaseType(CPPType);
+			TemplateType = RigVMTypeUtils::BaseTypeFromArrayType(TemplateType);
+		}		
 	}
 	
 	return CPPType;
@@ -12509,7 +12511,7 @@ bool URigVMController::EnsurePinValidity(URigVMPin* InPin, bool bRecursive)
 		}
 	}
 
-	InPin->CPPType = PostProcessCPPType(InPin->CPPType, InPin->GetCPPTypeObject(), InPin->IsArray());
+	InPin->CPPType = PostProcessCPPType(InPin->CPPType, InPin->GetCPPTypeObject());
 
 	if(bRecursive)
 	{
diff --git a/Engine/Source/Developer/RigVMDeveloper/Private/RigVMModel/RigVMPin.cpp b/Engine/Source/Developer/RigVMDeveloper/Private/RigVMModel/RigVMPin.cpp
--- a/Engine/Source/Developer/RigVMDeveloper/Private/RigVMModel/RigVMPin.cpp
+++ b/Engine/Source/Developer/RigVMDeveloper/Private/RigVMModel/RigVMPin.cpp
@@ -503,7 +503,7 @@ int32 URigVMPin::GetArraySize() const
 
 FString URigVMPin::GetCPPType() const
 {
-	return URigVMController::PostProcessCPPType(CPPType, GetCPPTypeObject(), IsArray());
+	return URigVMController::PostProcessCPPType(CPPType, GetCPPTypeObject());
 }
 
 FString URigVMPin::GetArrayElementCppType() const
diff --git a/Engine/Source/Developer/RigVMDeveloper/Public/RigVMModel/RigVMController.h b/Engine/Source/Developer/RigVMDeveloper/Public/RigVMModel/RigVMController.h
--- a/Engine/Source/Developer/RigVMDeveloper/Public/RigVMModel/RigVMController.h
+++ b/Engine/Source/Developer/RigVMDeveloper/Public/RigVMModel/RigVMController.h
@@ -857,7 +857,7 @@ private:
 	static FLinearColor GetColorFromMetadata(const FString& InMetadata);
 	static void CreateDefaultValueForStructIfRequired(UScriptStruct* InStruct, FString& InOutDefaultValue);
 	static void PostProcessDefaultValue(URigVMPin* Pin, FString& OutDefaultValue);
-	static FString PostProcessCPPType(const FString& InCPPType, UObject* InCPPTypeObject, bool bIsArray);
+	static FString PostProcessCPPType(const FString& InCPPType, UObject* InCPPTypeObject);
 
 	/*
 	void PotentiallyResolvePrototypeNode(URigVMPrototypeNode* InNode, bool bSetupUndoRedo);
diff --git a/Engine/Source/Developer/RigVMDeveloper/Public/RigVMTypeUtils.h b/Engine/Source/Developer/RigVMDeveloper/Public/RigVMTypeUtils.h
--- a/Engine/Source/Developer/RigVMDeveloper/Public/RigVMTypeUtils.h
+++ b/Engine/Source/Developer/RigVMDeveloper/Public/RigVMTypeUtils.h
@@ -21,6 +21,11 @@ namespace RigVMTypeUtils
 		return "TArray<" + InCPPType + ">";
 	}
 
+	FORCEINLINE FString BaseTypeFromArrayType(const FString& InCPPType)
+	{
+		return InCPPType.RightChop(7).LeftChop(1);
+	}
+
 	
 #if WITH_EDITOR
 	RIGVMDEVELOPER_API FRigVMExternalVariable ExternalVariableFromBPVariableDescription(const FBPVariableDescription& InVariableDescription);
