Control Rig: Refactor variable binding

diff --git a/Engine/Plugins/Animation/ControlRig/Source/ControlRig/Public/ControlRigObjectVersion.h b/Engine/Plugins/Animation/ControlRig/Source/ControlRig/Public/ControlRigObjectVersion.h
--- a/Engine/Plugins/Animation/ControlRig/Source/ControlRig/Public/ControlRigObjectVersion.h
+++ b/Engine/Plugins/Animation/ControlRig/Source/ControlRig/Public/ControlRigObjectVersion.h
@@ -45,6 +45,9 @@ struct CONTROLRIG_API FControlRigObjectVersion
 		// RenameGizmoToShape 
 		RenameGizmoToShape,
 
+		// BoundVariableWithInjectionNode 
+		BoundVariableWithInjectionNode,
+
 		// -----<new versions can be added above this line>-------------------------------------------------
 		VersionPlusOne,
 		LatestVersion = VersionPlusOne - 1
diff --git a/Engine/Plugins/Animation/ControlRig/Source/ControlRigDeveloper/Private/ControlRigBlueprint.cpp b/Engine/Plugins/Animation/ControlRig/Source/ControlRigDeveloper/Private/ControlRigBlueprint.cpp
--- a/Engine/Plugins/Animation/ControlRig/Source/ControlRigDeveloper/Private/ControlRigBlueprint.cpp
+++ b/Engine/Plugins/Animation/ControlRig/Source/ControlRigDeveloper/Private/ControlRigBlueprint.cpp
@@ -447,6 +447,7 @@ void UControlRigBlueprint::PostLoad()
 		PatchFunctionReferencesOnLoad();
 		PatchVariableNodesOnLoad();
 		PatchRigElementKeyCacheOnLoad();
+		PatchBoundVariables();
 
 #if WITH_EDITOR
 
@@ -2067,7 +2068,7 @@ TArray<UStruct*> UControlRigBlueprint::GetAvailableRigUnits()
 
 FName UControlRigBlueprint::AddMemberVariable(const FName& InName, const FString& InCPPType, bool bIsPublic, bool bIsReadOnly, FString InDefaultValue)
 {
-	FRigVMExternalVariable Variable = RigVMTypeUtils::ExternalVariableFromCPPType(InName, InCPPType, bIsPublic, bIsReadOnly);
+	FRigVMExternalVariable Variable = RigVMTypeUtils::ExternalVariableFromCPPTypePath(InName, InCPPType, bIsPublic, bIsReadOnly);
 	FName Result = AddCRMemberVariableFromExternal(Variable, InDefaultValue);
 	if (!Result.IsNone())
 	{
@@ -3510,6 +3511,39 @@ void UControlRigBlueprint::PatchRigElementKeyCacheOnLoad()
 	}
 }
 
+void UControlRigBlueprint::PatchBoundVariables()
+{
+	if (GetLinkerCustomVersion(FControlRigObjectVersion::GUID) < FControlRigObjectVersion::BoundVariableWithInjectionNode)
+	{
+		TGuardValue<bool> GuardNotifsSelf(bSuspendModelNotificationsForSelf, true);
+		
+		for (URigVMGraph* Graph : GetAllModels())
+		{
+			URigVMController* Controller = GetOrCreateController(Graph);
+			TArray<URigVMNode*> Nodes = Graph->GetNodes();
+			for (URigVMNode* Node : Nodes)
+			{
+				for (URigVMPin* Pin : Node->GetPins())
+				{
+					for (URigVMInjectionInfo* Info : Pin->GetInjectedNodes())
+					{
+						Info->Node = Info->UnitNode_DEPRECATED;
+						Info->UnitNode_DEPRECATED = nullptr;
+						bDirtyDuringLoad = true;						
+					}
+					
+					if (!Pin->BoundVariablePath_DEPRECATED.IsEmpty())
+					{
+						Controller->BindPinToVariable(Pin->GetPinPath(), Pin->BoundVariablePath_DEPRECATED, false);
+						Pin->BoundVariablePath_DEPRECATED = FString();
+						bDirtyDuringLoad = true;
+					}
+				}
+			}
+		}
+	}
+}
+
 void UControlRigBlueprint::PropagatePoseFromInstanceToBP(UControlRig* InControlRig)
 {
 	check(InControlRig);
diff --git a/Engine/Plugins/Animation/ControlRig/Source/ControlRigDeveloper/Public/ControlRigBlueprint.h b/Engine/Plugins/Animation/ControlRig/Source/ControlRigDeveloper/Public/ControlRigBlueprint.h
--- a/Engine/Plugins/Animation/ControlRig/Source/ControlRigDeveloper/Public/ControlRigBlueprint.h
+++ b/Engine/Plugins/Animation/ControlRig/Source/ControlRigDeveloper/Public/ControlRigBlueprint.h
@@ -538,6 +538,7 @@ private:
 	void PatchFunctionReferencesOnLoad();
 	void PatchVariableNodesOnLoad();
 	void PatchRigElementKeyCacheOnLoad();
+	void PatchBoundVariables();
 
 	TMap<FName, int32> AddedMemberVariableMap;
 	TArray<FBPVariableDescription> LastNewVariables;
diff --git a/Engine/Plugins/Animation/ControlRig/Source/ControlRigEditor/Private/ControlRigEditorModule.cpp b/Engine/Plugins/Animation/ControlRig/Source/ControlRigEditor/Private/ControlRigEditorModule.cpp
--- a/Engine/Plugins/Animation/ControlRig/Source/ControlRigEditor/Private/ControlRigEditorModule.cpp
+++ b/Engine/Plugins/Animation/ControlRig/Source/ControlRigEditor/Private/ControlRigEditorModule.cpp
@@ -1383,74 +1383,86 @@ void FControlRigEditorModule::GetContextMenuActions(const UControlRigGraphSchema
 						{
 							FToolMenuSection& Section = Menu->AddSection("EdGraphSchemaNodeInjectionInterp", LOCTEXT("NodeInjectionInterp", "Interpolate"));
 							URigVMNode* InterpNode = nullptr;
+							bool bBoundToVariable = false;
 							for (URigVMInjectionInfo* Injection : ModelPin->GetInjectedNodes())
 							{
 								FString PrototypeName;
-								if (Injection->UnitNode->GetScriptStruct()->GetStringMetaDataHierarchical(TEXT("PrototypeName"), &PrototypeName))
+								if (URigVMUnitNode* UnitNode = Cast<URigVMUnitNode>(Injection->Node))
 								{
-									if (PrototypeName == TEXT("AlphaInterp"))
+									if (UnitNode->GetScriptStruct()->GetStringMetaDataHierarchical(TEXT("PrototypeName"), &PrototypeName))
 									{
-										InterpNode = Injection->UnitNode;
-										break;
+										if (PrototypeName == TEXT("AlphaInterp"))
+										{
+											InterpNode = Injection->Node;
+											break;
+										}
 									}
 								}
+								else if (URigVMVariableNode* VariableNode = Cast<URigVMVariableNode>(Injection->Node))
+								{
+									bBoundToVariable = true;
+									break;
+								}
 							}
 
-							if (InterpNode == nullptr)
+							if(!bBoundToVariable)
 							{
-								UScriptStruct* ScriptStruct = nullptr;
-
-								if ((ModelPin->GetCPPType() == TEXT("float")) || (ModelPin->GetCPPType() == TEXT("double")))
+								if (InterpNode == nullptr)
 								{
-									ScriptStruct = FRigUnit_AlphaInterp::StaticStruct();
-								}
-								else if (ModelPin->GetCPPType() == TEXT("FVector"))
-								{
-									ScriptStruct = FRigUnit_AlphaInterpVector::StaticStruct();
+									UScriptStruct* ScriptStruct = nullptr;
+
+									if ((ModelPin->GetCPPType() == TEXT("float")) || (ModelPin->GetCPPType() == TEXT("double")))
+									{
+										ScriptStruct = FRigUnit_AlphaInterp::StaticStruct();
+									}
+									else if (ModelPin->GetCPPType() == TEXT("FVector"))
+									{
+										ScriptStruct = FRigUnit_AlphaInterpVector::StaticStruct();
+									}
+									else
+									{
+										checkNoEntry();
+									}
+
+									Section.AddMenuEntry(
+										"AddAlphaInterp",
+										LOCTEXT("AddAlphaInterp", "Add Interpolate"),
+										LOCTEXT("AddAlphaInterp_Tooltip", "Injects an interpolate node"),
+										FSlateIcon(),
+										FUIAction(FExecuteAction::CreateLambda([Controller, InGraphPin, ModelPin, ScriptStruct]() {
+											URigVMInjectionInfo* Injection = Controller->AddInjectedNode(ModelPin->GetPinPath(), ModelPin->GetDirection() != ERigVMPinDirection::Output, ScriptStruct, FRigUnit::GetMethodName(), TEXT("Value"), TEXT("Result"), FString(), true, true);
+											if (Injection)
+											{
+												TArray<FName> NodeNames;
+												NodeNames.Add(Injection->Node->GetFName());
+												Controller->SetNodeSelection(NodeNames);
+											}
+										})
+									));
 								}
 								else
 								{
-									checkNoEntry();
-								}
-
-								Section.AddMenuEntry(
-									"AddAlphaInterp",
-									LOCTEXT("AddAlphaInterp", "Add Interpolate"),
-									LOCTEXT("AddAlphaInterp_Tooltip", "Injects an interpolate node"),
-									FSlateIcon(),
-									FUIAction(FExecuteAction::CreateLambda([Controller, InGraphPin, ModelPin, ScriptStruct]() {
-										URigVMInjectionInfo* Injection = Controller->AddInjectedNode(ModelPin->GetPinPath(), ModelPin->GetDirection() != ERigVMPinDirection::Output, ScriptStruct, FRigUnit::GetMethodName(), TEXT("Value"), TEXT("Result"), FString(), true, true);
-										if (Injection)
-										{
-											TArray<FName> NodeNames;
-											NodeNames.Add(Injection->UnitNode->GetFName());
-											Controller->SetNodeSelection(NodeNames);
-										}
+									Section.AddMenuEntry(
+										"EditAlphaInterp",
+										LOCTEXT("EditAlphaInterp", "Edit Interpolate"),
+										LOCTEXT("EditAlphaInterp_Tooltip", "Edit the interpolate node"),
+										FSlateIcon(),
+										FUIAction(FExecuteAction::CreateLambda([RigBlueprint, InterpNode]() {
+										TArray<FName> NodeNames;
+										NodeNames.Add(InterpNode->GetFName());
+										RigBlueprint->GetController(InterpNode->GetGraph())->SetNodeSelection(NodeNames);
 									})
-								));
-							}
-							else
-							{
-								Section.AddMenuEntry(
-									"EditAlphaInterp",
-									LOCTEXT("EditAlphaInterp", "Edit Interpolate"),
-									LOCTEXT("EditAlphaInterp_Tooltip", "Edit the interpolate node"),
-									FSlateIcon(),
-									FUIAction(FExecuteAction::CreateLambda([RigBlueprint, InterpNode]() {
-									TArray<FName> NodeNames;
-									NodeNames.Add(InterpNode->GetFName());
-									RigBlueprint->GetController(InterpNode->GetGraph())->SetNodeSelection(NodeNames);
-								})
+										));
+									Section.AddMenuEntry(
+										"RemoveAlphaInterp",
+										LOCTEXT("RemoveAlphaInterp", "Remove Interpolate"),
+										LOCTEXT("RemoveAlphaInterp_Tooltip", "Removes the interpolate node"),
+										FSlateIcon(),
+										FUIAction(FExecuteAction::CreateLambda([Controller, InGraphPin, ModelPin, InterpNode]() {
+											Controller->RemoveNodeByName(InterpNode->GetFName(), true, false, true);
+										})
 									));
-								Section.AddMenuEntry(
-									"RemoveAlphaInterp",
-									LOCTEXT("RemoveAlphaInterp", "Remove Interpolate"),
-									LOCTEXT("RemoveAlphaInterp_Tooltip", "Removes the interpolate node"),
-									FSlateIcon(),
-									FUIAction(FExecuteAction::CreateLambda([Controller, InGraphPin, ModelPin, InterpNode]() {
-										Controller->RemoveNodeByName(InterpNode->GetFName(), true, false, true);
-									})
-								));
+								}
 							}
 						}
 
@@ -1461,122 +1473,134 @@ void FControlRigEditorModule::GetContextMenuActions(const UControlRigGraphSchema
 							FToolMenuSection& Section = Menu->AddSection("EdGraphSchemaNodeInjectionVisualDebug", LOCTEXT("NodeInjectionVisualDebug", "Visual Debug"));
 
 							URigVMNode* VisualDebugNode = nullptr;
+							bool bBoundToVariable = false;
 							for (URigVMInjectionInfo* Injection : ModelPin->GetInjectedNodes())
 							{
 								FString PrototypeName;
-								if (Injection->UnitNode->GetScriptStruct()->GetStringMetaDataHierarchical(TEXT("PrototypeName"), &PrototypeName))
+								if (URigVMUnitNode* UnitNode = Cast<URigVMUnitNode>(Injection->Node))
 								{
-									if (PrototypeName == TEXT("VisualDebug"))
+									if (UnitNode->GetScriptStruct()->GetStringMetaDataHierarchical(TEXT("PrototypeName"), &PrototypeName))
 									{
-										VisualDebugNode = Injection->UnitNode;
-										break;
+										if (PrototypeName == TEXT("VisualDebug"))
+										{
+											VisualDebugNode = Injection->Node;
+											break;
+										}
 									}
 								}
+								else if (URigVMVariableNode* VariableNode = Cast<URigVMVariableNode>(Injection->Node))
+								{
+									bBoundToVariable = true;
+									break;
+								}
 							}
 
-							if (VisualDebugNode == nullptr)
+							if (!bBoundToVariable)
 							{
-								UScriptStruct* ScriptStruct = nullptr;
-
-								if (ModelPin->GetCPPType() == TEXT("FVector"))
+								if (VisualDebugNode == nullptr)
 								{
-									ScriptStruct = FRigUnit_VisualDebugVectorItemSpace::StaticStruct();
-								}
-								else if (ModelPin->GetCPPType() == TEXT("FQuat"))
-								{
-									ScriptStruct = FRigUnit_VisualDebugQuatItemSpace::StaticStruct();
-								}
-								else if (ModelPin->GetCPPType() == TEXT("FTransform"))
-								{
-									ScriptStruct = FRigUnit_VisualDebugTransformItemSpace::StaticStruct();
-								}
-								else
-								{
-									checkNoEntry();
-								}
+									UScriptStruct* ScriptStruct = nullptr;
 
-								Section.AddMenuEntry(
-									"AddVisualDebug",
-									LOCTEXT("AddVisualDebug", "Add Visual Debug"),
-									LOCTEXT("AddVisualDebug_Tooltip", "Injects a visual debugging node"),
-									FSlateIcon(),
-									FUIAction(FExecuteAction::CreateLambda([RigBlueprint, Controller, InGraphPin, ModelPin, ScriptStruct]() {
-										URigVMInjectionInfo* Injection = Controller->AddInjectedNode(ModelPin->GetPinPath(), ModelPin->GetDirection() != ERigVMPinDirection::Output, ScriptStruct, FRigUnit::GetMethodName(), TEXT("Value"), TEXT("Value"), FString(), true, true);
-										if (Injection)
-										{
-											TArray<FName> NodeNames;
-											NodeNames.Add(Injection->UnitNode->GetFName());
-											Controller->SetNodeSelection(NodeNames);
+									if (ModelPin->GetCPPType() == TEXT("FVector"))
+									{
+										ScriptStruct = FRigUnit_VisualDebugVectorItemSpace::StaticStruct();
+									}
+									else if (ModelPin->GetCPPType() == TEXT("FQuat"))
+									{
+										ScriptStruct = FRigUnit_VisualDebugQuatItemSpace::StaticStruct();
+									}
+									else if (ModelPin->GetCPPType() == TEXT("FTransform"))
+									{
+										ScriptStruct = FRigUnit_VisualDebugTransformItemSpace::StaticStruct();
+									}
+									else
+									{
+										checkNoEntry();
+									}
 
-											if (URigVMUnitNode* UnitNode = Cast<URigVMUnitNode>(ModelPin->GetNode()))
+									Section.AddMenuEntry(
+										"AddVisualDebug",
+										LOCTEXT("AddVisualDebug", "Add Visual Debug"),
+										LOCTEXT("AddVisualDebug_Tooltip", "Injects a visual debugging node"),
+										FSlateIcon(),
+										FUIAction(FExecuteAction::CreateLambda([RigBlueprint, Controller, InGraphPin, ModelPin, ScriptStruct]() {
+											URigVMInjectionInfo* Injection = Controller->AddInjectedNode(ModelPin->GetPinPath(), ModelPin->GetDirection() != ERigVMPinDirection::Output, ScriptStruct, FRigUnit::GetMethodName(), TEXT("Value"), TEXT("Value"), FString(), true, true);
+											if (Injection)
 											{
-												if (TSharedPtr<FStructOnScope> DefaultStructScope = UnitNode->ConstructStructInstance())
+												TArray<FName> NodeNames;
+												NodeNames.Add(Injection->Node->GetFName());
+												Controller->SetNodeSelection(NodeNames);
+
+												if (URigVMUnitNode* UnitNode = Cast<URigVMUnitNode>(ModelPin->GetNode()))
 												{
-													FRigUnit* DefaultStruct = (FRigUnit*)DefaultStructScope->GetStructMemory();
+													if (TSharedPtr<FStructOnScope> DefaultStructScope = UnitNode->ConstructStructInstance())
+													{
+														FRigUnit* DefaultStruct = (FRigUnit*)DefaultStructScope->GetStructMemory();
 
-													FString PinPath = ModelPin->GetPinPath();
-													FString Left, Right;
+														FString PinPath = ModelPin->GetPinPath();
+														FString Left, Right;
 
-													FRigElementKey SpaceKey;
-													if (URigVMPin::SplitPinPathAtStart(PinPath, Left, Right))
-													{
-														SpaceKey = DefaultStruct->DetermineSpaceForPin(Right, RigBlueprint->Hierarchy);
-													}
+														FRigElementKey SpaceKey;
+														if (URigVMPin::SplitPinPathAtStart(PinPath, Left, Right))
+														{
+															SpaceKey = DefaultStruct->DetermineSpaceForPin(Right, RigBlueprint->Hierarchy);
+														}
 
-													if (SpaceKey.IsValid())
-													{
-														if (URigVMPin* SpacePin = Injection->UnitNode->FindPin(TEXT("Space")))
+														if (SpaceKey.IsValid())
 														{
-															if(URigVMPin* SpaceTypePin = SpacePin->FindSubPin(TEXT("Type")))
+															if (URigVMPin* SpacePin = Injection->Node->FindPin(TEXT("Space")))
 															{
-																FString SpaceTypeStr = StaticEnum<ERigElementType>()->GetDisplayNameTextByValue((int64)SpaceKey.Type).ToString();
-																Controller->SetPinDefaultValue(SpaceTypePin->GetPinPath(), SpaceTypeStr, true, true, false, true);
-															}
-															if(URigVMPin* SpaceNamePin = SpacePin->FindSubPin(TEXT("Name")))
-															{
-																Controller->SetPinDefaultValue(SpaceNamePin->GetPinPath(), SpaceKey.Name.ToString(), true, true, false, true);
+																if(URigVMPin* SpaceTypePin = SpacePin->FindSubPin(TEXT("Type")))
+																{
+																	FString SpaceTypeStr = StaticEnum<ERigElementType>()->GetDisplayNameTextByValue((int64)SpaceKey.Type).ToString();
+																	Controller->SetPinDefaultValue(SpaceTypePin->GetPinPath(), SpaceTypeStr, true, true, false, true);
+																}
+																if(URigVMPin* SpaceNamePin = SpacePin->FindSubPin(TEXT("Name")))
+																{
+																	Controller->SetPinDefaultValue(SpaceNamePin->GetPinPath(), SpaceKey.Name.ToString(), true, true, false, true);
+																}
 															}
 														}
 													}
 												}
 											}
-										}
-									})
-								));
-							}
-							else
-							{
-								Section.AddMenuEntry(
-									"EditVisualDebug",
-									LOCTEXT("EditVisualDebug", "Edit Visual Debug"),
-									LOCTEXT("EditVisualDebug_Tooltip", "Edit the visual debugging node"),
-									FSlateIcon(),
-									FUIAction(FExecuteAction::CreateLambda([Controller, VisualDebugNode]() {
-										TArray<FName> NodeNames;
-										NodeNames.Add(VisualDebugNode->GetFName());
-										Controller->SetNodeSelection(NodeNames);
-									})
-								));
-								Section.AddMenuEntry(
-									"ToggleVisualDebug",
-									LOCTEXT("ToggleVisualDebug", "Toggle Visual Debug"),
-									LOCTEXT("ToggleVisualDebug_Tooltip", "Toggle the visibility the visual debugging"),
-									FSlateIcon(),
-									FUIAction(FExecuteAction::CreateLambda([Controller, VisualDebugNode]() {
-										URigVMPin* EnabledPin = VisualDebugNode->FindPin(TEXT("bEnabled"));
-										check(EnabledPin);
-										Controller->SetPinDefaultValue(EnabledPin->GetPinPath(), EnabledPin->GetDefaultValue() == TEXT("True") ? TEXT("False") : TEXT("True"), false, true, false, true);
-									})
-								));
-								Section.AddMenuEntry(
-									"RemoveVisualDebug",
-									LOCTEXT("RemoveVisualDebug", "Remove Visual Debug"),
-									LOCTEXT("RemoveVisualDebug_Tooltip", "Removes the visual debugging node"),
-									FSlateIcon(),
-										FUIAction(FExecuteAction::CreateLambda([Controller, InGraphPin, ModelPin, VisualDebugNode]() {
-										Controller->RemoveNodeByName(VisualDebugNode->GetFName(), true, false, true);
-									})
-								));
+										})
+									));
+								}
+								else
+								{
+									Section.AddMenuEntry(
+										"EditVisualDebug",
+										LOCTEXT("EditVisualDebug", "Edit Visual Debug"),
+										LOCTEXT("EditVisualDebug_Tooltip", "Edit the visual debugging node"),
+										FSlateIcon(),
+										FUIAction(FExecuteAction::CreateLambda([Controller, VisualDebugNode]() {
+											TArray<FName> NodeNames;
+											NodeNames.Add(VisualDebugNode->GetFName());
+											Controller->SetNodeSelection(NodeNames);
+										})
+									));
+									Section.AddMenuEntry(
+										"ToggleVisualDebug",
+										LOCTEXT("ToggleVisualDebug", "Toggle Visual Debug"),
+										LOCTEXT("ToggleVisualDebug_Tooltip", "Toggle the visibility the visual debugging"),
+										FSlateIcon(),
+										FUIAction(FExecuteAction::CreateLambda([Controller, VisualDebugNode]() {
+											URigVMPin* EnabledPin = VisualDebugNode->FindPin(TEXT("bEnabled"));
+											check(EnabledPin);
+											Controller->SetPinDefaultValue(EnabledPin->GetPinPath(), EnabledPin->GetDefaultValue() == TEXT("True") ? TEXT("False") : TEXT("True"), false, true, false, true);
+										})
+									));
+									Section.AddMenuEntry(
+										"RemoveVisualDebug",
+										LOCTEXT("RemoveVisualDebug", "Remove Visual Debug"),
+										LOCTEXT("RemoveVisualDebug_Tooltip", "Removes the visual debugging node"),
+										FSlateIcon(),
+											FUIAction(FExecuteAction::CreateLambda([Controller, InGraphPin, ModelPin, VisualDebugNode]() {
+											Controller->RemoveNodeByName(VisualDebugNode->GetFName(), true, false, true);
+										})
+									));
+								}
 							}
 						}
 					}
diff --git a/Engine/Plugins/Animation/ControlRig/Source/ControlRigEditor/Private/Graph/ControlRigConnectionDrawingPolicy.cpp b/Engine/Plugins/Animation/ControlRig/Source/ControlRigEditor/Private/Graph/ControlRigConnectionDrawingPolicy.cpp
--- a/Engine/Plugins/Animation/ControlRig/Source/ControlRigEditor/Private/Graph/ControlRigConnectionDrawingPolicy.cpp
+++ b/Engine/Plugins/Animation/ControlRig/Source/ControlRigEditor/Private/Graph/ControlRigConnectionDrawingPolicy.cpp
@@ -286,11 +286,14 @@ void FControlRigConnectionDrawingPolicy::DetermineWiringStyle(UEdGraphPin* Outpu
 			if (InputModelPin)
 			{
 				InputModelPin = InputModelPin->GetPinForLink();
-				if (URigVMInjectionInfo* InputInjection = InputModelPin->GetNode()->GetInjectionInfo())
+				if (InputModelPin)
 				{
-					if (InputModelPin->GetNode()->IsSelected())
+					if (URigVMInjectionInfo* InputInjection = InputModelPin->GetNode()->GetInjectionInfo())
 					{
-						bInjectionIsSelected = true;
+						if (InputModelPin->GetNode()->IsSelected())
+						{
+							bInjectionIsSelected = true;
+						}
 					}
 				}
 			}
diff --git a/Engine/Plugins/Animation/ControlRig/Source/ControlRigEditor/Private/Graph/SControlRigGraphNode.cpp b/Engine/Plugins/Animation/ControlRig/Source/ControlRigEditor/Private/Graph/SControlRigGraphNode.cpp
--- a/Engine/Plugins/Animation/ControlRig/Source/ControlRigEditor/Private/Graph/SControlRigGraphNode.cpp
+++ b/Engine/Plugins/Animation/ControlRig/Source/ControlRigEditor/Private/Graph/SControlRigGraphNode.cpp
@@ -624,7 +624,7 @@ void SControlRigGraphNode::AddPin(const TSharedRef<SGraphPin>& PinToAdd)
 		{
 			if (URigVMPin* ModelPin = ModelNode->FindPin(PinPath))
 			{
-				if (ModelPin->HasInjectedNodes())
+				if (ModelPin->HasInjectedUnitNodes())
 				{
 					PinToAdd->SetCustomPinIcon(CachedImg_CR_Pin_Connected, CachedImg_CR_Pin_Disconnected);
 				}
@@ -1012,46 +1012,47 @@ TArray<FOverlayWidgetInfo> SControlRigGraphNode::GetOverlayWidgets(bool bSelecte
 
 		for (URigVMPin* ModelPin : ModelNode->GetPins())
 		{
-			if (ModelPin->HasInjectedNodes())
+			if (ModelPin->HasInjectedUnitNodes())
 			{
 				for (URigVMInjectionInfo* Injection : ModelPin->GetInjectedNodes())
 				{
-					URigVMUnitNode* VisualDebugNode = Injection->UnitNode;
-
-					FString PrototypeName;
-					if (VisualDebugNode->GetScriptStruct()->GetStringMetaDataHierarchical(TEXT("PrototypeName"), &PrototypeName))
+					if (URigVMUnitNode* VisualDebugNode = Cast<URigVMUnitNode>(Injection->Node))
 					{
-						if (PrototypeName == TEXT("VisualDebug"))
-						{
-							if (!bSetColor)
-							{
-								if (VisualDebugNode->FindPin(TEXT("bEnabled"))->GetDefaultValue() == TEXT("True"))
-								{
-									if (URigVMPin* ColorPin = VisualDebugNode->FindPin(TEXT("Color")))
-									{
-										TBaseStructure<FLinearColor>::Get()->ImportText(*ColorPin->GetDefaultValue(), &Color, nullptr, PPF_None, nullptr, TBaseStructure<FLinearColor>::Get()->GetName());
-									}
-									else
-									{
-										Color = FLinearColor::White;
-									}
-
-									VisualDebugIndicatorWidget->SetColorAndOpacity(Color);
-									bSetColor = true;
-								}
-							}
-
-							if (Widgets.Num() == PreviousNumWidgets)
-							{
-								const FVector2D ImageSize = VisualDebugIndicatorWidget->GetDesiredSize();
-
-								FOverlayWidgetInfo Info;
-								Info.OverlayOffset = FVector2D(WidgetSize.X - ImageSize.X - 6.f, 6.f);
-								Info.Widget = VisualDebugIndicatorWidget;
-
-								Widgets.Add(Info);
-							}
-						}
+						FString PrototypeName;
+					   if (VisualDebugNode->GetScriptStruct()->GetStringMetaDataHierarchical(TEXT("PrototypeName"), &PrototypeName))
+					   {
+						   if (PrototypeName == TEXT("VisualDebug"))
+						   {
+							   if (!bSetColor)
+							   {
+								   if (VisualDebugNode->FindPin(TEXT("bEnabled"))->GetDefaultValue() == TEXT("True"))
+								   {
+									   if (URigVMPin* ColorPin = VisualDebugNode->FindPin(TEXT("Color")))
+									   {
+										   TBaseStructure<FLinearColor>::Get()->ImportText(*ColorPin->GetDefaultValue(), &Color, nullptr, PPF_None, nullptr, TBaseStructure<FLinearColor>::Get()->GetName());
+									   }
+									   else
+									   {
+										   Color = FLinearColor::White;
+									   }
+
+									   VisualDebugIndicatorWidget->SetColorAndOpacity(Color);
+									   bSetColor = true;
+								   }
+							   }
+
+							   if (Widgets.Num() == PreviousNumWidgets)
+							   {
+								   const FVector2D ImageSize = VisualDebugIndicatorWidget->GetDesiredSize();
+
+								   FOverlayWidgetInfo Info;
+								   Info.OverlayOffset = FVector2D(WidgetSize.X - ImageSize.X - 6.f, 6.f);
+								   Info.Widget = VisualDebugIndicatorWidget;
+
+								   Widgets.Add(Info);
+							   }
+						   }
+					   }
 					}
 				}
 			}
diff --git a/Engine/Source/Developer/RigVMDeveloper/Private/RigVMCompiler/RigVMAST.cpp b/Engine/Source/Developer/RigVMDeveloper/Private/RigVMCompiler/RigVMAST.cpp
--- a/Engine/Source/Developer/RigVMDeveloper/Private/RigVMCompiler/RigVMAST.cpp
+++ b/Engine/Source/Developer/RigVMDeveloper/Private/RigVMCompiler/RigVMAST.cpp
@@ -384,12 +384,6 @@ bool FRigVMNodeExprAST::IsConstant() const
 			{
 				return false;
 			}
-
-			URigVMPin::FPinOverride PinOverride(GetProxy(), GetParser()->GetPinOverrides());
-			if (Pin->IsBoundToVariable(PinOverride))
-			{
-				return false;
-			}
 		}
 	}
 	return FRigVMExprAST::IsConstant();
@@ -1014,25 +1008,6 @@ FRigVMExprAST* FRigVMParserAST::TraversePin(const FRigVMASTProxy& InPinProxy, FR
 		);
 	}
 
-	struct Local
-	{
-		static void LookForPinsBoundToVariables(const FRigVMASTProxy& InPinProxy, URigVMPin* InPin, TArray<URigVMPin*>& SubPinsBoundToVariables, const URigVMPin::FPinOverrideMap& InPinOverrides)
-		{
-			for (URigVMPin* SubPin : InPin->GetSubPins())
-			{
-				URigVMPin::FPinOverride PinOverride(InPinProxy, InPinOverrides);
-				if (SubPin->IsBoundToVariable(PinOverride))
-				{
-					SubPinsBoundToVariables.Add(SubPin);
-				}
-				LookForPinsBoundToVariables(InPinProxy, SubPin, SubPinsBoundToVariables, InPinOverrides);
-			}
-		}
-	};
-
-	TArray<URigVMPin*> SubPinsBoundToVariables;
-	Local::LookForPinsBoundToVariables(InPinProxy, Pin, SubPinsBoundToVariables, PinOverrides);
-
 	FRigVMExprAST* PinExpr = nullptr;
 
 	if (Cast<URigVMVariableNode>(Pin->GetNode()))
@@ -1053,18 +1028,9 @@ FRigVMExprAST* FRigVMParserAST::TraversePin(const FRigVMASTProxy& InPinProxy, FR
 
 	if ((Pin->GetDirection() == ERigVMPinDirection::Input ||
 		Pin->GetDirection() == ERigVMPinDirection::Visible) &&
-		Links.Num() == 0 &&
-		SubPinsBoundToVariables.Num() == 0)
+		Links.Num() == 0)
 	{
-		if (Pin->IsBoundToExternalVariable(PinOverride))
-		{
-			PinExpr = MakeExpr<FRigVMExternalVarExprAST>(InPinProxy);
-		}
-		else if (Pin->IsBoundToLocalVariable(PinOverride))
-		{
-			PinExpr = MakeExpr<FRigVMVarExprAST>(FRigVMExprAST::EType::Var, InPinProxy);
-		}
-		else if (Cast<URigVMParameterNode>(Pin->GetNode()) ||
+		if (Cast<URigVMParameterNode>(Pin->GetNode()) ||
 				 Cast<URigVMVariableNode>(Pin->GetNode()))
 		{
 			PinExpr = MakeExpr<FRigVMVarExprAST>(FRigVMExprAST::EType::Var, InPinProxy);
@@ -1119,7 +1085,7 @@ FRigVMExprAST* FRigVMParserAST::TraversePin(const FRigVMASTProxy& InPinProxy, FR
 
 		if (!bHasSourceLinkToRoot && 
 			GetSourcePins(InPinProxy).Num() == 0 &&
-			(Pin->GetDirection() == ERigVMPinDirection::IO || Links.Num() > 0 || SubPinsBoundToVariables.Num() > 0))
+			(Pin->GetDirection() == ERigVMPinDirection::IO || Links.Num() > 0))
 		{
 			FRigVMLiteralExprAST* LiteralExpr = MakeExpr<FRigVMLiteralExprAST>(InPinProxy);
 			FRigVMCopyExprAST* LiteralCopyExpr = MakeExpr<FRigVMCopyExprAST>(InPinProxy, InPinProxy);
@@ -1130,14 +1096,6 @@ FRigVMExprAST* FRigVMParserAST::TraversePin(const FRigVMASTProxy& InPinProxy, FR
 
 			SubjectToExpression[InPinProxy] = LiteralExpr;
 		}
-		else
-		{
-			SubPinsBoundToVariables.Reset();
-		}
-	}
-	else
-	{
-		SubPinsBoundToVariables.Reset();
 	}
 
 	FRigVMExprAST* ParentExprForLinks = PinExpr;
@@ -1157,17 +1115,6 @@ FRigVMExprAST* FRigVMParserAST::TraversePin(const FRigVMASTProxy& InPinProxy, FR
 		TraverseLink(SourceLink, ParentExprForLinks);
 	}
 
-	for (URigVMPin* SubPinBoundToVariable : SubPinsBoundToVariables)
-	{
-		FRigVMASTProxy SubPinBoundToVariableProxy = InPinProxy.GetSibling(SubPinBoundToVariable);
-		FRigVMExternalVarExprAST* ExternalVarExpr = MakeExpr<FRigVMExternalVarExprAST>(SubPinBoundToVariableProxy);
-		FRigVMCopyExprAST* ExternalVarExprCopyExpr = MakeExpr<FRigVMCopyExprAST>(SubPinBoundToVariableProxy, SubPinBoundToVariableProxy);
-		ExternalVarExprCopyExpr->Name = *FString::Printf(TEXT("%s -> %s"), *SubPinBoundToVariable->GetPinPath(), *SubPinBoundToVariable->GetPinPath());
-		ExternalVarExprCopyExpr->AddParent(PinExpr);
-		ExternalVarExpr->AddParent(ExternalVarExprCopyExpr);
-		ExternalVarExpr->Name = *SubPinBoundToVariable->GetPinPath();
-	}
-
 	return PinExpr;
 }
 
@@ -3017,7 +2964,7 @@ void FRigVMParserAST::Inline(URigVMGraph* InGraph, const TArray<FRigVMASTProxy>&
 				}
 
 				TArray<URigVMLink*> SourceLinks = ChildPin->GetSourceLinks(false /* recursive */);
-				if (SourceLinks.Num() > 0 || ChildPin->IsBoundToInputArgument())
+				if (SourceLinks.Num() > 0)
 				{
 					URigVMPin* SourcePin = nullptr;
 					if (SourceLinks.Num() > 0)
diff --git a/Engine/Source/Developer/RigVMDeveloper/Private/RigVMCompiler/RigVMCompiler.cpp b/Engine/Source/Developer/RigVMDeveloper/Private/RigVMCompiler/RigVMCompiler.cpp
--- a/Engine/Source/Developer/RigVMDeveloper/Private/RigVMCompiler/RigVMCompiler.cpp
+++ b/Engine/Source/Developer/RigVMDeveloper/Private/RigVMCompiler/RigVMCompiler.cpp
@@ -384,26 +384,6 @@ bool URigVMCompiler::Compile(URigVMGraph* InGraph, URigVMController* InControlle
 						}
 					}
 				}
-				else
-				{
-					bool bAdded = false;
-					for (URigVMPin* Pin : Node->GetPins())
-					{
-						if (Pin->IsBoundToLocalVariable())
-						{
-							if (Pin->GetBoundVariableName() == LocalVariable.Name.ToString())
-							{
-								AddDefaultValueOperand(Pin);
-								bAdded = true;
-								break;
-							}
-						}
-					}
-					if (bAdded)
-					{
-						break;
-					}
-				}
 			}
 		}
 	}	
@@ -1064,12 +1044,6 @@ void URigVMCompiler::TraverseAssign(const FRigVMAssignExprAST* InExpr, FRigVMCom
 						return;
 					}
 
-					URigVMPin::FPinOverride PinOverride(VarExpr->GetProxy(), VarExpr->GetParser()->GetPinOverrides());
-					if (bSource && Pin->IsBoundToVariable(PinOverride))
-					{
-						return;
-					}
-
 					if (URigVMSelectNode* SelectNode = Cast<URigVMSelectNode>(RootPin->GetNode()))
 					{
 						if (Pin->GetParentPin() == RootPin && RootPin->GetName() == URigVMSelectNode::ValueName)
@@ -1948,24 +1922,6 @@ FString URigVMCompiler::GetPinHash(const URigVMPin* InPin, const FRigVMVarExprAS
 			return FString::Printf(TEXT("%sVariable::%s%s"), *Prefix, *VariablePath, *Suffix);
 		}
 
-		if (InPin->IsBoundToLocalVariable())
-		{
-			URigVMPin::FPinOverride PinOverride(InVarExpr->GetProxy(), InVarExpr->GetParser()->GetPinOverrides());
-			FString VariablePath = InPin->GetBoundVariablePath(PinOverride);
-
-			FRigVMASTProxy ParentProxy = InVarExpr->GetProxy();
-			while(ParentProxy.GetCallstack().Num() > 1)
-			{
-				ParentProxy = ParentProxy.GetParent();
-
-				if(URigVMLibraryNode* LibraryNode = ParentProxy.GetSubject<URigVMLibraryNode>())
-				{
-					// Local variables are in the format "LocalVariable::PathToGraph|VariableName"
-					return FString::Printf(TEXT("%sLocalVariable::%s|%s%s"), *Prefix, *LibraryNode->GetNodePath(true), *VariablePath, *Suffix);
-				}
-			}
-		}
-
 		// for IO array pins we'll walk left and use that pin hash instead
 		if(const FRigVMVarExprAST* SourceVarExpr = GetSourceVarExpr(InVarExpr))
 		{
@@ -2495,49 +2451,6 @@ FRigVMOperand URigVMCompiler::FindOrAddRegister(const FRigVMVarExprAST* InVarExp
 		}
 	}
 
-	// check if this is a variable with a segment path
-	if (!Operand.IsValid())
-	{
-		if (Pin->IsBoundToVariable(PinOverride) && !bIsDebugValue)
-		{
-			FString VariablePath = Pin->GetBoundVariablePath(PinOverride);
-			FString VariableName = VariablePath, SegmentPath;
-			VariablePath.Split(TEXT("."), &VariableName, &SegmentPath);
-
-			ExistingOperand = WorkData.PinPathToOperand->Find(FString::Printf(TEXT("LocalVariable::%s"), *VariableName));
-			if (ExistingOperand)
-			{
-				Operand = *ExistingOperand;
-				Operand.RegisterOffset = INDEX_NONE;
-				check(SegmentPath.IsEmpty());
-			}
-			else
-			{
-				ExistingOperand = WorkData.PinPathToOperand->Find(FString::Printf(TEXT("Variable::%s"), *VariableName));
-				if (ExistingOperand)
-				{
-					Operand = *ExistingOperand;
-					Operand.RegisterOffset = INDEX_NONE;
-					if (!SegmentPath.IsEmpty())
-					{
-#if UE_RIGVM_UCLASS_BASED_STORAGE_DISABLED
-						
-						const FRigVMExternalVariable& ExternalVariable = WorkData.VM->GetExternalVariables()[Operand.GetRegisterIndex()];
-						UScriptStruct* ScriptStruct = CastChecked<UScriptStruct>(ExternalVariable.TypeObject);
-						Operand.RegisterOffset = WorkData.VM->GetWorkMemory().GetOrAddRegisterOffset(Operand.GetRegisterIndex(), ScriptStruct, SegmentPath, 0 /*ArrayIndex */);
-
-#else
-
-						const int32 PropertyPathIndex = WorkData.FindOrAddPropertyPath(Operand, Pin->GetCPPType(), SegmentPath);
-						Operand.RegisterOffset = (uint16)PropertyPathIndex;
-						
-#endif
-					}
-				}
-			}
-		}
-	}
-
 	// create remaining operands / registers
 	if (!Operand.IsValid())
 	{
diff --git a/Engine/Source/Developer/RigVMDeveloper/Private/RigVMModel/RigVMController.cpp b/Engine/Source/Developer/RigVMDeveloper/Private/RigVMModel/RigVMController.cpp
--- a/Engine/Source/Developer/RigVMDeveloper/Private/RigVMModel/RigVMController.cpp
+++ b/Engine/Source/Developer/RigVMDeveloper/Private/RigVMModel/RigVMController.cpp
@@ -531,8 +531,8 @@ TArray<FString> URigVMController::GetAddNodePythonCommands(URigVMNode* Node) con
 		{
 			const URigVMPin* InjectionInfoPin = InjectionInfo->GetPin();
 			const FString InjectionInfoPinPath = GetSanitizedPinPath(InjectionInfoPin->GetPinPath());
-			const FString InjectionInfoInputPinName = GetSanitizedPinName(InjectionInfo->InputPin->GetName());
-			const FString InjectionInfoOutputPinName = GetSanitizedPinName(InjectionInfo->OutputPin->GetName());
+			const FString InjectionInfoInputPinName = InjectionInfo->InputPin ? GetSanitizedPinName(InjectionInfo->InputPin->GetName()) : FString();
+			const FString InjectionInfoOutputPinName = InjectionInfo->OutputPin ? GetSanitizedPinName(InjectionInfo->OutputPin->GetName()) : FString();
 
 			//URigVMInjectionInfo* AddInjectedNodeFromStructPath(const FString& InPinPath, bool bAsInput, const FString& InScriptStructPath, const FName& InMethodName, const FName& InInputPinName, const FName& InOutputPinName, const FString& InNodeName = TEXT(""), bool bSetupUndoRedo = true);
 			Commands.Add(FString::Printf(TEXT("%s_info = blueprint.get_controller_by_name('%s').add_injected_node_from_struct_path('%s', %s, '%s', '%s', '%s', '%s', '%s')"),
@@ -1285,15 +1285,6 @@ void URigVMController::OnExternalVariableRemoved(const FName& InVarName, bool bS
 				}
 			}
 		}
-
-		TArray<URigVMPin*> AllPins = Node->GetAllPinsRecursively();
-		for (URigVMPin* Pin : AllPins)
-		{
-			if (Pin->GetBoundVariableName() == InVarName.ToString())
-			{
-				BindPinToVariable(Pin, FString(), bSetupUndoRedo);
-			}
-		}
 	}
 
 	if (bSetupUndoRedo)
@@ -1382,17 +1373,6 @@ void URigVMController::OnExternalVariableRenamed(const FName& InOldVarName, cons
 				}
 			}
 		}
-
-		TArray<URigVMPin*> AllPins = Node->GetAllPinsRecursively();
-		for (URigVMPin* Pin : AllPins)
-		{
-			if (Pin->GetBoundVariableName() == InOldVarName.ToString())
-			{
-				FString OldVariablePath = Pin->GetBoundVariablePath();
-				FString NewVariablePath = OldVariablePath.Replace(*InOldVarName.ToString(), *InNewVarName.ToString());
-				BindPinToVariable(Pin, NewVariablePath, bSetupUndoRedo);
-			}
-		}
 	}
 
 	if (bSetupUndoRedo)
@@ -1485,7 +1465,7 @@ void URigVMController::OnExternalVariableTypeChanged(const FName& InVarName, con
 			if (Pin->GetBoundVariableName() == InVarName.ToString())
 			{
 				FString BoundVariablePath = Pin->GetBoundVariablePath();
-				BindPinToVariable(Pin, FString(), bSetupUndoRedo);
+				UnbindPinFromVariable(Pin, bSetupUndoRedo);
 				// try to bind it again - maybe it can be bound (due to cast rules etc)
 				BindPinToVariable(Pin, BoundVariablePath, bSetupUndoRedo);
 			}
@@ -2135,9 +2115,9 @@ URigVMInjectionInfo* URigVMController::AddInjectedNode(const FString& InPinPath,
 	URigVMInjectionInfo* InjectionInfo = NewObject<URigVMInjectionInfo>(Pin);
 
 	// re-parent the unit node to be under the injection info
-	UnitNode->Rename(nullptr, InjectionInfo);
+	RenameObject(UnitNode, nullptr, InjectionInfo);
 
-	InjectionInfo->UnitNode = UnitNode;
+	InjectionInfo->Node = UnitNode;
 	InjectionInfo->bInjectedAsInput = bAsInput;
 	InjectionInfo->InputPin = InputPin;
 	InjectionInfo->OutputPin = OutputPin;
@@ -2258,14 +2238,10 @@ URigVMNode* URigVMController::EjectNodeFromPin(const FString& InPinPath, bool bS
 
 	URigVMInjectionInfo* Injection = Pin->InjectionInfos.Last();
 
-	UScriptStruct* ScriptStruct = Injection->UnitNode->GetScriptStruct();
-	FName UnitNodeName = Injection->UnitNode->GetFName();
-	FName MethodName = Injection->UnitNode->GetMethodName();
-	FName InputPinName = Injection->InputPin->GetFName();
-	FName OutputPinName = Injection->OutputPin->GetFName();
-
+	FName NodeName = Injection->Node->GetFName();
+	
 	TMap<FName, FString> DefaultValues;
-	for (URigVMPin* PinOnNode : Injection->UnitNode->GetPins())
+	for (URigVMPin* PinOnNode : Injection->Node->GetPins())
 	{
 		if (PinOnNode->GetDirection() == ERigVMPinDirection::Input ||
 			PinOnNode->GetDirection() == ERigVMPinDirection::Visible ||
@@ -2295,7 +2271,15 @@ URigVMNode* URigVMController::EjectNodeFromPin(const FString& InPinPath, bool bS
 		Position -= FVector2D(250.f, 0.f);
 	}
 
-	URigVMNode* EjectedNode = AddUnitNode(ScriptStruct, MethodName, Position, FString(), bSetupUndoRedo);
+	URigVMNode* EjectedNode = nullptr;
+	if (URigVMUnitNode* UnitNode = Cast<URigVMUnitNode>(Injection->Node))
+	{
+		EjectedNode = AddUnitNode(UnitNode->GetScriptStruct(), UnitNode->GetMethodName(), Position, FString(), bSetupUndoRedo);
+	}
+	else if (URigVMVariableNode* VariableNode = Cast<URigVMVariableNode>(Injection->Node))
+	{
+		EjectedNode = AddVariableNode(VariableNode->GetVariableName(), VariableNode->GetCPPType(), VariableNode->GetCPPTypeObject(), true, VariableNode->GetDefaultValue(), Position, FString(), bSetupUndoRedo);
+	}
 
 	for (const TPair<FName, FString>& Pair : DefaultValues)
 	{
@@ -2309,17 +2293,24 @@ URigVMNode* URigVMController::EjectNodeFromPin(const FString& InPinPath, bool bS
 		}
 	}
 
-	TArray<URigVMLink*> PreviousLinks = Injection->InputPin->GetSourceLinks(true);
-	PreviousLinks.Append(Injection->OutputPin->GetTargetLinks(true));
+	TArray<URigVMLink*> PreviousLinks;
+	if (Injection->InputPin)
+	{
+		Injection->InputPin->GetSourceLinks(true);
+	}
+	if (Injection->OutputPin)
+	{
+		PreviousLinks.Append(Injection->OutputPin->GetTargetLinks(true));
+	}
 	for (URigVMLink* PreviousLink : PreviousLinks)
 	{
 		PreviousLink->PrepareForCopy();
 		PreviousLink->SourcePin = PreviousLink->TargetPin = nullptr;
 	}
 
-	RemoveNode(Injection->UnitNode, bSetupUndoRedo);
+	RemoveNode(Injection->Node, bSetupUndoRedo);
 
-	FString OldNodeNamePrefix = UnitNodeName.ToString() + TEXT(".");
+	FString OldNodeNamePrefix = NodeName.ToString() + TEXT(".");
 	FString NewNodeNamePrefix = EjectedNode->GetName() + TEXT(".");
 
 	for (URigVMLink* PreviousLink : PreviousLinks)
@@ -2436,7 +2427,7 @@ FString URigVMController::ExportNodesToText(const TArray<FName>& InNodeNames)
 			{
 				for (URigVMInjectionInfo* Injection : Pin->GetInjectedNodes())
 				{
-					AllNodeNames.AddUnique(Injection->UnitNode->GetFName());
+					AllNodeNames.AddUnique(Injection->Node->GetFName());
 				}
 			}
 		}
@@ -2536,13 +2527,13 @@ protected:
 			{
 				for (URigVMInjectionInfo* Injection : Pin->GetInjectedNodes())
 				{
-					ProcessConstructedObject(Injection->UnitNode);
+					ProcessConstructedObject(Injection->Node);
 
-					FName NewName = Injection->UnitNode->GetFName();
+					FName NewName = Injection->Node->GetFName();
 					UpdateObjectName(URigVMNode::StaticClass(), NewName);
-					Injection->UnitNode->Rename(*NewName.ToString(), nullptr);
-					Injection->InputPin = Injection->UnitNode->FindPin(Injection->InputPin->GetName());
-					Injection->OutputPin = Injection->UnitNode->FindPin(Injection->OutputPin->GetName());
+					Controller->RenameObject(Injection->Node, *NewName.ToString(), nullptr);
+					Injection->InputPin = Injection->InputPin ? Injection->Node->FindPin(Injection->InputPin->GetName()) : nullptr;
+					Injection->OutputPin = Injection->OutputPin ? Injection->Node->FindPin(Injection->OutputPin->GetName()) : nullptr;
 				}
 			}
 		}
@@ -4818,13 +4809,19 @@ bool URigVMController::RemoveNode(URigVMNode* InNode, bool bSetupUndoRedo, bool
 
 		if (InjectionInfo->bInjectedAsInput)
 		{
-			URigVMPin* LastInputPin = Pin;
-			RewireLinks(InjectionInfo->InputPin, LastInputPin, true, false);
+			if (InjectionInfo->InputPin)
+			{
+				URigVMPin* LastInputPin = Pin;
+				RewireLinks(InjectionInfo->InputPin, LastInputPin, true, false);
+			}
 		}
 		else
 		{
-			URigVMPin* LastOutputPin = Pin;
-			RewireLinks(InjectionInfo->OutputPin, LastOutputPin, false, false);
+			if (InjectionInfo->OutputPin)
+			{
+				URigVMPin* LastOutputPin = Pin;
+				RewireLinks(InjectionInfo->OutputPin, LastOutputPin, false, false);
+			}
 		}
 	}
 
@@ -4911,7 +4908,7 @@ bool URigVMController::RemoveNode(URigVMNode* InNode, bool bSetupUndoRedo, bool
 			TArray<URigVMInjectionInfo*> InjectedNodes = Pin->GetInjectedNodes();
 			for (URigVMInjectionInfo* InjectedNode : InjectedNodes)
 			{
-				RemoveNode(InjectedNode->UnitNode, bSetupUndoRedo, bRecursive);
+				RemoveNode(InjectedNode->Node, bSetupUndoRedo, bRecursive);
 			}
 
 			BreakAllLinks(Pin, true, bSetupUndoRedo);
@@ -5044,7 +5041,7 @@ bool URigVMController::RenameNode(URigVMNode* InNode, const FName& InNewName, bo
 	}
 
 	InNode->PreviousName = InNode->GetFName();
-	if (!InNode->Rename(*ValidNewName.ToString()))
+	if (!RenameObject(InNode, *ValidNewName.ToString()))
 	{
 		ActionStack->CancelAction(Action);
 		return false;
@@ -6538,7 +6535,7 @@ URigVMPin* URigVMController::InsertArrayPin(URigVMPin* ArrayPin, int32 InIndex,
 	for (int32 ExistingIndex = ArrayPin->GetSubPins().Num() - 1; ExistingIndex >= InIndex; ExistingIndex--)
 	{
 		URigVMPin* ExistingPin = ArrayPin->GetSubPins()[ExistingIndex];
-		ExistingPin->Rename(*FString::FormatAsNumber(ExistingIndex + 1));
+		RenameObject(ExistingPin, *FString::FormatAsNumber(ExistingIndex + 1));
 	}
 
 	URigVMPin* Pin = NewObject<URigVMPin>(ArrayPin, *FString::FormatAsNumber(InIndex));
@@ -6805,7 +6802,16 @@ bool URigVMController::BindPinToVariable(const FString& InPinPath, const FString
 		return false;
 	}
 
-	const bool bSuccess = BindPinToVariable(Pin, InNewBoundVariablePath, bSetupUndoRedo);
+	bool bSuccess = false;
+	if (InNewBoundVariablePath.IsEmpty())
+	{
+		bSuccess = UnbindPinFromVariable(Pin, bSetupUndoRedo);
+	}
+	else
+	{
+		bSuccess = BindPinToVariable(Pin, InNewBoundVariablePath, bSetupUndoRedo);
+	}
+	
 	if (bSuccess && bPrintPythonCommand)
 	{
 		const FString GraphName = GetSanitizedGraphName(GetGraph()->GetGraphName());
@@ -6820,11 +6826,6 @@ bool URigVMController::BindPinToVariable(const FString& InPinPath, const FString
 	return bSuccess;
 }
 
-bool URigVMController::UnbindPinFromVariable(const FString& InPinPath, bool bSetupUndoRedo, bool bPrintPythonCommand)
-{
-	return BindPinToVariable(InPinPath, FString(), bSetupUndoRedo, bPrintPythonCommand);
-}
-
 bool URigVMController::BindPinToVariable(URigVMPin* InPin, const FString& InNewBoundVariablePath, bool bSetupUndoRedo)
 {
 	if (!IsValidPinForGraph(InPin))
@@ -6832,79 +6833,207 @@ bool URigVMController::BindPinToVariable(URigVMPin* InPin, const FString& InNewB
 		return false;
 	}
 
+	if (GetGraph()->IsA<URigVMFunctionLibrary>())
+	{
+		ReportError(TEXT("Cannot bind pins to variables in function library graphs."));
+		return false;
+	}
+
 	if (InPin->GetBoundVariablePath() == InNewBoundVariablePath)
 	{
 		return false;
 	}
 
-	if (GetGraph()->IsA<URigVMFunctionLibrary>())
+	if (InPin->GetDirection() != ERigVMPinDirection::Input)
 	{
-		ReportError(TEXT("Cannot bind pins to variables in function library graphs."));
 		return false;
 	}
 
-	// check that the variable is compatible
-	if (!InNewBoundVariablePath.IsEmpty())
+	FString VariableName = InNewBoundVariablePath, SegmentPath;
+	InNewBoundVariablePath.Split(TEXT("."), &VariableName, &SegmentPath);
+
+	FRigVMExternalVariable Variable;
+	for (const FRigVMExternalVariable& VariableDescription : GetAllVariables(true))
 	{
-		FString VariableName = InNewBoundVariablePath, SegmentPath;
-		InNewBoundVariablePath.Split(TEXT("."), &VariableName, &SegmentPath);
+		if (VariableDescription.Name.ToString() == VariableName)
+		{
+			Variable = VariableDescription;
+			break;
+		}
+	}
 
-		FRigVMExternalVariable Variable = GetVariableByName(*VariableName, true);
-		if (Variable.IsValid(true))
+	if (!Variable.Name.IsValid())
+	{
+		ReportError(TEXT("Cannot find variable in this graph."));
+		return false;
+	}
+
+	
+	if (!RigVMTypeUtils::AreCompatible(Variable, InPin->ToExternalVariable(), SegmentPath))
+	{
+		ReportError(TEXT("Cannot find variable in this graph."));
+		return false;
+	}
+	
+	FRigVMControllerCompileBracketScope CompileScope(this);
+	FRigVMBaseAction Action;
+	if (bSetupUndoRedo)
+	{
+		Action.Title = TEXT("Bind pin to variable");
+		ActionStack->BeginAction(Action);
+	}
+
+	{
+		if (InPin->IsBoundToVariable())
 		{
-#if UE_RIGVM_UCLASS_BASED_STORAGE_DISABLED
-			FRigVMRegisterOffset RegisterOffset;
-			if (!SegmentPath.IsEmpty())
-			{
-				RegisterOffset = FRigVMRegisterOffset(Cast<UScriptStruct>(ExternalVariable.TypeObject), SegmentPath);
-			}
-			if (!InPin->CanBeBoundToVariable(ExternalVariable, RegisterOffset))
-			{
-				return false;
-			}
-#else
-			if (!InPin->CanBeBoundToVariable(Variable, SegmentPath))
-			{
-				return false;
-			}
-#endif
+			UnbindPinFromVariable(InPin, bSetupUndoRedo);
 		}
-		else
+		TArray<URigVMInjectionInfo*> Infos = InPin->GetInjectedNodes();
+		for (URigVMInjectionInfo* Info : Infos)
 		{
-			return false;
+			RemoveNode(Info->Node, bSetupUndoRedo);
 		}
+		BreakAllLinks(InPin, true, bSetupUndoRedo);
 	}
 
-	FRigVMControllerCompileBracketScope CompileScope(this);
-	FRigVMBaseAction Action;
 	if (bSetupUndoRedo)
 	{
-		if (InNewBoundVariablePath.IsEmpty())
+		ActionStack->AddAction(FRigVMSetPinBoundVariableAction(InPin, InNewBoundVariablePath));
+	}		
+	
+	// Create variable node
+	URigVMVariableNode* VariableNode = nullptr;
+	{
 		{
-			Action.Title = TEXT("Unbind pin from variable");
+			TGuardValue<bool> GuardNotifications(bSuspendNotifications, true);
+			FString CPPType;
+			UObject* CPPTypeObject;
+			RigVMTypeUtils::CPPTypeFromExternalVariable(Variable, CPPType, &CPPTypeObject);
+			VariableNode = AddVariableNode(*VariableName, CPPType, CPPTypeObject, true, FString(), FVector2D::ZeroVector, FString(), false);
 		}
-		else
+		if (VariableNode == nullptr)
 		{
-			Action.Title = TEXT("Bind pin to variable");
+			if (bSetupUndoRedo)
+			{
+				ActionStack->CancelAction(Action);
+			}
+			return false;
 		}
-		ActionStack->BeginAction(Action);
+	}
+	
+	URigVMPin* ValuePin = VariableNode->FindPin(URigVMVariableNode::ValueName);
+	if (!SegmentPath.IsEmpty())
+	{
+		ValuePin = ValuePin->FindSubPin(SegmentPath);
 	}
 
-	if (!InPin->IsBoundToVariable() && bSetupUndoRedo)
+	// Add injection 
+	URigVMInjectionInfo* InjectionInfo = NewObject<URigVMInjectionInfo>(InPin);
 	{
-		// break all links on pin towards parent + children
-		BreakAllLinks(InPin, true, bSetupUndoRedo);
-		BreakAllLinksRecursive(InPin, true, true, bSetupUndoRedo);
-		BreakAllLinksRecursive(InPin, true, false, bSetupUndoRedo);
+		// re-parent the unit node to be under the injection info
+		RenameObject(VariableNode, nullptr, InjectionInfo);
+		
+		InjectionInfo->Node = VariableNode;
+		InjectionInfo->bInjectedAsInput = true;
+		InjectionInfo->InputPin = nullptr;
+		InjectionInfo->OutputPin = ValuePin;
+	
+		InPin->InjectionInfos.Add(InjectionInfo);
+		Notify(ERigVMGraphNotifType::NodeAdded, VariableNode);
 	}
 
+	{		
+		if (!AddLink(ValuePin, InPin, false))
+		{
+			if (bSetupUndoRedo)
+			{
+				ActionStack->CancelAction(Action);
+			}
+			return false;
+		}
+	}
+	
 	if (bSetupUndoRedo)
 	{
-		ActionStack->AddAction(FRigVMSetPinBoundVariableAction(InPin, InNewBoundVariablePath));
+		ActionStack->EndAction(Action);
+	}
+
+	return true;
+}
+
+bool URigVMController::UnbindPinFromVariable(const FString& InPinPath, bool bSetupUndoRedo, bool bPrintPythonCommand)
+{
+	if (!IsValidGraph())
+	{
+		return false;
+	}
+
+	URigVMGraph* Graph = GetGraph();
+	check(Graph);
+
+	URigVMPin* Pin = Graph->FindPin(InPinPath);
+	if (Pin == nullptr)
+	{
+		ReportErrorf(TEXT("Cannot find pin '%s'."), *InPinPath);
+		return false;
+	}
+
+	const bool bSuccess = UnbindPinFromVariable(Pin, bSetupUndoRedo);
+	if (bSuccess && bPrintPythonCommand)
+	{
+		const FString GraphName = GetSanitizedGraphName(GetGraph()->GetGraphName());
+		
+		RigVMPythonUtils::Print(GetGraphOuterName(),
+			FString::Printf(TEXT("blueprint.get_controller_by_name('%s').unbind_pin_from_variable('%s')"),
+			*GraphName,
+			*GetSanitizedPinPath(InPinPath)));
+	}
+	
+	return bSuccess;
+}
+
+bool URigVMController::UnbindPinFromVariable(URigVMPin* InPin, bool bSetupUndoRedo)
+{
+	if (!IsValidPinForGraph(InPin))
+	{
+		return false;
 	}
 
-	InPin->BoundVariablePath = InNewBoundVariablePath;
-	Notify(ERigVMGraphNotifType::PinBoundVariableChanged, InPin);
+
+	if (GetGraph()->IsA<URigVMFunctionLibrary>())
+	{
+		ReportError(TEXT("Cannot unbind pins from variables in function library graphs."));
+		return false;
+	}
+
+	if (!InPin->IsBoundToVariable())
+	{
+		ReportError(TEXT("Pin is not bound to any variable."));
+		return false;
+	}
+
+	FRigVMControllerCompileBracketScope CompileScope(this);
+	FRigVMBaseAction Action;
+	if (bSetupUndoRedo)
+	{
+		Action.Title = TEXT("Unbind pin from variable");
+		ActionStack->BeginAction(Action);
+	}
+
+	if (bSetupUndoRedo)
+	{
+		ActionStack->AddAction(FRigVMSetPinBoundVariableAction(InPin, FString()));
+	}
+
+	URigVMInjectionInfo* Info = nullptr;
+	for (URigVMInjectionInfo* InjectionInfo : InPin->GetInjectedNodes())
+	{
+		if (InjectionInfo->Node->IsA<URigVMVariableNode>())
+		{
+			BreakAllLinks(InPin, true, false);
+			RemoveNode(InjectionInfo->Node, false);
+		}
+	}	
 
 	if (bSetupUndoRedo)
 	{
@@ -7217,12 +7346,6 @@ bool URigVMController::AddLink(URigVMPin* OutputPin, URigVMPin* InputPin, bool b
 	}
 
 	{
-		// Temporarily remove the bound variable from the pin
-		// if it exists, so that link validation can work.
-		// BreakAllPins will remove the bound variable for real later.
-		TGuardValue<FString> OutputBoundVariableGuard(OutputPin->BoundVariablePath, FString());
-		TGuardValue<FString> InputBoundVariableGuard(InputPin->BoundVariablePath, FString());
-
 		FString FailureReason;
 		if (!Graph->CanLink(OutputPin, InputPin, &FailureReason, GetCurrentByteCode()))
 		{
@@ -7235,10 +7358,9 @@ bool URigVMController::AddLink(URigVMPin* OutputPin, URigVMPin* InputPin, bool b
 	ensure(!InputPin->IsLinkedTo(OutputPin));
 
 	FRigVMControllerCompileBracketScope CompileScope(this);
-	FRigVMAddLinkAction Action;
+	FRigVMBaseAction Action;
 	if (bSetupUndoRedo)
 	{
-		Action = FRigVMAddLinkAction(OutputPin, InputPin);
 		Action.Title = FString::Printf(TEXT("Add Link"));
 		ActionStack->BeginAction(Action);
 	}
@@ -7274,8 +7396,13 @@ bool URigVMController::AddLink(URigVMPin* OutputPin, URigVMPin* InputPin, bool b
 
 	if (bSetupUndoRedo)
 	{
-		ExpandPinRecursively(OutputPin->GetParentPin(), true);
-		ExpandPinRecursively(InputPin->GetParentPin(), true);
+		ExpandPinRecursively(OutputPin->GetParentPin(), bSetupUndoRedo);
+		ExpandPinRecursively(InputPin->GetParentPin(), bSetupUndoRedo);
+	}
+
+	if (bSetupUndoRedo)
+	{
+		ActionStack->AddAction(FRigVMAddLinkAction(OutputPin, InputPin));
 	}
 
 	URigVMLink* Link = NewObject<URigVMLink>(Graph);
@@ -7518,7 +7645,7 @@ bool URigVMController::BreakAllLinks(URigVMPin* Pin, bool bAsInput, bool bSetupU
 	int32 LinksBroken = 0;
 	if (Pin->IsBoundToVariable() && bAsInput && bSetupUndoRedo)
 	{
-		BindPinToVariable(Pin, FString(), bSetupUndoRedo);
+		UnbindPinFromVariable(Pin, bSetupUndoRedo);
 		LinksBroken++;
 	}
 
@@ -7923,7 +8050,7 @@ bool URigVMController::RenameExposedPin(const FName& InOldPinName, const FName&
 				InController->Notify(ERigVMGraphNotifType::LinkRemoved, Link);
 			}
 
-			if (!InPin->Rename(*InNewName.ToString()))
+			if (!InController->RenameObject(InPin, *InNewName.ToString()))
 			{
 				return false;
 			}
@@ -8743,15 +8870,6 @@ bool URigVMController::RemoveLocalVariable(const FName& InVariableName, bool bSe
 						}
 					}
 				}
-
-				TArray<URigVMPin*> AllPins = Node->GetAllPinsRecursively();
-				for (URigVMPin* Pin : AllPins)
-				{
-					if (Pin->GetBoundVariableName() == InVariableName.ToString())
-					{
-						BindPinToVariable(Pin, FString(), bSetupUndoRedo);
-					}
-				}
 			}
 		}
 		else
@@ -8778,7 +8896,7 @@ bool URigVMController::RemoveLocalVariable(const FName& InVariableName, bool bSe
 					{
 						if (Pin->GetCPPType() != ExternalVariableToSwitch.TypeName.ToString() || Pin->GetCPPTypeObject() == ExternalVariableToSwitch.TypeObject)
 						{
-							BindPinToVariable(Pin, FString(), bSetupUndoRedo);
+							UnbindPinFromVariable(Pin, bSetupUndoRedo);
 						}
 					}
 				}
@@ -8879,17 +8997,6 @@ bool URigVMController::RenameLocalVariable(const FName& InVariableName, const FN
 				RenamedNodes.Add(Node);
 			}
 		}
-		
-		const TArray<URigVMPin*> AllPins = Node->GetAllPinsRecursively();
-		for (URigVMPin* Pin : AllPins)
-		{
-			if (Pin->GetBoundVariableName() == InVariableName.ToString())
-			{
-				FString OldVariablePath = Pin->GetBoundVariablePath();
-				FString NewVariablePath = OldVariablePath.Replace(*InVariableName.ToString(), *InNewVariableName.ToString());
-				BindPinToVariable(Pin, NewVariablePath, bSetupUndoRedo);
-			}
-		}
 	}
 
 	for (URigVMNode* RenamedNode : RenamedNodes)
@@ -8992,7 +9099,7 @@ bool URigVMController::SetLocalVariableType(const FName& InVariableName, const F
 		{
 			if (Pin->GetBoundVariableName() == InVariableName.ToString())
 			{
-				BindPinToVariable(Pin, FString(), bSetupUndoRedo);
+				UnbindPinFromVariable(Pin, bSetupUndoRedo);
 			}
 		}
 	}
@@ -11416,8 +11523,8 @@ void URigVMController::RepopulatePinsOnNode(URigVMNode* InNode, bool bFollowCore
 	FName InjectionOutputPinName = NAME_None;
 	if (URigVMInjectionInfo* InjectionInfo = InNode->GetInjectionInfo())
 	{
-		InjectionInputPinName = InjectionInfo->InputPin->GetFName();
-		InjectionOutputPinName = InjectionInfo->OutputPin->GetFName();
+		InjectionInputPinName = InjectionInfo->InputPin ? InjectionInfo->InputPin->GetFName() : NAME_None;
+		InjectionOutputPinName = InjectionInfo->OutputPin ? InjectionInfo->OutputPin->GetFName() : NAME_None;
 	}
 
 	// step 2/3: clear pins on the node and repopulate the node with new pins
@@ -11645,7 +11752,7 @@ void URigVMController::RemovePinsDuringRepopulate(URigVMNode* InNode, TArray<URi
 				if(Pin->IsRootPin()) // if we are passing root pins we can reparent them directly
 				{
 					RootPin->DisplayName = RootPin->GetFName();
-					RootPin->Rename(*OrphanedName, nullptr, REN_ForceNoResetLoaders | REN_DoNotDirty | REN_DontCreateRedirectors | REN_NonTransactional);
+					RenameObject(RootPin, *OrphanedName, nullptr);
 					InNode->Pins.Remove(RootPin);
 
 					if(bNotify)
@@ -11677,7 +11784,7 @@ void URigVMController::RemovePinsDuringRepopulate(URigVMNode* InNode, TArray<URi
 
 			if(!Pin->IsRootPin() && (OrphanedRootPin != nullptr))
 			{
-				Pin->Rename(nullptr, OrphanedRootPin, REN_ForceNoResetLoaders | REN_DoNotDirty | REN_DontCreateRedirectors | REN_NonTransactional);
+				RenameObject(Pin, nullptr, OrphanedRootPin);
 				RootPin->SubPins.Remove(Pin);
 				EnsurePinValidity(Pin, false);
 				AddSubPin(OrphanedRootPin, Pin);
@@ -11845,7 +11952,6 @@ URigVMController::FPinState URigVMController::GetPinState(URigVMPin* InPin) cons
 	State.CPPType = InPin->GetCPPType();
 	State.CPPTypeObject = InPin->GetCPPTypeObject();
 	State.DefaultValue = InPin->GetDefaultValue();
-	State.BoundVariable = InPin->GetBoundVariablePath();
 	State.bIsExpanded = InPin->IsExpanded();
 	State.InjectionInfos = InPin->GetInjectedNodes();
 	return State;
@@ -11872,9 +11978,9 @@ void URigVMController::ApplyPinState(URigVMPin* InPin, const FPinState& InPinSta
 {
 	for (URigVMInjectionInfo* InjectionInfo : InPinState.InjectionInfos)
 	{
-		InjectionInfo->Rename(nullptr, InPin, REN_ForceNoResetLoaders | REN_DoNotDirty | REN_DontCreateRedirectors | REN_NonTransactional);
-		InjectionInfo->InputPin = InjectionInfo->UnitNode->FindPin(InjectionInfo->InputPin->GetName());
-		InjectionInfo->OutputPin = InjectionInfo->UnitNode->FindPin(InjectionInfo->OutputPin->GetName());
+		RenameObject(InjectionInfo, nullptr, InPin);
+		InjectionInfo->InputPin = InjectionInfo->InputPin ? InjectionInfo->Node->FindPin(InjectionInfo->InputPin->GetName()) : nullptr;
+		InjectionInfo->OutputPin = InjectionInfo->OutputPin ? InjectionInfo->Node->FindPin(InjectionInfo->OutputPin->GetName()) : nullptr;
 		InPin->InjectionInfos.Add(InjectionInfo);
 	}
 
@@ -11884,7 +11990,6 @@ void URigVMController::ApplyPinState(URigVMPin* InPin, const FPinState& InPinSta
 	}
 
 	SetPinExpansion(InPin, InPinState.bIsExpanded, false);
-	BindPinToVariable(InPin, InPinState.BoundVariable, false);
 }
 
 void URigVMController::ApplyPinStates(URigVMNode* InNode, const TMap<FString, URigVMController::FPinState>& InPinStates, const TMap<FString, FString>& InRedirectedPinPaths)
@@ -11908,7 +12013,7 @@ void URigVMController::ApplyPinStates(URigVMNode* InNode, const TMap<FString, UR
 		{
 			for (URigVMInjectionInfo* InjectionInfo : PinState.InjectionInfos)
 			{
-				InjectionInfo->UnitNode->Rename(nullptr, InNode->GetGraph(), REN_ForceNoResetLoaders | REN_DoNotDirty | REN_DontCreateRedirectors | REN_NonTransactional);
+				RenameObject(InjectionInfo->Node, nullptr, InNode->GetGraph());
 				DestroyObject(InjectionInfo);
 			}
 		}
@@ -12456,9 +12561,14 @@ void URigVMController::RewireLinks(URigVMPin* InOldPin, URigVMPin* InNewPin, boo
 
 #endif
 
+bool URigVMController::RenameObject(UObject* InObjectToRename, const TCHAR* InNewName, UObject* InNewOuter)
+{
+	return InObjectToRename->Rename(InNewName, InNewOuter, REN_ForceNoResetLoaders | REN_DoNotDirty | REN_DontCreateRedirectors | REN_NonTransactional);
+}
+
 void URigVMController::DestroyObject(UObject* InObjectToDestroy)
 {
-	InObjectToDestroy->Rename(nullptr, GetTransientPackage(), REN_ForceNoResetLoaders | REN_DoNotDirty | REN_DontCreateRedirectors | REN_NonTransactional);
+	RenameObject(InObjectToDestroy, nullptr, GetTransientPackage());
 	InObjectToDestroy->RemoveFromRoot();
 }
 
diff --git a/Engine/Source/Developer/RigVMDeveloper/Private/RigVMModel/RigVMControllerActions.cpp b/Engine/Source/Developer/RigVMDeveloper/Private/RigVMModel/RigVMControllerActions.cpp
--- a/Engine/Source/Developer/RigVMDeveloper/Private/RigVMModel/RigVMControllerActions.cpp
+++ b/Engine/Source/Developer/RigVMDeveloper/Private/RigVMModel/RigVMControllerActions.cpp
@@ -83,6 +83,8 @@ bool URigVMActionStack::Undo(URigVMController* InController)
 		RedoActions.Add(KeyToUndo);
 		return true;
 	}
+
+	InController->ReportAndNotifyErrorf(TEXT("Error while undoing action %s."), *Wrapper.GetAction()->Title);
 	return false;
 }
 
@@ -105,6 +107,8 @@ bool URigVMActionStack::Redo(URigVMController* InController)
 		ActionIndex = UndoActions.Num();
 		return true;
 	}
+
+	InController->ReportAndNotifyErrorf(TEXT("Error while undoing action %s."), *Wrapper.GetAction()->Title);
 	return false;
 }
 
@@ -175,6 +179,7 @@ bool FRigVMBaseAction::Undo(URigVMController* InController)
 		FRigVMActionWrapper Wrapper(SubActions[KeyIndex]);
 		if(!Wrapper.GetAction()->Undo(InController))
 		{
+			InController->ReportAndNotifyErrorf(TEXT("Error while undoing action '%s'."), *Wrapper.GetAction()->Title);
 			Result = false;
 		}
 	}
@@ -189,6 +194,7 @@ bool FRigVMBaseAction::Redo(URigVMController* InController)
 		FRigVMActionWrapper Wrapper(SubActions[KeyIndex]);
 		if (!Wrapper.GetAction()->Redo(InController))
 		{
+			InController->ReportAndNotifyErrorf(TEXT("Error while redoing action '%s'."), *Wrapper.GetAction()->Title);
 			Result = false;
 		}
 	}
@@ -622,12 +628,25 @@ FRigVMAddInjectedNodeAction::FRigVMAddInjectedNodeAction()
 FRigVMAddInjectedNodeAction::FRigVMAddInjectedNodeAction(URigVMInjectionInfo* InInjectionInfo)
 	: PinPath(InInjectionInfo->GetPin()->GetPinPath())
 	, bAsInput(InInjectionInfo->bInjectedAsInput)
-	, ScriptStructPath(InInjectionInfo->UnitNode->GetScriptStruct()->GetPathName())
-	, MethodName(InInjectionInfo->UnitNode->GetMethodName())
-	, InputPinName(InInjectionInfo->InputPin->GetFName())
-	, OutputPinName(InInjectionInfo->OutputPin->GetFName())
-	, NodePath(InInjectionInfo->UnitNode->GetName())
+	, NodePath(InInjectionInfo->Node->GetName())
 {
+	if (InInjectionInfo->InputPin)
+	{
+		InputPinName = InInjectionInfo->InputPin->GetFName();
+	}
+	if (InInjectionInfo->OutputPin)
+	{
+		OutputPinName = InInjectionInfo->OutputPin->GetFName();
+	}
+	if (URigVMUnitNode* UnitNode = Cast<URigVMUnitNode>(InInjectionInfo->Node))
+	{
+		ScriptStructPath = UnitNode->GetScriptStruct()->GetPathName();
+		MethodName = UnitNode->GetMethodName();
+	}
+	else
+	{
+		checkNoEntry();
+	}
 }
 
 bool FRigVMAddInjectedNodeAction::Undo(URigVMController* InController)
@@ -642,9 +661,12 @@ bool FRigVMAddInjectedNodeAction::Undo(URigVMController* InController)
 bool FRigVMAddInjectedNodeAction::Redo(URigVMController* InController)
 {
 #if WITH_EDITOR
-	if (URigVMInjectionInfo* InjectionInfo = InController->AddInjectedNodeFromStructPath(PinPath, bAsInput, ScriptStructPath, MethodName, InputPinName, OutputPinName, NodePath, false))
+	if (!ScriptStructPath.IsEmpty())
 	{
-		return FRigVMBaseAction::Redo(InController);
+		if (URigVMInjectionInfo* InjectionInfo = InController->AddInjectedNodeFromStructPath(PinPath, bAsInput, ScriptStructPath, MethodName, InputPinName, OutputPinName, NodePath, false))
+		{
+			return FRigVMBaseAction::Redo(InController);
+		}
 	}
 #endif
 	return false;
@@ -656,7 +678,15 @@ FRigVMRemoveNodeAction::FRigVMRemoveNodeAction(URigVMNode* InNode, URigVMControl
 
 	if (URigVMInjectionInfo* InjectionInfo = InNode->GetInjectionInfo())
 	{
-		InverseAction.AddAction(FRigVMAddInjectedNodeAction(InjectionInfo));
+		if (URigVMUnitNode* UnitNode = Cast<URigVMUnitNode>(InjectionInfo->Node))
+		{
+			InverseAction.AddAction(FRigVMAddInjectedNodeAction(InjectionInfo));
+		}
+		else if (URigVMVariableNode* VariableNode = Cast<URigVMVariableNode>(InjectionInfo->Node))
+		{
+			InverseAction.AddAction(FRigVMSetPinBoundVariableAction(InjectionInfo->GetPin(), InjectionInfo->GetPin()->GetBoundVariablePath()));
+		}
+
 		for (URigVMPin* Pin : InNode->GetPins())
 		{
 			if (Pin->GetDirection() == ERigVMPinDirection::Input ||
@@ -1430,6 +1460,7 @@ FRigVMSetPinBoundVariableAction::FRigVMSetPinBoundVariableAction(URigVMPin* InPi
 	, OldBoundVariablePath(InPin->GetBoundVariablePath())
 	, NewBoundVariablePath(InNewBoundVariablePath)
 {
+	
 }
 
 bool FRigVMSetPinBoundVariableAction::Merge(const FRigVMBaseAction* Other)
@@ -1861,7 +1892,7 @@ bool FRigVMRemoveLocalVariableAction::Undo(URigVMController* InController)
 	{
 		return false;
 	}
-	return InController->AddLocalVariable(LocalVariable.Name, LocalVariable.CPPType, LocalVariable.CPPTypeObject, LocalVariable.DefaultValue, false).Name.IsNone() != false;
+	return !InController->AddLocalVariable(LocalVariable.Name, LocalVariable.CPPType, LocalVariable.CPPTypeObject, LocalVariable.DefaultValue, false).Name.IsNone();
 }
 
 bool FRigVMRemoveLocalVariableAction::Redo(URigVMController* InController)
diff --git a/Engine/Source/Developer/RigVMDeveloper/Private/RigVMModel/RigVMPin.cpp b/Engine/Source/Developer/RigVMDeveloper/Private/RigVMModel/RigVMPin.cpp
--- a/Engine/Source/Developer/RigVMDeveloper/Private/RigVMModel/RigVMPin.cpp
+++ b/Engine/Source/Developer/RigVMDeveloper/Private/RigVMModel/RigVMPin.cpp
@@ -192,7 +192,7 @@ URigVMPin::URigVMPin()
 	, CPPTypeObject(nullptr)
 	, CPPTypeObjectPath(NAME_None)
 	, DefaultValue(FString())
-	, BoundVariablePath()
+	, BoundVariablePath_DEPRECATED()
 {
 #if UE_BUILD_DEBUG
 	CachedPinPath = GetPinPath();
@@ -376,10 +376,13 @@ FName URigVMPin::GetDisplayName() const
 
 		for (URigVMInjectionInfo* Injection : InjectionInfos)
 		{
-			if (TSharedPtr<FStructOnScope> DefaultStructScope = Injection->UnitNode->ConstructStructInstance())
+			if (URigVMUnitNode* InjectedUnitNode = Cast<URigVMUnitNode>(Injection->Node))
 			{
-				FRigVMStruct* DefaultStruct = (FRigVMStruct*)DefaultStructScope->GetStructMemory();
-				ProcessedDisplayName = DefaultStruct->ProcessPinLabelForInjection(ProcessedDisplayName);
+				if (TSharedPtr<FStructOnScope> DefaultStructScope = InjectedUnitNode->ConstructStructInstance())
+				{
+					FRigVMStruct* DefaultStruct = (FRigVMStruct*)DefaultStructScope->GetStructMemory();
+					ProcessedDisplayName = DefaultStruct->ProcessPinLabelForInjection(ProcessedDisplayName);
+				}
 			}
 		}
 
@@ -914,34 +917,59 @@ UObject* URigVMPin::FindObjectFromCPPTypeObjectPath(const FString& InObjectPath)
 	return FindObject<UObject>(ANY_PACKAGE, *InObjectPath);
 }
 
+URigVMVariableNode* URigVMPin::GetBoundVariableNode() const
+{
+	for (TObjectPtr<URigVMInjectionInfo> InjectionInfo : InjectionInfos)
+	{
+		if (URigVMVariableNode* VariableNode = Cast<URigVMVariableNode>(InjectionInfo->Node))
+		{
+			return VariableNode;
+		}
+	}
+	
+	return nullptr;
+}
+
 // Returns the variable bound to this pin (or NAME_None)
-const FString& URigVMPin::GetBoundVariablePath() const
+const FString URigVMPin::GetBoundVariablePath() const
 {
 	return GetBoundVariablePath(EmptyPinOverride);
 }
 
 // Returns the variable bound to this pin (or NAME_None)
-const FString& URigVMPin::GetBoundVariablePath(const URigVMPin::FPinOverride& InOverride) const
+const FString URigVMPin::GetBoundVariablePath(const URigVMPin::FPinOverride& InOverride) const
 {
 	if (FPinOverrideValue const* OverrideValuePtr = InOverride.Value.Find(InOverride.Key.GetSibling((URigVMPin*)this)))
 	{
 		return OverrideValuePtr->BoundVariablePath;
 	}
-	return BoundVariablePath;
+
+	for (TObjectPtr<URigVMInjectionInfo> InjectionInfo : InjectionInfos)
+	{
+		if (URigVMVariableNode* VariableNode = Cast<URigVMVariableNode>(InjectionInfo->Node))
+		{
+			FString SegmentPath = InjectionInfo->OutputPin->GetSegmentPath(false);
+			if (SegmentPath.IsEmpty())
+			{
+				return VariableNode->GetVariableName().ToString();
+			}
+			
+			return VariableNode->GetVariableName().ToString() + TEXT(".") + SegmentPath;
+		}
+	}
+	
+	return FString();
 }
 
 // Returns the variable bound to this pin (or NAME_None)
 FString URigVMPin::GetBoundVariableName() const
 {
-	return GetBoundVariableName(EmptyPinOverride);
-}
+	if (URigVMVariableNode* VariableNode = GetBoundVariableNode())
+	{
+		return VariableNode->GetVariableName().ToString();
+	}
 
-// Returns the variable bound to this pin (or NAME_None)
-FString URigVMPin::GetBoundVariableName(const URigVMPin::FPinOverride& InOverride) const
-{
-	FString VariableName = GetBoundVariablePath(InOverride);
-	BoundVariablePath.Split(TEXT("."), &VariableName, nullptr);
-	return VariableName;
+	return FString();
 }
 
 // Returns true if this pin is bound to a variable
@@ -958,12 +986,7 @@ bool URigVMPin::IsBoundToVariable(const URigVMPin::FPinOverride& InOverride) con
 
 bool URigVMPin::IsBoundToExternalVariable() const
 {
-	return IsBoundToExternalVariable(EmptyPinOverride);
-}
-
-bool URigVMPin::IsBoundToExternalVariable(const FPinOverride& InOverride) const
-{
-	FString VariableName = GetBoundVariableName(InOverride);
+	FString VariableName = GetBoundVariableName();
 	if (VariableName.IsEmpty())
 	{
 		return false;
@@ -983,12 +1006,7 @@ bool URigVMPin::IsBoundToExternalVariable(const FPinOverride& InOverride) const
 
 bool URigVMPin::IsBoundToLocalVariable() const
 {
-	return IsBoundToLocalVariable(EmptyPinOverride);
-}
-
-bool URigVMPin::IsBoundToLocalVariable(const FPinOverride& InOverride) const
-{
-	FString VariableName = GetBoundVariableName(InOverride);
+	FString VariableName = GetBoundVariableName();
 	if (VariableName.IsEmpty())
 	{
 		return false;
@@ -1008,12 +1026,7 @@ bool URigVMPin::IsBoundToLocalVariable(const FPinOverride& InOverride) const
 
 bool URigVMPin::IsBoundToInputArgument() const
 {
-	return IsBoundToInputArgument(EmptyPinOverride);
-}
-
-bool URigVMPin::IsBoundToInputArgument(const FPinOverride& InOverride) const
-{
-	FString VariableName = GetBoundVariableName(InOverride);
+	FString VariableName = GetBoundVariableName();
 	if (VariableName.IsEmpty())
 	{
 		return false;
@@ -1147,25 +1160,9 @@ FRigVMExternalVariable URigVMPin::ToExternalVariable() const
 
 		VariableName = VariableName.Replace(TEXT("."), TEXT("_"));
 	}
-	ExternalVariable.Name = *VariableName;
-
-	if (RigVMTypeUtils::IsArrayType(CPPType))
-	{
-		ExternalVariable.bIsArray = true;
-		ExternalVariable.TypeName = *CPPType.Mid(7, CPPType.Len() - 8);
-		ExternalVariable.TypeObject = CPPTypeObject;
-	}
-	else
-	{
-		ExternalVariable.bIsArray = false;
-		ExternalVariable.TypeName = *CPPType;
-		ExternalVariable.TypeObject = CPPTypeObject;
-	}
-
-	ExternalVariable.bIsPublic = false;
-	ExternalVariable.bIsReadOnly = false;
-	ExternalVariable.Memory = nullptr;
 
+	ExternalVariable = RigVMTypeUtils::ExternalVariableFromCPPType(*VariableName, CPPType, GetCPPTypeObject(), false, false);
+	
 	return ExternalVariable;
 }
 
@@ -1237,7 +1234,7 @@ URigVMPin* URigVMPin::GetPinForLink() const
 {
 	URigVMPin* RootPin = GetRootPin();
 
-	if (RootPin->InjectionInfos.Num() == 0)
+	if (!RootPin->HasInjectedUnitNodes())
 	{
 		return const_cast<URigVMPin*>(this);
 	}
@@ -1262,7 +1259,7 @@ URigVMPin* URigVMPin::GetOriginalPinFromInjectedNode() const
 	{
 		URigVMPin* RootPin = GetRootPin();
 		URigVMPin* OriginalPin = nullptr;
-		if (Injection->bInjectedAsInput && Injection->InputPin == RootPin)
+		if (Injection->bInjectedAsInput && Injection->InputPin == RootPin && Injection->OutputPin)
 		{
 			TArray<URigVMPin*> LinkedPins = Injection->OutputPin->GetLinkedTargetPins();
 			if (LinkedPins.Num() == 1)
@@ -1270,7 +1267,7 @@ URigVMPin* URigVMPin::GetOriginalPinFromInjectedNode() const
 				OriginalPin = LinkedPins[0]->GetOriginalPinFromInjectedNode();
 			}
 		}
-		else if (!Injection->bInjectedAsInput && Injection->OutputPin == RootPin)
+		else if (!Injection->bInjectedAsInput && Injection->OutputPin == RootPin && Injection->InputPin)
 		{
 			TArray<URigVMPin*> LinkedPins = Injection->InputPin->GetLinkedSourcePins();
 			if (LinkedPins.Num() == 1)
@@ -1710,3 +1707,16 @@ bool URigVMPin::CanLink(URigVMPin* InSourcePin, URigVMPin* InTargetPin, FString*
 
 	return true;
 }
+
+bool URigVMPin::HasInjectedUnitNodes() const
+{
+	for (URigVMInjectionInfo* Info : InjectionInfos)
+	{
+		if (Info->Node.IsA<URigVMUnitNode>())
+		{
+			return true;
+		}
+	}
+	
+	return false;
+}
diff --git a/Engine/Source/Developer/RigVMDeveloper/Private/RigVMTypeUtils.cpp b/Engine/Source/Developer/RigVMDeveloper/Private/RigVMTypeUtils.cpp
--- a/Engine/Source/Developer/RigVMDeveloper/Private/RigVMTypeUtils.cpp
+++ b/Engine/Source/Developer/RigVMDeveloper/Private/RigVMTypeUtils.cpp
@@ -5,6 +5,7 @@
 
 #if WITH_EDITOR
 
+#include "EdGraphSchema_K2.h"
 #include "Engine/Blueprint.h"
 
 FRigVMExternalVariable RigVMTypeUtils::ExternalVariableFromBPVariableDescription(
@@ -117,8 +118,76 @@ FRigVMExternalVariable RigVMTypeUtils::ExternalVariableFromPinType(const FName&
 	return ExternalVariable;
 }
 
-FRigVMExternalVariable RigVMTypeUtils::ExternalVariableFromCPPType(const FName& InName, const FString& InCPPType,
-	bool bInPublic, bool bInReadonly)
+FRigVMExternalVariable RigVMTypeUtils::ExternalVariableFromCPPTypePath(const FName& InName, const FString& InCPPTypePath, bool bInPublic, bool bInReadonly)
+{
+	FRigVMExternalVariable Variable;
+	if (InCPPTypePath.StartsWith(TEXT("TMap<")))
+	{
+		return Variable;
+	}
+	
+	Variable.Name = InName;
+	Variable.bIsPublic = bInPublic;
+	Variable.bIsReadOnly = bInReadonly;
+
+	FString CPPTypePath = InCPPTypePath;
+	Variable.bIsArray = RigVMTypeUtils::IsArrayType(CPPTypePath);
+	if (Variable.bIsArray)
+	{
+		CPPTypePath = BaseTypeFromArrayType(CPPTypePath);
+	}
+
+	if (CPPTypePath == TEXT("bool"))
+	{
+		Variable.TypeName = *CPPTypePath;
+		Variable.Size = sizeof(bool);
+	}
+	else if (CPPTypePath == TEXT("float"))
+	{
+		Variable.TypeName = *CPPTypePath;
+		Variable.Size = sizeof(float);
+	}
+	else if (CPPTypePath == TEXT("double"))
+	{
+		Variable.TypeName = *CPPTypePath;
+		Variable.Size = sizeof(double);
+	}
+	else if (CPPTypePath == TEXT("int32"))
+	{
+		Variable.TypeName = *CPPTypePath;
+		Variable.Size = sizeof(int32);
+	}
+	else if (CPPTypePath == TEXT("FString"))
+	{
+		Variable.TypeName = *CPPTypePath;
+		Variable.Size = sizeof(FString);
+	}
+	else if (CPPTypePath == TEXT("FName"))
+	{
+		Variable.TypeName = *CPPTypePath;
+		Variable.Size = sizeof(FName);
+	}
+	else if(UScriptStruct* ScriptStruct = URigVMPin::FindObjectFromCPPTypeObjectPath<UScriptStruct>(CPPTypePath))
+	{
+		Variable.TypeName = *ScriptStruct->GetStructCPPName();
+		Variable.TypeObject = ScriptStruct;
+		Variable.Size = ScriptStruct->GetStructureSize();
+	}
+	else if (UEnum* Enum= URigVMPin::FindObjectFromCPPTypeObjectPath<UEnum>(CPPTypePath))
+	{
+		Variable.TypeName = *Enum->CppType;
+		Variable.TypeObject = Enum;
+		Variable.Size = Enum->GetResourceSizeBytes(EResourceSizeMode::EstimatedTotal);
+	}
+	else
+	{
+		check(false);
+	}
+
+	return Variable;
+}
+
+FRigVMExternalVariable RigVMTypeUtils::ExternalVariableFromCPPType(const FName& InName, const FString& InCPPType, UObject* InCPPTypeObject, bool bInPublic, bool bInReadonly)
 {
 	FRigVMExternalVariable Variable;
 	if (InCPPType.StartsWith(TEXT("TMap<")))
@@ -134,7 +203,7 @@ FRigVMExternalVariable RigVMTypeUtils::ExternalVariableFromCPPType(const FName&
 	Variable.bIsArray = RigVMTypeUtils::IsArrayType(CPPType);
 	if (Variable.bIsArray)
 	{
-		CPPType = CPPType.RightChop(7).LeftChop(1);
+		CPPType = BaseTypeFromArrayType(CPPType);
 	}
 
 	if (CPPType == TEXT("bool"))
@@ -181,7 +250,9 @@ FRigVMExternalVariable RigVMTypeUtils::ExternalVariableFromCPPType(const FName&
 	}
 	else
 	{
-		check(false);
+		Variable.TypeName = *CPPType;
+		Variable.TypeObject = InCPPTypeObject;
+		Variable.Size = InCPPTypeObject->StaticClass()->GetStructureSize();
 	}
 
 	return Variable;
@@ -302,6 +373,55 @@ FEdGraphPinType RigVMTypeUtils::PinTypeFromRigVMVariableDescription(
 	return PinType;
 }
 
+FEdGraphPinType RigVMTypeUtils::SubPinType(const FEdGraphPinType& InPinType, const FString& InSegmentPath)
+{
+	FEdGraphPinType Result;
+	if (InSegmentPath.IsEmpty())
+	{
+		return InPinType;
+	}
+
+	if (InPinType.PinSubCategoryObject.IsValid())
+	{
+		if (UStruct* Struct = Cast<UStruct>(InPinType.PinSubCategoryObject))
+		{
+			int32 PartIndex = 0; 
+			TArray<FString> Parts;
+			if (!URigVMPin::SplitPinPath(InSegmentPath, Parts))
+			{
+				Parts.Add(InSegmentPath);
+			}
+			
+			FProperty* Property = Struct->FindPropertyByName(*Parts[PartIndex++]);
+			while (PartIndex < Parts.Num() && Property != nullptr)
+			{
+				if (FArrayProperty* ArrayProperty = CastField<FArrayProperty>(Property))
+				{
+					Property = ArrayProperty->Inner;
+					PartIndex++;
+					continue;
+				}
+
+				if (FStructProperty* StructProperty = CastField<FStructProperty>(Property))
+				{
+					Struct = StructProperty->Struct;
+					Property = Struct->FindPropertyByName(*Parts[PartIndex++]);
+					continue;
+				}
+
+				break;
+			}
+
+			if (PartIndex == Parts.Num() && Property)
+			{
+				UEdGraphSchema_K2::StaticClass()->GetDefaultObject<UEdGraphSchema_K2>()->ConvertPropertyToPinType(Property, Result);
+			}
+		}
+	}
+
+	return Result;
+}
+
 bool RigVMTypeUtils::CPPTypeFromPinType(const FEdGraphPinType& InPinType, FString& OutCPPType, UObject** OutCPPTypeObject)
 {
 	FString Prefix = "";
@@ -429,4 +549,17 @@ bool RigVMTypeUtils::CPPTypeFromExternalVariable(const FRigVMExternalVariable& I
 		
 	return true;
 }
+
+bool RigVMTypeUtils::AreCompatible(const FRigVMExternalVariable& InTypeA, const FRigVMExternalVariable& InTypeB, const FString& InSegmentPathA, const FString& InSegmentPathB)
+{
+	return AreCompatible(PinTypeFromExternalVariable(InTypeA), PinTypeFromExternalVariable(InTypeB), InSegmentPathA, InSegmentPathB);
+}
+
+bool RigVMTypeUtils::AreCompatible(const FEdGraphPinType& InTypeA, const FEdGraphPinType& InTypeB, const FString& InSegmentPathA, const FString& InSegmentPathB)
+{
+	FEdGraphPinType SubPinTypeA = SubPinType(InTypeA, InSegmentPathA);
+	FEdGraphPinType SubPinTypeB = SubPinType(InTypeB, InSegmentPathB);
+	
+	return UEdGraphSchema_K2::StaticClass()->GetDefaultObject<UEdGraphSchema_K2>()->ArePinTypesCompatible(SubPinTypeA, SubPinTypeB);	
+}
 #endif
diff --git a/Engine/Source/Developer/RigVMDeveloper/Public/RigVMModel/RigVMController.h b/Engine/Source/Developer/RigVMDeveloper/Public/RigVMModel/RigVMController.h
--- a/Engine/Source/Developer/RigVMDeveloper/Public/RigVMModel/RigVMController.h
+++ b/Engine/Source/Developer/RigVMDeveloper/Public/RigVMModel/RigVMController.h
@@ -807,6 +807,7 @@ private:
 	bool RemovePin(URigVMPin* InPinToRemove, bool bSetupUndoRedo, bool bNotify);
 	FProperty* FindPropertyForPin(const FString& InPinPath);
 	bool BindPinToVariable(URigVMPin* InPin, const FString& InNewBoundVariablePath, bool bSetupUndoRedo);
+	bool UnbindPinFromVariable(URigVMPin* InPin, bool bSetupUndoRedo);
 	bool MakeBindingsFromVariableNode(URigVMVariableNode* InNode, bool bSetupUndoRedo);
 	bool PromotePinToVariable(URigVMPin* InPin, bool bCreateVariableNode, const FVector2D& InNodePosition, bool bSetupUndoRedo);
 
@@ -842,7 +843,6 @@ private:
 		FString CPPType;
 		UObject* CPPTypeObject;
 		FString DefaultValue;
-		FString BoundVariable;
 		bool bIsExpanded;
 		TArray<URigVMInjectionInfo*> InjectionInfos;
 	};
@@ -872,6 +872,8 @@ private:
 	void RewireLinks(URigVMPin* OldPin, URigVMPin* NewPin, bool bAsInput, bool bSetupUndoRedo, TArray<URigVMLink*> InLinks = TArray<URigVMLink*>());
 #endif
 
+
+	bool RenameObject(UObject* InObjectToRename, const TCHAR* InNewName, UObject* InNewOuter = nullptr);
 	void DestroyObject(UObject* InObjectToDestroy);
 	static void AddNodePin(URigVMNode* InNode, URigVMPin* InPin);
 	static void AddSubPin(URigVMPin* InParentPin, URigVMPin* InPin);
diff --git a/Engine/Source/Developer/RigVMDeveloper/Public/RigVMModel/RigVMPin.h b/Engine/Source/Developer/RigVMDeveloper/Public/RigVMModel/RigVMPin.h
--- a/Engine/Source/Developer/RigVMDeveloper/Public/RigVMModel/RigVMPin.h
+++ b/Engine/Source/Developer/RigVMDeveloper/Public/RigVMModel/RigVMPin.h
@@ -15,6 +15,7 @@ class URigVMNode;
 class URigVMUnitNode;
 class URigVMPin;
 class URigVMLink;
+class URigVMVariableNode;
 
 /**
  * The Injected Info is used for injecting a node on a pin.
@@ -34,7 +35,10 @@ public:
 	}
 
 	UPROPERTY()
-	TObjectPtr<URigVMUnitNode> UnitNode;
+	TObjectPtr<URigVMUnitNode> UnitNode_DEPRECATED;
+
+	UPROPERTY()
+	TObjectPtr<URigVMNode> Node;
 
 	UPROPERTY()
 	bool bInjectedAsInput;
@@ -333,20 +337,22 @@ public:
 	// Returns true if this pin has injected nodes
 	bool HasInjectedNodes() const { return InjectionInfos.Num() > 0; }
 
+	// Returns true if this pin has injected nodes
+	bool HasInjectedUnitNodes() const;
+
 	// Returns the injected nodes this pin contains.
 	const TArray<URigVMInjectionInfo*> GetInjectedNodes() const { return InjectionInfos; }
 
-	// Returns the variable bound to this pin (or NAME_None)
-	const FString& GetBoundVariablePath() const;
+	URigVMVariableNode* GetBoundVariableNode() const;
 
 	// Returns the variable bound to this pin (or NAME_None)
-	const FString& GetBoundVariablePath(const FPinOverride& InOverride) const;
+	const FString GetBoundVariablePath() const;
 
 	// Returns the variable bound to this pin (or NAME_None)
-	FString GetBoundVariableName() const;
+	const FString GetBoundVariablePath(const FPinOverride& InOverride) const;
 
 	// Returns the variable bound to this pin (or NAME_None)
-	FString GetBoundVariableName(const FPinOverride& InOverride) const;
+	FString GetBoundVariableName() const;
 
 	// Returns true if this pin is bound to a variable
 	bool IsBoundToVariable() const;
@@ -357,21 +363,12 @@ public:
 	// Returns true if this pin is bound to an external variable
 	bool IsBoundToExternalVariable() const;
 
-	// Returns true if this pin is bound to an external variable
-	bool IsBoundToExternalVariable(const FPinOverride& InOverride) const;
-
 	// Returns true if this pin is bound to a local variable
 	bool IsBoundToLocalVariable() const;
 
-	// Returns true if this pin is bound to a local variable
-	bool IsBoundToLocalVariable(const FPinOverride& InOverride) const;
-
 	// Returns true if this pin is bound to an input argument
 	bool IsBoundToInputArgument() const;
 
-	// Returns true if this pin is bound to an input argument
-	bool IsBoundToInputArgument(const FPinOverride& InOverride) const;
-
 #if UE_RIGVM_UCLASS_BASED_STORAGE_DISABLED
 	// Returns true if the pin can be bound to a given variable
 	bool CanBeBoundToVariable(const FRigVMExternalVariable& InExternalVariable, const FRigVMRegisterOffset& InOffset = FRigVMRegisterOffset()) const;
@@ -457,7 +454,7 @@ private:
 	TArray<TObjectPtr<URigVMInjectionInfo>> InjectionInfos;
 
 	UPROPERTY()
-	FString BoundVariablePath;
+	FString BoundVariablePath_DEPRECATED;
 
 	static const FString OrphanPinPrefix;
 
diff --git a/Engine/Source/Developer/RigVMDeveloper/Public/RigVMTypeUtils.h b/Engine/Source/Developer/RigVMDeveloper/Public/RigVMTypeUtils.h
--- a/Engine/Source/Developer/RigVMDeveloper/Public/RigVMTypeUtils.h
+++ b/Engine/Source/Developer/RigVMDeveloper/Public/RigVMTypeUtils.h
@@ -34,16 +34,23 @@ namespace RigVMTypeUtils
 	
 	RIGVMDEVELOPER_API FRigVMExternalVariable ExternalVariableFromPinType(const FName& InName, const FEdGraphPinType& InPinType, bool bInPublic = false, bool bInReadonly = false);
 
-	RIGVMDEVELOPER_API FRigVMExternalVariable ExternalVariableFromCPPType(const FName& InName, const FString& InCPPType, bool bInPublic = false, bool bInReadonly = false);
+	RIGVMDEVELOPER_API FRigVMExternalVariable ExternalVariableFromCPPTypePath(const FName& InName, const FString& InCPPTypePath, bool bInPublic = false, bool bInReadonly = false);
+
+	RIGVMDEVELOPER_API FRigVMExternalVariable ExternalVariableFromCPPType(const FName& InName, const FString& InCPPType, UObject* InCPPTypeObject, bool bInPublic = false, bool bInReadonly = false);
 
 	RIGVMDEVELOPER_API FEdGraphPinType PinTypeFromExternalVariable(const FRigVMExternalVariable& InExternalVariable);
 
 	RIGVMDEVELOPER_API FEdGraphPinType PinTypeFromRigVMVariableDescription(const FRigVMGraphVariableDescription& InVariableDescription);
 
+	RIGVMDEVELOPER_API FEdGraphPinType SubPinType(const FEdGraphPinType& InPinType, const FString& SegmentPath);
+
 	RIGVMDEVELOPER_API bool CPPTypeFromPinType(const FEdGraphPinType& InPinType, FString& OutCPPType, UObject** OutCPPTypeObject);
 
 	RIGVMDEVELOPER_API bool CPPTypeFromExternalVariable(const FRigVMExternalVariable& InExternalVariable, FString& OutCPPType, UObject** OutCPPTypeObject);
 
+	RIGVMDEVELOPER_API bool AreCompatible(const FRigVMExternalVariable& InTypeA, const FRigVMExternalVariable& InTypeB, const FString& InSegmentPathA = FString(), const FString& InSegmentPathB = FString());
+	
+	RIGVMDEVELOPER_API bool AreCompatible(const FEdGraphPinType& InTypeA, const FEdGraphPinType& InTypeB, const FString& InSegmentPathA = FString(), const FString& InSegmentPathB = FString());
 	
 #endif
 }
