Control Rig: Fix horde compiler error in EjectNode

diff --git a/Engine/Source/Developer/RigVMDeveloper/Private/RigVMModel/RigVMController.cpp b/Engine/Source/Developer/RigVMDeveloper/Private/RigVMModel/RigVMController.cpp
--- a/Engine/Source/Developer/RigVMDeveloper/Private/RigVMModel/RigVMController.cpp
+++ b/Engine/Source/Developer/RigVMDeveloper/Private/RigVMModel/RigVMController.cpp
@@ -2281,6 +2281,16 @@ URigVMNode* URigVMController::EjectNodeFromPin(const FString& InPinPath, bool bS
 		EjectedNode = AddVariableNode(VariableNode->GetVariableName(), VariableNode->GetCPPType(), VariableNode->GetCPPTypeObject(), true, VariableNode->GetDefaultValue(), Position, FString(), bSetupUndoRedo);
 	}
 
+	if (!EjectedNode)
+	{
+		ReportErrorf(TEXT("Could not eject node."));
+		if (bSetupUndoRedo)
+		{
+			ActionStack->CancelAction(Action);
+		}
+		return nullptr;
+	}
+
 	for (const TPair<FName, FString>& Pair : DefaultValues)
 	{
 		if (Pair.Value.IsEmpty())
