Control Rig: make sure pins are in a valid state before querying its default value

diff --git a/Engine/Source/Developer/RigVMDeveloper/Private/RigVMModel/RigVMController.cpp b/Engine/Source/Developer/RigVMDeveloper/Private/RigVMModel/RigVMController.cpp
--- a/Engine/Source/Developer/RigVMDeveloper/Private/RigVMModel/RigVMController.cpp
+++ b/Engine/Source/Developer/RigVMDeveloper/Private/RigVMModel/RigVMController.cpp
@@ -11977,6 +11977,16 @@ TMap<FString, URigVMController::FPinState> URigVMController::GetPinStates(URigVM
 		FString PinPath, NodeName;
 		URigVMPin::SplitPinPathAtStart(Pin->GetPinPath(), NodeName, PinPath);
 
+		// we need to ensure validity here because GetPinState()-->GetDefaultValue() needs pin to be in a valid state.
+		// some additional context:
+		// right after load, some pins will be a invalid state because they don't have their CPPTypeObject,
+		// which is expected since it is a transient property.
+		// if the CPPTypeObject is not there, those pins may struggle with producing a valid default value
+		// because Pin->IsStruct() will always be false if the pin does not have a valid type object.
+		if (Pin->IsRootPin())
+		{
+			EnsurePinValidity(Pin, true);
+		}
 		FPinState State = GetPinState(Pin);
 		PinStates.Add(PinPath, State);
 	}
