Control Rig: Fix redirector lookups when loading pins whose type has changed.

diff --git a/Engine/Plugins/Animation/ControlRig/Config/BaseControlRig.ini b/Engine/Plugins/Animation/ControlRig/Config/BaseControlRig.ini
--- a/Engine/Plugins/Animation/ControlRig/Config/BaseControlRig.ini
+++ b/Engine/Plugins/Animation/ControlRig/Config/BaseControlRig.ini
@@ -12,7 +12,7 @@
 +ClassRedirects=(OldName="RigVMGraph",NewName="RigVMGraph",ValueChanges=(("Model","RigVMModel")))
 
 +EnumRedirects=(OldName="ERigElementType",ValueChanges=(("Space","Null"))
-+EnumRedirects=(OldName="/Script/ControlRig.EControlRigRotationOrder",NewName="/Script/AnimationCore.EEulerRotationOrder")
++EnumRedirects=(OldName="EControlRigRotationOrder",NewName="/Script/AnimationCore.EEulerRotationOrder")
 
 
 ; renaming gizmo to shape
diff --git a/Engine/Source/Developer/RigVMDeveloper/Private/RigVMModel/RigVMController.cpp b/Engine/Source/Developer/RigVMDeveloper/Private/RigVMModel/RigVMController.cpp
--- a/Engine/Source/Developer/RigVMDeveloper/Private/RigVMModel/RigVMController.cpp
+++ b/Engine/Source/Developer/RigVMDeveloper/Private/RigVMModel/RigVMController.cpp
@@ -12611,6 +12611,43 @@ void URigVMController::AddSubPin(URigVMPin* InParentPin, URigVMPin* InPin)
 	InParentPin->SubPins.Add(InPin);
 }
 
+
+static UObject* FindObjectGloballyWithRedirectors(const TCHAR* InObjectName)
+{
+	// Do a global search for the CPP type. Note that searching with ANY_PACKAGE _does not_
+	// apply redirectors. So only if this fails do we apply them manually below.
+	UObject* Object = FindObject<UField>(ANY_PACKAGE, InObjectName);
+	if(Object != nullptr)
+	{
+		return Object;
+	}
+	
+	// Apply redirectors and do another lookup using the redirected name. GetRedirectedName
+	// will always return a valid name.
+	const FCoreRedirectObjectName NewObjectName = FCoreRedirects::GetRedirectedName(
+		ECoreRedirectFlags::Type_Class | ECoreRedirectFlags::Type_Struct | ECoreRedirectFlags::Type_Enum,
+		FCoreRedirectObjectName(InObjectName));
+
+	// If there was no package name, then there was no redirect set up for this type.
+	if (NewObjectName.PackageName.IsNone())
+	{
+		return nullptr;
+	}
+
+	const FString RedirectedObjectName = NewObjectName.ObjectName.ToString();
+	UPackage *Package = FindPackage(nullptr, *NewObjectName.PackageName.ToString());
+	if (Package != nullptr)
+	{
+		Object = FindObject<UField>(Package, *RedirectedObjectName);
+	}
+	if (Package == nullptr || Object == nullptr)
+	{
+		// Hail Mary pass.
+		Object = FindObject<UField>(ANY_PACKAGE, *RedirectedObjectName);
+	}
+	return Object;
+}
+
 bool URigVMController::EnsurePinValidity(URigVMPin* InPin, bool bRecursive)
 {
 	check(InPin);
@@ -12622,26 +12659,29 @@ bool URigVMController::EnsurePinValidity(URigVMPin* InPin, bool bRecursive)
 		{
 			// try to find the CPPTypeObject by name
 			FString CPPType = InPin->IsArray() ? InPin->GetArrayElementCppType() : InPin->GetCPPType();
-			UObject* CPPTypeObject = FindObject<UField>(ANY_PACKAGE, *CPPType);
-			if(CPPTypeObject == nullptr)
+
+			UObject* CPPTypeObject = FindObjectGloballyWithRedirectors(*CPPType);
+
+			if (CPPTypeObject == nullptr)
 			{
-				if(CPPType.StartsWith(TEXT("E"), ESearchCase::CaseSensitive) ||
-					CPPType.StartsWith(TEXT("F"), ESearchCase::CaseSensitive) ||
-					CPPType.StartsWith(TEXT("U"), ESearchCase::CaseSensitive))
+				// If we've mistakenly stored the struct type with the 'F', 'U', or 'A' prefixes, we need to strip them
+				// off first. Enums are always named with their prefix intact.
+				if (!CPPType.IsEmpty() && (CPPType[0] == TEXT('F') || CPPType[0] == TEXT('U') || CPPType[0] == TEXT('A')))
 				{
 					CPPType = CPPType.Mid(1);
-					CPPTypeObject = FindObject<UField>(ANY_PACKAGE, *CPPType);
 				}
+				CPPTypeObject = FindObjectGloballyWithRedirectors(*CPPType);
+			}
 
-				if(CPPTypeObject == nullptr)
-				{
-					const FString Message = FString::Printf(
-						TEXT("Pin '%s' is missing the CPPTypeObject for CPPType '%s'."),
-						*InPin->GetPinPath(), *InPin->GetCPPType());
-					FScriptExceptionHandler::Get().HandleException(ELogVerbosity::Error, *Message, *FString());
-					return false;
-				}
+			if(CPPTypeObject == nullptr)
+			{
+				const FString Message = FString::Printf(
+					TEXT("%s: Pin '%s' is missing the CPPTypeObject for CPPType '%s'."),
+					*InPin->GetPathName(), *InPin->GetPinPath(), *InPin->GetCPPType());
+				FScriptExceptionHandler::Get().HandleException(ELogVerbosity::Error, *Message, *FString());
+				return false;
 			}
+			
 			InPin->CPPTypeObject = CPPTypeObject;
 		}
 	}
@@ -12662,6 +12702,7 @@ bool URigVMController::EnsurePinValidity(URigVMPin* InPin, bool bRecursive)
 	return true;
 }
 
+
 void URigVMController::ValidatePin(URigVMPin* InPin)
 {
 	check(InPin);
