Control Rig: Fix adding execute context pin in function does not appear on both entry and return nodes

diff --git a/Engine/Source/Developer/RigVMDeveloper/Private/RigVMModel/RigVMController.cpp b/Engine/Source/Developer/RigVMDeveloper/Private/RigVMModel/RigVMController.cpp
--- a/Engine/Source/Developer/RigVMDeveloper/Private/RigVMModel/RigVMController.cpp
+++ b/Engine/Source/Developer/RigVMDeveloper/Private/RigVMModel/RigVMController.cpp
@@ -8234,23 +8234,17 @@ bool URigVMController::ChangeExposedPinType(const FName& InPinName, const FStrin
 		}
 		RemoveUnusedOrphanedPins(LibraryNode, true);
 	}
-	
 
 	if (URigVMFunctionEntryNode* EntryNode = Graph->GetEntryNode())
 	{
-		if (URigVMPin* EntryPin = EntryNode->FindPin(Pin->GetName()))
-		{
-			ChangePinType(EntryPin, InCPPType, InCPPTypeObjectPath, bSetupUndoRedo, bSetupOrphanPins);
-			RemoveUnusedOrphanedPins(EntryNode, true);
-		}
+		RepopulatePinsOnNode(EntryNode, true, true);
+		RemoveUnusedOrphanedPins(EntryNode, true);
 	}
+	
 	if (URigVMFunctionReturnNode* ReturnNode = Graph->GetReturnNode())
 	{
-		if (URigVMPin* ReturnPin = ReturnNode->FindPin(Pin->GetName()))
-		{
-			ChangePinType(ReturnPin, InCPPType, InCPPTypeObjectPath, bSetupUndoRedo, bSetupOrphanPins);
-			RemoveUnusedOrphanedPins(ReturnNode, true);
-		}
+		RepopulatePinsOnNode(ReturnNode, true, true);
+		RemoveUnusedOrphanedPins(ReturnNode, true);
 	}
 
 	if (URigVMFunctionLibrary* FunctionLibrary = Cast<URigVMFunctionLibrary>(LibraryNode->GetGraph()))
@@ -12497,7 +12491,7 @@ bool URigVMController::ChangePinType(URigVMPin* InPin, const FString& InCPPType,
 	// we might want to use GetPinInitialDefaultValue here for a better default value
 	InPin->DefaultValue = FString();
 
-	if (InPin->IsExecuteContext())
+	if (InPin->IsExecuteContext() && !InPin->GetNode()->IsA<URigVMFunctionEntryNode>() && !InPin->GetNode()->IsA<URigVMFunctionReturnNode>())
 	{
 		InPin->Direction = ERigVMPinDirection::IO;
 	}
