Control Rig: Fix check failed when copy/paste nodes with variable bindings

diff --git a/Engine/Source/Developer/RigVMDeveloper/Private/RigVMModel/RigVMController.cpp b/Engine/Source/Developer/RigVMDeveloper/Private/RigVMModel/RigVMController.cpp
--- a/Engine/Source/Developer/RigVMDeveloper/Private/RigVMModel/RigVMController.cpp
+++ b/Engine/Source/Developer/RigVMDeveloper/Private/RigVMModel/RigVMController.cpp
@@ -2720,13 +2720,31 @@ TArray<FName> URigVMController::ImportNodesFromText(const FString& InText, bool
 					
 					if (SourcePin && TargetPin)
 					{
+						// BreakAllLinks will unbind and destroy the injected variable node
+						// We need to rebind to recreate the variable node with the same name
+						bool bWasBinded = TargetPin->IsBoundToVariable();
+						FString VariableNodeName, BindingPath;
+						if (bWasBinded)
+						{
+							VariableNodeName = TargetPin->GetBoundVariableNode()->GetName();
+							BindingPath = TargetPin->GetBoundVariablePath();	
+						}
+						
 						BreakAllLinksRecursive(TargetPin, true, true, bSetupUndoRedo);
 						BreakAllLinks(TargetPin, true, bSetupUndoRedo);
 						BreakAllLinksRecursive(TargetPin, true, false, bSetupUndoRedo);
 
-						Graph->Links.Add(CreatedLink);
-						SourcePin->Links.Add(CreatedLink);
-						TargetPin->Links.Add(CreatedLink);
+						// recreate binding if needed
+						if (bWasBinded)
+						{
+							BindPinToVariable(TargetPin, BindingPath, bSetupUndoRedo, VariableNodeName);
+						}
+						else
+						{
+							Graph->Links.Add(CreatedLink);
+							SourcePin->Links.Add(CreatedLink);
+							TargetPin->Links.Add(CreatedLink);
+						}
 
 						if (bSetupUndoRedo)
 						{
@@ -6836,7 +6854,7 @@ bool URigVMController::BindPinToVariable(const FString& InPinPath, const FString
 	return bSuccess;
 }
 
-bool URigVMController::BindPinToVariable(URigVMPin* InPin, const FString& InNewBoundVariablePath, bool bSetupUndoRedo)
+bool URigVMController::BindPinToVariable(URigVMPin* InPin, const FString& InNewBoundVariablePath, bool bSetupUndoRedo, const FString& InVariableNodeName)
 {
 	if (!IsValidPinForGraph(InPin))
 	{
@@ -6919,7 +6937,7 @@ bool URigVMController::BindPinToVariable(URigVMPin* InPin, const FString& InNewB
 			FString CPPType;
 			UObject* CPPTypeObject;
 			RigVMTypeUtils::CPPTypeFromExternalVariable(Variable, CPPType, &CPPTypeObject);
-			VariableNode = AddVariableNode(*VariableName, CPPType, CPPTypeObject, true, FString(), FVector2D::ZeroVector, FString(), false);
+			VariableNode = AddVariableNode(*VariableName, CPPType, CPPTypeObject, true, FString(), FVector2D::ZeroVector, InVariableNodeName, false);
 		}
 		if (VariableNode == nullptr)
 		{
diff --git a/Engine/Source/Developer/RigVMDeveloper/Public/RigVMModel/RigVMController.h b/Engine/Source/Developer/RigVMDeveloper/Public/RigVMModel/RigVMController.h
--- a/Engine/Source/Developer/RigVMDeveloper/Public/RigVMModel/RigVMController.h
+++ b/Engine/Source/Developer/RigVMDeveloper/Public/RigVMModel/RigVMController.h
@@ -806,7 +806,7 @@ private:
 	URigVMPin* InsertArrayPin(URigVMPin* ArrayPin, int32 InIndex, const FString& InDefaultValue, bool bSetupUndoRedo);
 	bool RemovePin(URigVMPin* InPinToRemove, bool bSetupUndoRedo, bool bNotify);
 	FProperty* FindPropertyForPin(const FString& InPinPath);
-	bool BindPinToVariable(URigVMPin* InPin, const FString& InNewBoundVariablePath, bool bSetupUndoRedo);
+	bool BindPinToVariable(URigVMPin* InPin, const FString& InNewBoundVariablePath, bool bSetupUndoRedo, const FString& InVariableNodeName = FString());
 	bool UnbindPinFromVariable(URigVMPin* InPin, bool bSetupUndoRedo);
 	bool MakeBindingsFromVariableNode(URigVMVariableNode* InNode, bool bSetupUndoRedo);
 	bool PromotePinToVariable(URigVMPin* InPin, bool bCreateVariableNode, const FVector2D& InNodePosition, bool bSetupUndoRedo);
