Control Rig: Fix undo change type of local variable not relinking variable nodes

diff --git a/Engine/Source/Developer/RigVMDeveloper/Private/RigVMModel/RigVMController.cpp b/Engine/Source/Developer/RigVMDeveloper/Private/RigVMModel/RigVMController.cpp
--- a/Engine/Source/Developer/RigVMDeveloper/Private/RigVMModel/RigVMController.cpp
+++ b/Engine/Source/Developer/RigVMDeveloper/Private/RigVMModel/RigVMController.cpp
@@ -8765,7 +8765,7 @@ FRigVMGraphVariableDescription URigVMController::AddLocalVariable(const FName& I
 		{
 			if (VariableName == VariableNode->GetVariableName())
 			{
-				RefreshVariableNode(VariableNode->GetFName(), VariableName, InCPPType, InCPPTypeObject, bSetupUndoRedo);
+				RefreshVariableNode(VariableNode->GetFName(), VariableName, InCPPType, InCPPTypeObject, bSetupUndoRedo, false);
 			}
 		}
 	}
@@ -8905,7 +8905,7 @@ bool URigVMController::RemoveLocalVariable(const FName& InVariableName, bool bSe
 					{
 						if (VariablePin->GetDefaultValue() == VarNameStr)
 						{
-							RefreshVariableNode(VariableNode->GetFName(), ExternalVariableToSwitch.Name, ExternalVariableToSwitch.TypeName.ToString(), ExternalVariableToSwitch.TypeObject, bSetupUndoRedo);
+							RefreshVariableNode(VariableNode->GetFName(), ExternalVariableToSwitch.Name, ExternalVariableToSwitch.TypeName.ToString(), ExternalVariableToSwitch.TypeObject, bSetupUndoRedo, false);
 							continue;
 						}
 					}
@@ -9077,10 +9077,11 @@ bool URigVMController::SetLocalVariableType(const FName& InVariableName, const F
 	}
 
 	FRigVMControllerCompileBracketScope CompileScope(this);
+	FRigVMBaseAction BaseAction;
 	if (bSetupUndoRedo)
 	{
-		FRigVMInverseAction InverseAction;
-		InverseAction.Title = FString::Printf(TEXT("Change Local Variable type %s to %s"), *InVariableName.ToString(), *InCPPType);
+		BaseAction.Title = FString::Printf(TEXT("Change Local Variable type %s to %s"), *InVariableName.ToString(), *InCPPType);
+		ActionStack->BeginAction(BaseAction);
 
 		ActionStack->AddAction(FRigVMChangeLocalVariableTypeAction(LocalVariables[FoundIndex], InCPPType, InCPPTypeObject));
 	}	
@@ -9110,7 +9111,7 @@ bool URigVMController::SetLocalVariableType(const FName& InVariableName, const F
 			{
 				if (VariablePin->GetDefaultValue() == InVariableName.ToString())
 				{
-					RefreshVariableNode(Node->GetFName(), InVariableName, InCPPType, InCPPTypeObject, bSetupUndoRedo);
+					RefreshVariableNode(Node->GetFName(), InVariableName, InCPPType, InCPPTypeObject, bSetupUndoRedo, false);
 					continue;
 				}
 			}
@@ -9126,6 +9127,11 @@ bool URigVMController::SetLocalVariableType(const FName& InVariableName, const F
 		}
 	}
 
+	if (bSetupUndoRedo)
+	{
+		ActionStack->EndAction(BaseAction);
+	}
+
 	if (bPrintPythonCommand)
 	{
 		const FString GraphName = GetSanitizedGraphName(GetGraph()->GetGraphName());
diff --git a/Engine/Source/Developer/RigVMDeveloper/Private/RigVMModel/RigVMControllerActions.cpp b/Engine/Source/Developer/RigVMDeveloper/Private/RigVMModel/RigVMControllerActions.cpp
--- a/Engine/Source/Developer/RigVMDeveloper/Private/RigVMModel/RigVMControllerActions.cpp
+++ b/Engine/Source/Developer/RigVMDeveloper/Private/RigVMModel/RigVMControllerActions.cpp
@@ -32,7 +32,11 @@ FString FRigVMActionWrapper::ExportText()
 	FString ExportedText;
 	if (Data.Num() > 0 && ScriptStruct != nullptr)
 	{
-		ScriptStruct->ExportText(ExportedText, GetAction(), nullptr, nullptr, PPF_None, nullptr);
+		TArray<uint8, TAlignedHeapAllocator<16>> DefaultStructData;
+		DefaultStructData.AddZeroed(ScriptStruct->GetStructureSize());
+		ScriptStruct->InitializeDefaultValue(DefaultStructData.GetData());
+		
+		ScriptStruct->ExportText(ExportedText, GetAction(), DefaultStructData.GetData(), nullptr, PPF_None, nullptr);
 	}
 	return ExportedText;
 }
diff --git a/Engine/Source/Developer/RigVMDeveloper/Public/RigVMModel/RigVMControllerActions.h b/Engine/Source/Developer/RigVMDeveloper/Public/RigVMModel/RigVMControllerActions.h
--- a/Engine/Source/Developer/RigVMDeveloper/Public/RigVMModel/RigVMControllerActions.h
+++ b/Engine/Source/Developer/RigVMDeveloper/Public/RigVMModel/RigVMControllerActions.h
@@ -50,7 +50,12 @@ struct FRigVMActionKey
 		UScriptStruct* ScriptStruct = ActionType::StaticStruct();
 		FRigVMActionKey Key;
 		ScriptStructPath = ScriptStruct->GetPathName();
-		ScriptStruct->ExportText(ExportedText, &InAction, nullptr, nullptr, PPF_None, nullptr);
+		
+		TArray<uint8, TAlignedHeapAllocator<16>> DefaultStructData;
+		DefaultStructData.AddZeroed(ScriptStruct->GetStructureSize());
+		ScriptStruct->InitializeDefaultValue(DefaultStructData.GetData());
+		
+		ScriptStruct->ExportText(ExportedText, &InAction, DefaultStructData.GetData(), nullptr, PPF_None, nullptr);
 	}
 };
 
