Control Rig: fix binding types differ from linking types

diff --git a/Engine/Plugins/Animation/ControlRig/Source/ControlRigDeveloper/Private/Graph/ControlRigGraphSchema.cpp b/Engine/Plugins/Animation/ControlRig/Source/ControlRigDeveloper/Private/Graph/ControlRigGraphSchema.cpp
--- a/Engine/Plugins/Animation/ControlRig/Source/ControlRigDeveloper/Private/Graph/ControlRigGraphSchema.cpp
+++ b/Engine/Plugins/Animation/ControlRig/Source/ControlRigDeveloper/Private/Graph/ControlRigGraphSchema.cpp
@@ -1367,7 +1367,7 @@ UEdGraphPin* UControlRigGraphSchema::DropPinOnNode(UEdGraphNode* InTargetNode, c
 							FString TypeName = ExternalVar.TypeName.ToString();
 							if (ExternalVar.bIsArray)
 							{
-								TypeName = FString::Printf(URigVMController::TArrayTemplate, *TypeName);
+								TypeName = RigVMTypeUtils::ArrayTypeFromBaseType(*TypeName);
 							}
 							FName TypeObjectPathName = NAME_None;
 							if (ExternalVar.TypeObject)
diff --git a/Engine/Source/Developer/RigVMDeveloper/Private/RigVMCompiler/RigVMCompiler.cpp b/Engine/Source/Developer/RigVMDeveloper/Private/RigVMCompiler/RigVMCompiler.cpp
--- a/Engine/Source/Developer/RigVMDeveloper/Private/RigVMCompiler/RigVMCompiler.cpp
+++ b/Engine/Source/Developer/RigVMDeveloper/Private/RigVMCompiler/RigVMCompiler.cpp
@@ -2794,7 +2794,7 @@ FRigVMOperand URigVMCompiler::FindOrAddRegister(const FRigVMVarExprAST* InVarExp
 		if(bIsDebugValue)
 		{
 			// debug values are always stored as arrays
-			CPPType = FString::Printf(URigVMController::TArrayTemplate, *CPPType);
+			CPPType = RigVMTypeUtils::ArrayTypeFromBaseType(CPPType);
 			JoinedDefaultValue = URigVMPin::GetDefaultValueForArray({ JoinedDefaultValue });
 		}
 		else if(Pin->GetDirection() == ERigVMPinDirection::Hidden && Pin->GetNode()->IsA<URigVMUnitNode>())
@@ -2805,7 +2805,7 @@ FRigVMOperand URigVMCompiler::FindOrAddRegister(const FRigVMVarExprAST* InVarExp
 
 			if (!Property->HasMetaData(FRigVMStruct::SingletonMetaName))
 			{
-				CPPType = FString::Printf(URigVMController::TArrayTemplate, *CPPType);
+				CPPType = RigVMTypeUtils::ArrayTypeFromBaseType(CPPType);
 				JoinedDefaultValue = URigVMPin::GetDefaultValueForArray({ JoinedDefaultValue });
 			}
 		}
diff --git a/Engine/Source/Developer/RigVMDeveloper/Private/RigVMModel/RigVMController.cpp b/Engine/Source/Developer/RigVMDeveloper/Private/RigVMModel/RigVMController.cpp
--- a/Engine/Source/Developer/RigVMDeveloper/Private/RigVMModel/RigVMController.cpp
+++ b/Engine/Source/Developer/RigVMDeveloper/Private/RigVMModel/RigVMController.cpp
@@ -30,10 +30,6 @@
 #endif
 
 TMap<URigVMController::FControlRigStructPinRedirectorKey, FString> URigVMController::PinPathCoreRedirectors;
-const TCHAR URigVMController::TArrayPrefix[] = TEXT("TArray<");
-const TCHAR URigVMController::TObjectPtrPrefix[] = TEXT("TObjectPtr<");
-const TCHAR URigVMController::TArrayTemplate[] = TEXT("TArray<%s>");
-const TCHAR URigVMController::TObjectPtrTemplate[] = TEXT("TObjectPtr<%s%s>");
 
 FRigVMControllerCompileBracketScope::FRigVMControllerCompileBracketScope(URigVMController* InController)
 : Graph(nullptr), bSuspendNotifications(InController->bSuspendNotifications)
@@ -1038,9 +1034,9 @@ URigVMVariableNode* URigVMController::AddVariableNode(const FName& InVariableNam
 		ValuePin->CPPTypeObject = ExternalVariable.TypeObject;
 		ValuePin->bIsDynamicArray = ExternalVariable.bIsArray;
 
-		if(ValuePin->bIsDynamicArray && !ValuePin->CPPType.StartsWith(TArrayPrefix))
+		if(ValuePin->bIsDynamicArray && !RigVMTypeUtils::IsArrayType(ValuePin->CPPType))
 		{
-			ValuePin->CPPType = FString::Printf(TArrayTemplate, *ValuePin->CPPType);
+			ValuePin->CPPType = RigVMTypeUtils::ArrayTypeFromBaseType(*ValuePin->CPPType);
 		}
 	}
 	else
@@ -9835,7 +9831,7 @@ URigVMSelectNode* URigVMController::AddSelectNode(const FString& InCPPType, cons
 	AddNodePin(Node, IndexPin);
 
 	URigVMPin* ValuePin = NewObject<URigVMPin>(Node, *URigVMSelectNode::ValueName);
-	ValuePin->CPPType = FString::Printf(TArrayTemplate, *CPPType);
+	ValuePin->CPPType = RigVMTypeUtils::ArrayTypeFromBaseType(CPPType);
 	ValuePin->CPPTypeObject = CPPTypeObject;
 	ValuePin->CPPTypeObjectPath = InCPPTypeObjectPath;
 	ValuePin->Direction = ERigVMPinDirection::Input;
@@ -10100,9 +10096,9 @@ URigVMArrayNode* URigVMController::AddArrayNode(ERigVMOpCode InOpCode, const FSt
 			{
 				Pin->CPPTypeObjectPath = *Pin->CPPTypeObject->GetPathName();
 			}
-			if(bIsArray && !Pin->CPPType.StartsWith(TArrayPrefix))
+			if(bIsArray && !RigVMTypeUtils::IsArrayType(Pin->CPPType))
 			{
-				Pin->CPPType = FString::Printf(TArrayTemplate, *Pin->CPPType);
+				Pin->CPPType = RigVMTypeUtils::ArrayTypeFromBaseType(*Pin->CPPType);
 			}
 			Pin->Direction = InDirection;
 			Pin->bIsDynamicArray = bIsArray;
@@ -12198,7 +12194,7 @@ FString URigVMController::PostProcessCPPType(const FString& InCPPType, UObject*
 	
 	if (const UClass* Class = Cast<UClass>(InCPPTypeObject))
 	{
-		CPPType = FString::Printf(TObjectPtrTemplate, Class->GetPrefixCPP(), *Class->GetName());
+		CPPType = FString::Printf(RigVMTypeUtils::TObjectPtrTemplate, Class->GetPrefixCPP(), *Class->GetName());
 	}
 	else if (const UScriptStruct* ScriptStruct = Cast<UScriptStruct>(InCPPTypeObject))
 	{
@@ -12372,7 +12368,7 @@ void URigVMController::ResolveUnknownTypePin(URigVMPin* InPinToResolve, const UR
 	FString ResolvedCPPType = InTemplatePin->IsArray() ? InTemplatePin->GetArrayElementCppType() : InTemplatePin->GetCPPType();
 	if(InPinToResolve->IsArray())
 	{
-		ResolvedCPPType = FString::Printf(TArrayTemplate, *ResolvedCPPType);
+		ResolvedCPPType = RigVMTypeUtils::ArrayTypeFromBaseType(ResolvedCPPType);
 	}
 
 	if(!ChangePinType(InPinToResolve, ResolvedCPPType, InTemplatePin->GetCPPTypeObject(), bSetupUndoRedo, false, false, false))
diff --git a/Engine/Source/Developer/RigVMDeveloper/Private/RigVMModel/RigVMPin.cpp b/Engine/Source/Developer/RigVMDeveloper/Private/RigVMModel/RigVMPin.cpp
--- a/Engine/Source/Developer/RigVMDeveloper/Private/RigVMModel/RigVMPin.cpp
+++ b/Engine/Source/Developer/RigVMDeveloper/Private/RigVMModel/RigVMPin.cpp
@@ -451,12 +451,12 @@ bool URigVMPin::IsStructMember() const
 
 bool URigVMPin::IsUObject() const
 {
-	return CPPType.StartsWith(URigVMController::TObjectPtrPrefix);
+	return RigVMTypeUtils::IsUObjectType(CPPType);
 }
 
 bool URigVMPin::IsArray() const
 {
-	return CPPType.StartsWith(URigVMController::TArrayPrefix);
+	return RigVMTypeUtils::IsArrayType(CPPType);
 }
 
 bool URigVMPin::IsArrayElement() const
@@ -691,8 +691,8 @@ bool URigVMPin::IsValidDefaultValue(const FString& InDefaultValue) const
 	}
 
 	FString BaseCPPType = GetCPPType()
-		.Replace(URigVMController::TArrayPrefix, TEXT(""))
-		.Replace(URigVMController::TObjectPtrPrefix, TEXT(""))
+		.Replace(RigVMTypeUtils::TArrayPrefix, TEXT(""))
+		.Replace(RigVMTypeUtils::TObjectPtrPrefix, TEXT(""))
 		.Replace(TEXT(">"), TEXT(""));
 
 	for (const FString& Value : DefaultValues)
@@ -821,8 +821,8 @@ FString URigVMPin::ClampDefaultValueFromMetaData(const FString& InDefaultValue)
 		
 
 		FString BaseCPPType = GetCPPType()
-			.Replace(URigVMController::TArrayPrefix, TEXT(""))
-			.Replace(URigVMController::TObjectPtrPrefix, TEXT(""))
+			.Replace(RigVMTypeUtils::TArrayPrefix, TEXT(""))
+			.Replace(RigVMTypeUtils::TObjectPtrPrefix, TEXT(""))
 			.Replace(TEXT(">"), TEXT(""));
 
 		RetVals.SetNumZeroed(DefaultValues.Num());
@@ -1102,23 +1102,8 @@ bool URigVMPin::CanBeBoundToVariable(const FRigVMExternalVariable& InExternalVar
 	}
 #endif
 
-	if (GetCPPTypeObject() != nullptr)
-	{
-		if (GetCPPTypeObject() != ExternalCPPTypeObject)
-		{
-			return false;
-		}
-	}
-	else
-	{
-		const FString CPPBaseType = IsArray() ? GetArrayElementCppType() : GetCPPType();
-		if (CPPBaseType != ExternalCPPType.ToString())
-		{
-			return false;
-		}
-	}
-
-	return true;
+	const FString CPPBaseType = IsArray() ? GetArrayElementCppType() : GetCPPType();
+	return RigVMTypeUtils::AreCompatible(CPPBaseType, GetCPPTypeObject(), ExternalCPPType.ToString(), ExternalCPPTypeObject);
 }
 
 bool URigVMPin::ShowInDetailsPanelOnly() const
@@ -1529,13 +1514,7 @@ bool URigVMPin::CanLink(URigVMPin* InSourcePin, URigVMPin* InTargetPin, FString*
 			(InSourcePin->CPPType == Double && InTargetPin->CPPType == Float))
 #else
 
-		static const FString FloatArray = TEXT("TArray<float>");
-		static const FString DoubleArray = TEXT("TArray<double>");
-		
-		if((InSourcePin->CPPType == Float && InTargetPin->CPPType == Double) ||
-			(InSourcePin->CPPType == Double && InTargetPin->CPPType == Float) ||
-			(InSourcePin->CPPType == FloatArray && InTargetPin->CPPType == DoubleArray) ||
-			(InSourcePin->CPPType == DoubleArray && InTargetPin->CPPType == FloatArray))
+		if (RigVMTypeUtils::AreCompatible(InSourcePin->CPPType, InSourcePin->CPPTypeObject, InTargetPin->CPPType, InTargetPin->CPPTypeObject))
 #endif
 		{
 			bCPPTypesDiffer = false;
diff --git a/Engine/Source/Developer/RigVMDeveloper/Private/RigVMTypeUtils.cpp b/Engine/Source/Developer/RigVMDeveloper/Private/RigVMTypeUtils.cpp
--- a/Engine/Source/Developer/RigVMDeveloper/Private/RigVMTypeUtils.cpp
+++ b/Engine/Source/Developer/RigVMDeveloper/Private/RigVMTypeUtils.cpp
@@ -21,7 +21,7 @@ FRigVMExternalVariable RigVMTypeUtils::ExternalVariableFromRigVMVariableDescript
 	FRigVMExternalVariable ExternalVariable;
 	ExternalVariable.Name = InVariableDescription.Name;
 
-	if (InVariableDescription.CPPType.StartsWith(TEXT("TArray<")))
+	if (IsArrayType(InVariableDescription.CPPType))
 	{
 		ExternalVariable.bIsArray = true;
 		ExternalVariable.TypeName = *InVariableDescription.CPPType.Mid(7, InVariableDescription.CPPType.Len() - 8);
@@ -258,6 +258,58 @@ FRigVMExternalVariable RigVMTypeUtils::ExternalVariableFromCPPType(const FName&
 	return Variable;
 }
 
+FEdGraphPinType RigVMTypeUtils::PinTypeFromCPPType(const FString& InCPPType, UObject* InCPPTypeObject)
+{
+	FEdGraphPinType PinType;
+	PinType.ResetToDefaults();
+	PinType.PinCategory = NAME_None;
+
+	FString BaseCPPType = InCPPType;
+	PinType.ContainerType = EPinContainerType::None;
+	if (RigVMTypeUtils::IsArrayType(InCPPType))
+	{
+		BaseCPPType = RigVMTypeUtils::BaseTypeFromArrayType(InCPPType);
+		PinType.ContainerType = EPinContainerType::Array;
+	}
+
+	if (BaseCPPType == TEXT("bool"))
+	{
+		PinType.PinCategory = UEdGraphSchema_K2::PC_Boolean;
+	}
+	else if (BaseCPPType == TEXT("int32"))
+	{
+		PinType.PinCategory = UEdGraphSchema_K2::PC_Int;
+	}
+	else if (BaseCPPType == TEXT("float"))
+	{
+		PinType.PinCategory = UEdGraphSchema_K2::PC_Float;
+	}
+	else if (BaseCPPType == TEXT("double"))
+	{
+		PinType.PinCategory = UEdGraphSchema_K2::PC_Double;
+	}
+	else if (BaseCPPType == TEXT("FName"))
+	{
+		PinType.PinCategory = UEdGraphSchema_K2::PC_Name;
+	}
+	else if (BaseCPPType == TEXT("FString"))
+	{
+		PinType.PinCategory = UEdGraphSchema_K2::PC_String;
+	}
+	else if (Cast<UScriptStruct>(InCPPTypeObject))
+	{
+		PinType.PinCategory = UEdGraphSchema_K2::PC_Struct;
+		PinType.PinSubCategoryObject = InCPPTypeObject;
+	}
+	else if (Cast<UEnum>(InCPPTypeObject))
+	{
+		PinType.PinCategory = UEdGraphSchema_K2::PC_Byte;
+		PinType.PinSubCategoryObject = InCPPTypeObject;
+	}
+
+	return PinType;
+}
+
 FEdGraphPinType RigVMTypeUtils::PinTypeFromExternalVariable(const FRigVMExternalVariable& InExternalVariable)
 {
 	FEdGraphPinType PinType;
@@ -319,11 +371,10 @@ FEdGraphPinType RigVMTypeUtils::PinTypeFromRigVMVariableDescription(
 	PinType.PinCategory = NAME_None;
 
 	FString CurrentCPPType = InVariableDescription.CPPType;
-	if (CurrentCPPType.StartsWith(TEXT("TArray<")))
+	if (IsArrayType(CurrentCPPType))
 	{
 		PinType.ContainerType = EPinContainerType::Array;
-		CurrentCPPType.RemoveFromStart(TEXT("TArray<"));
-		CurrentCPPType.RemoveFromEnd(TEXT(">"));
+		CurrentCPPType = BaseTypeFromArrayType(CurrentCPPType);
 	}
 	else
 	{
@@ -579,6 +630,11 @@ bool RigVMTypeUtils::CPPTypeFromExternalVariable(const FRigVMExternalVariable& I
 	return true;
 }
 
+bool RigVMTypeUtils::AreCompatible(const FString& InCPPTypeA, UObject* InCPPTypeObjectA, const FString& InCPPTypeB,	UObject* InCPPTypeObjectB)
+{
+	return AreCompatible(PinTypeFromCPPType(InCPPTypeA, InCPPTypeObjectA), PinTypeFromCPPType(InCPPTypeB, InCPPTypeObjectB));
+}
+
 bool RigVMTypeUtils::AreCompatible(const FRigVMExternalVariable& InTypeA, const FRigVMExternalVariable& InTypeB, const FString& InSegmentPathA, const FString& InSegmentPathB)
 {
 	return AreCompatible(PinTypeFromExternalVariable(InTypeA), PinTypeFromExternalVariable(InTypeB), InSegmentPathA, InSegmentPathB);
@@ -588,6 +644,15 @@ bool RigVMTypeUtils::AreCompatible(const FEdGraphPinType& InTypeA, const FEdGrap
 {
 	FEdGraphPinType SubPinTypeA = SubPinType(InTypeA, InSegmentPathA);
 	FEdGraphPinType SubPinTypeB = SubPinType(InTypeB, InSegmentPathB);
+
+	// We allow connectiongs between floats and doubles, while EdGraphSchema_K2 does not
+	// Every other case is evaluated by UEdGraphSchema_K2::ArePinTypesCompatible
+	if (SubPinTypeA.ContainerType == SubPinTypeB.ContainerType &&
+		(SubPinTypeA.PinCategory == TEXT("float") && SubPinTypeB.PinCategory == TEXT("double") ||
+		 SubPinTypeA.PinCategory == TEXT("double") && SubPinTypeB.PinCategory == TEXT("float")))
+	{
+		return true;
+	}
 	
 	return UEdGraphSchema_K2::StaticClass()->GetDefaultObject<UEdGraphSchema_K2>()->ArePinTypesCompatible(SubPinTypeA, SubPinTypeB);	
 }
diff --git a/Engine/Source/Developer/RigVMDeveloper/Public/RigVMModel/RigVMController.h b/Engine/Source/Developer/RigVMDeveloper/Public/RigVMModel/RigVMController.h
--- a/Engine/Source/Developer/RigVMDeveloper/Public/RigVMModel/RigVMController.h
+++ b/Engine/Source/Developer/RigVMDeveloper/Public/RigVMModel/RigVMController.h
@@ -897,12 +897,6 @@ public:
 	static FString GetSanitizedPinPath(const FString& InName);
 	static void SanitizeName(FString& InOutName, bool bAllowPeriod, bool bAllowSpace);
 
-
-	static const TCHAR TArrayPrefix[];
-	static const TCHAR TObjectPtrPrefix[];
-	static const TCHAR TArrayTemplate[];
-	static const TCHAR TObjectPtrTemplate[];
-
 private: 
 	UPROPERTY(transient)
 	TArray<TObjectPtr<URigVMGraph>> Graphs;
diff --git a/Engine/Source/Developer/RigVMDeveloper/Public/RigVMTypeUtils.h b/Engine/Source/Developer/RigVMDeveloper/Public/RigVMTypeUtils.h
--- a/Engine/Source/Developer/RigVMDeveloper/Public/RigVMTypeUtils.h
+++ b/Engine/Source/Developer/RigVMDeveloper/Public/RigVMTypeUtils.h
@@ -10,15 +10,21 @@ struct FRigVMGraphVariableDescription;
 
 namespace RigVMTypeUtils
 {
+
+	const TCHAR TArrayPrefix[] = TEXT("TArray<");
+	const TCHAR TObjectPtrPrefix[] = TEXT("TObjectPtr<");
+	const TCHAR TArrayTemplate[] = TEXT("TArray<%s>");
+	const TCHAR TObjectPtrTemplate[] = TEXT("TObjectPtr<%s%s>");
+	
 	// Returns true if the type specified is an array
 	FORCEINLINE bool IsArrayType(const FString& InCPPType)
 	{
-		return InCPPType.StartsWith(TEXT("TArray<"));
+		return InCPPType.StartsWith(TArrayPrefix);
 	}
 
 	FORCEINLINE FString ArrayTypeFromBaseType(const FString& InCPPType)
 	{
-		return "TArray<" + InCPPType + ">";
+		return FString::Printf(TArrayTemplate, *InCPPType);
 	}
 
 	FORCEINLINE FString BaseTypeFromArrayType(const FString& InCPPType)
@@ -26,6 +32,11 @@ namespace RigVMTypeUtils
 		return InCPPType.RightChop(7).LeftChop(1);
 	}
 
+	FORCEINLINE bool IsUObjectType(const FString& InCPPType)
+	{
+		return InCPPType.StartsWith(TObjectPtrPrefix);
+	}
+
 	
 #if WITH_EDITOR
 	RIGVMDEVELOPER_API FRigVMExternalVariable ExternalVariableFromBPVariableDescription(const FBPVariableDescription& InVariableDescription);
@@ -38,6 +49,8 @@ namespace RigVMTypeUtils
 
 	RIGVMDEVELOPER_API FRigVMExternalVariable ExternalVariableFromCPPType(const FName& InName, const FString& InCPPType, UObject* InCPPTypeObject, bool bInPublic = false, bool bInReadonly = false);
 
+	RIGVMDEVELOPER_API FEdGraphPinType PinTypeFromCPPType(const FString& InCPPType, UObject* InCPPTypeObject);
+
 	RIGVMDEVELOPER_API FEdGraphPinType PinTypeFromExternalVariable(const FRigVMExternalVariable& InExternalVariable);
 
 	RIGVMDEVELOPER_API FEdGraphPinType PinTypeFromRigVMVariableDescription(const FRigVMGraphVariableDescription& InVariableDescription);
@@ -50,6 +63,8 @@ namespace RigVMTypeUtils
 
 	RIGVMDEVELOPER_API bool CPPTypeFromExternalVariable(const FRigVMExternalVariable& InExternalVariable, FString& OutCPPType, UObject** OutCPPTypeObject);
 
+	RIGVMDEVELOPER_API bool AreCompatible(const FString& InCPPTypeA, UObject* InCPPTypeObjectA, const FString& InCPPTypeB, UObject* InCPPTypeObjectB);
+
 	RIGVMDEVELOPER_API bool AreCompatible(const FRigVMExternalVariable& InTypeA, const FRigVMExternalVariable& InTypeB, const FString& InSegmentPathA = FString(), const FString& InSegmentPathB = FString());
 	
 	RIGVMDEVELOPER_API bool AreCompatible(const FEdGraphPinType& InTypeA, const FEdGraphPinType& InTypeB, const FString& InSegmentPathA = FString(), const FString& InSegmentPathB = FString());
