Control Rig: Fix compilation error

diff --git a/Engine/Source/Developer/RigVMDeveloper/Private/RigVMTypeUtils.cpp b/Engine/Source/Developer/RigVMDeveloper/Private/RigVMTypeUtils.cpp
--- a/Engine/Source/Developer/RigVMDeveloper/Private/RigVMTypeUtils.cpp
+++ b/Engine/Source/Developer/RigVMDeveloper/Private/RigVMTypeUtils.cpp
@@ -647,11 +647,13 @@ bool RigVMTypeUtils::AreCompatible(const FEdGraphPinType& InTypeA, const FEdGrap
 
 	// We allow connectiongs between floats and doubles, while EdGraphSchema_K2 does not
 	// Every other case is evaluated by UEdGraphSchema_K2::ArePinTypesCompatible
-	if (SubPinTypeA.ContainerType == SubPinTypeB.ContainerType &&
-		(SubPinTypeA.PinCategory == TEXT("float") && SubPinTypeB.PinCategory == TEXT("double") ||
-		 SubPinTypeA.PinCategory == TEXT("double") && SubPinTypeB.PinCategory == TEXT("float")))
+	if (SubPinTypeA.ContainerType == SubPinTypeB.ContainerType)
 	{
-		return true;
+		if ((SubPinTypeA.PinCategory == TEXT("float") && SubPinTypeB.PinCategory == TEXT("double")) ||
+			(SubPinTypeA.PinCategory == TEXT("double") && SubPinTypeB.PinCategory == TEXT("float")))
+		{
+			return true;
+		}
 	}
 	
 	return UEdGraphSchema_K2::StaticClass()->GetDefaultObject<UEdGraphSchema_K2>()->ArePinTypesCompatible(SubPinTypeA, SubPinTypeB);	
