Control Rig: fix eject last node

diff --git a/Engine/Plugins/Animation/ControlRig/Source/ControlRigEditor/Private/ControlRigEditorModule.cpp b/Engine/Plugins/Animation/ControlRig/Source/ControlRigEditor/Private/ControlRigEditorModule.cpp
--- a/Engine/Plugins/Animation/ControlRig/Source/ControlRigEditor/Private/ControlRigEditorModule.cpp
+++ b/Engine/Plugins/Animation/ControlRig/Source/ControlRigEditor/Private/ControlRigEditorModule.cpp
@@ -1372,7 +1372,10 @@ void FControlRigEditorModule::GetContextMenuActions(const UControlRigGraphSchema
 								LOCTEXT("EjectLastNode_Tooltip", "Eject the last injected node"),
 								FSlateIcon(),
 								FUIAction(FExecuteAction::CreateLambda([Controller, ModelPin]() {
-									Controller->EjectNodeFromPin(ModelPin->GetPinPath(), true, true);
+									Controller->OpenUndoBracket(TEXT("Eject node from pin"));
+									URigVMNode* Node = Controller->EjectNodeFromPin(ModelPin->GetPinPath(), true, true);
+									Controller->SelectNode(Node, true, true, true);
+									Controller->CloseUndoBracket();
 								})
 							));
 						}
@@ -1430,6 +1433,7 @@ void FControlRigEditorModule::GetContextMenuActions(const UControlRigGraphSchema
 										LOCTEXT("AddAlphaInterp_Tooltip", "Injects an interpolate node"),
 										FSlateIcon(),
 										FUIAction(FExecuteAction::CreateLambda([Controller, InGraphPin, ModelPin, ScriptStruct]() {
+											Controller->OpenUndoBracket(TEXT("Add injected node"));
 											URigVMInjectionInfo* Injection = Controller->AddInjectedNode(ModelPin->GetPinPath(), ModelPin->GetDirection() != ERigVMPinDirection::Output, ScriptStruct, FRigUnit::GetMethodName(), TEXT("Value"), TEXT("Result"), FString(), true, true);
 											if (Injection)
 											{
@@ -1437,6 +1441,7 @@ void FControlRigEditorModule::GetContextMenuActions(const UControlRigGraphSchema
 												NodeNames.Add(Injection->Node->GetFName());
 												Controller->SetNodeSelection(NodeNames);
 											}
+											Controller->CloseUndoBracket();
 										})
 									));
 								}
@@ -1459,7 +1464,7 @@ void FControlRigEditorModule::GetContextMenuActions(const UControlRigGraphSchema
 										LOCTEXT("RemoveAlphaInterp_Tooltip", "Removes the interpolate node"),
 										FSlateIcon(),
 										FUIAction(FExecuteAction::CreateLambda([Controller, InGraphPin, ModelPin, InterpNode]() {
-											Controller->RemoveNodeByName(InterpNode->GetFName(), true, false, true);
+											Controller->RemoveInjectedNode(ModelPin->GetPinPath(), ModelPin->GetDirection() != ERigVMPinDirection::Output, true);
 										})
 									));
 								}
diff --git a/Engine/Source/Developer/RigVMDeveloper/Private/RigVMModel/RigVMController.cpp b/Engine/Source/Developer/RigVMDeveloper/Private/RigVMModel/RigVMController.cpp
--- a/Engine/Source/Developer/RigVMDeveloper/Private/RigVMModel/RigVMController.cpp
+++ b/Engine/Source/Developer/RigVMDeveloper/Private/RigVMModel/RigVMController.cpp
@@ -2047,6 +2047,10 @@ URigVMInjectionInfo* URigVMController::AddInjectedNode(const FString& InPinPath,
 		return nullptr;
 	}
 
+	// 1.- Create unit node
+	// 2.- Rewire links
+	// 3.- Inject node into pin
+
 	FRigVMControllerCompileBracketScope CompileScope(this);
 	FRigVMBaseAction Action;
 	if (bSetupUndoRedo)
@@ -2055,112 +2059,109 @@ URigVMInjectionInfo* URigVMController::AddInjectedNode(const FString& InPinPath,
 		ActionStack->BeginAction(Action);
 	}
 
+	// 1.- Create unit node
 	URigVMUnitNode* UnitNode = nullptr;
+	URigVMPin* InputPin = nullptr;
+	URigVMPin* OutputPin = nullptr;
 	{
-		TGuardValue<bool> GuardNotifications(bSuspendNotifications, true);
-		UnitNode = AddUnitNode(InScriptStruct, InMethodName, FVector2D::ZeroVector, InNodeName, false);
-	}
-	if (UnitNode == nullptr)
-	{
-		if (bSetupUndoRedo)
 		{
-			ActionStack->CancelAction(Action);
+			TGuardValue<bool> GuardNotifications(bSuspendNotifications, true);
+			UnitNode = AddUnitNode(InScriptStruct, InMethodName, FVector2D::ZeroVector, InNodeName, bSetupUndoRedo);
 		}
-		return nullptr;
-	}
-	else if (UnitNode->IsMutable())
-	{
-		ReportErrorf(TEXT("Injected node %s is mutable."), *InScriptStruct->GetName());
-		RemoveNode(UnitNode, false);
-		if (bSetupUndoRedo)
+		if (UnitNode == nullptr)
 		{
-			ActionStack->CancelAction(Action);
+			if (bSetupUndoRedo)
+			{
+				ActionStack->CancelAction(Action);
+			}
+			return nullptr;
+		}
+		else if (UnitNode->IsMutable())
+		{
+			ReportErrorf(TEXT("Injected node %s is mutable."), *InScriptStruct->GetName());
+			RemoveNode(UnitNode, false);
+			if (bSetupUndoRedo)
+			{
+				ActionStack->CancelAction(Action);
+			}
+			return nullptr;
 		}
-		return nullptr;
-	}
 
-	URigVMPin* InputPin = UnitNode->FindPin(InInputPinName.ToString());
-	check(InputPin);
-	URigVMPin* OutputPin = UnitNode->FindPin(InOutputPinName.ToString());
-	check(OutputPin);
+		InputPin = UnitNode->FindPin(InInputPinName.ToString());
+		check(InputPin);
+		OutputPin = UnitNode->FindPin(InOutputPinName.ToString());
+		check(OutputPin);
 
-	if (InputPin->GetCPPType() != OutputPin->GetCPPType() ||
-		InputPin->IsArray() != OutputPin->IsArray())
-	{
-		ReportErrorf(TEXT("Injected node %s is using incompatible input and output pins."), *InScriptStruct->GetName());
-		RemoveNode(UnitNode, false);
-		if (bSetupUndoRedo)
+		if (InputPin->GetCPPType() != OutputPin->GetCPPType() ||
+			InputPin->IsArray() != OutputPin->IsArray())
 		{
-			ActionStack->CancelAction(Action);
+			ReportErrorf(TEXT("Injected node %s is using incompatible input and output pins."), *InScriptStruct->GetName());
+			RemoveNode(UnitNode, false);
+			if (bSetupUndoRedo)
+			{
+				ActionStack->CancelAction(Action);
+			}
+			return nullptr;
 		}
-		return nullptr;
-	}
 
-	if (InputPin->GetCPPType() != Pin->GetCPPType() ||
-		InputPin->IsArray() != Pin->IsArray())
-	{
-		ReportErrorf(TEXT("Injected node %s is using incompatible pin."), *InScriptStruct->GetName());
-		RemoveNode(UnitNode, false);
-		if (bSetupUndoRedo)
+		if (InputPin->GetCPPType() != Pin->GetCPPType() ||
+			InputPin->IsArray() != Pin->IsArray())
 		{
-			ActionStack->CancelAction(Action);
+			ReportErrorf(TEXT("Injected node %s is using incompatible pin."), *InScriptStruct->GetName());
+			RemoveNode(UnitNode, false);
+			if (bSetupUndoRedo)
+			{
+				ActionStack->CancelAction(Action);
+			}
+			return nullptr;
 		}
-		return nullptr;
-	}
-
-	URigVMInjectionInfo* InjectionInfo = NewObject<URigVMInjectionInfo>(Pin);
-
-	// re-parent the unit node to be under the injection info
-	RenameObject(UnitNode, nullptr, InjectionInfo);
-
-	InjectionInfo->Node = UnitNode;
-	InjectionInfo->bInjectedAsInput = bAsInput;
-	InjectionInfo->InputPin = InputPin;
-	InjectionInfo->OutputPin = OutputPin;
-
-	if (bSetupUndoRedo)
-	{
-		ActionStack->AddAction(FRigVMAddInjectedNodeAction(InjectionInfo));
 	}
 
-	URigVMPin* PreviousInputPin = Pin;
-	URigVMPin* PreviousOutputPin = Pin;
-	if (Pin->InjectionInfos.Num() > 0)
-	{
-		PreviousInputPin = Pin->InjectionInfos.Last()->InputPin;
-		PreviousOutputPin = Pin->InjectionInfos.Last()->OutputPin;
-	}
-	Pin->InjectionInfos.Add(InjectionInfo);
-
-	Notify(ERigVMGraphNotifType::NodeAdded, UnitNode);
-
-	// now update all of the links
-	if (bAsInput)
+	// 2.- Rewire links
+	TArray<URigVMLink*> NewLinks;
 	{
-		FString PinDefaultValue = PreviousInputPin->GetDefaultValue();
-		if (!PinDefaultValue.IsEmpty())
+		URigVMPin* PreviousInputPin = Pin;
+		URigVMPin* PreviousOutputPin = Pin;
+		if (Pin->InjectionInfos.Num() > 0)
 		{
-			SetPinDefaultValue(InjectionInfo->InputPin, PinDefaultValue, true, false, false);
+			PreviousInputPin = Pin->InjectionInfos.Last()->InputPin;
+			PreviousOutputPin = Pin->InjectionInfos.Last()->OutputPin;
 		}
-		TArray<URigVMLink*> Links = PreviousInputPin->GetSourceLinks(true /* recursive */);
-		BreakAllLinks(PreviousInputPin, true, false);
-		AddLink(InjectionInfo->OutputPin, PreviousInputPin, false);
-		if (Links.Num() > 0)
+		if (bAsInput)
 		{
-			RewireLinks(PreviousInputPin, InjectionInfo->InputPin, true, false, Links);
+			FString PinDefaultValue = PreviousInputPin->GetDefaultValue();
+			if (!PinDefaultValue.IsEmpty())
+			{
+				SetPinDefaultValue(InputPin, PinDefaultValue, true, bSetupUndoRedo, false);
+			}
+			TArray<URigVMLink*> Links = PreviousInputPin->GetSourceLinks(true /* recursive */);
+			if (Links.Num() > 0)
+			{
+				RewireLinks(PreviousInputPin, InputPin, true, bSetupUndoRedo, Links);
+				NewLinks = InputPin->GetSourceLinks();
+			}
+			AddLink(OutputPin, PreviousInputPin, bSetupUndoRedo);
 		}
-	}
-	else
-	{
-		TArray<URigVMLink*> Links = PreviousOutputPin->GetTargetLinks(true /* recursive */);
-		BreakAllLinks(PreviousOutputPin, false, false);
-		AddLink(PreviousOutputPin, InjectionInfo->InputPin, false);
-		if (Links.Num() > 0)
+		else
 		{
-			RewireLinks(PreviousOutputPin, InjectionInfo->OutputPin, false, false, Links);
+			TArray<URigVMLink*> Links = PreviousOutputPin->GetTargetLinks(true /* recursive */);
+			if (Links.Num() > 0)
+			{
+				RewireLinks(PreviousOutputPin, OutputPin, false, bSetupUndoRedo, Links);
+				NewLinks = OutputPin->GetTargetLinks();
+			}
+			AddLink(PreviousOutputPin, InputPin, bSetupUndoRedo);
 		}
 	}
 
+	// 3.- Inject node into pin
+	URigVMInjectionInfo* InjectionInfo = InjectNodeIntoPin(InPinPath, bAsInput, InInputPinName, InOutputPinName, bSetupUndoRedo);
+
+	if (!bSuspendNotifications)
+	{
+		Graph->MarkPackageDirty();
+	}
+	
 	if (bSetupUndoRedo)
 	{
 		ActionStack->EndAction(Action);
@@ -2203,11 +2204,11 @@ URigVMInjectionInfo* URigVMController::AddInjectedNodeFromStructPath(const FStri
 	return AddInjectedNode(InPinPath, bAsInput, ScriptStruct, InMethodName, InInputPinName, InOutputPinName, InNodeName, bSetupUndoRedo);
 }
 
-URigVMNode* URigVMController::EjectNodeFromPin(const FString& InPinPath, bool bSetupUndoRedo, bool bPrintPythonCommand)
+bool URigVMController::RemoveInjectedNode(const FString& InPinPath, bool bAsInput, bool bSetupUndoRedo, bool bPrintPythonCommand)
 {
 	if (!IsValidGraph())
 	{
-		return nullptr;
+		return false;
 	}
 
 	URigVMGraph* Graph = GetGraph();
@@ -2215,135 +2216,329 @@ URigVMNode* URigVMController::EjectNodeFromPin(const FString& InPinPath, bool bS
 
 	if (Graph->IsA<URigVMFunctionLibrary>())
 	{
-		ReportError(TEXT("Cannot eject nodes in function library graphs."));
-		return nullptr;
+		ReportError(TEXT("Cannot add injected nodes to function library graphs."));
+		return false;
 	}
 
 	URigVMPin* Pin = Graph->FindPin(InPinPath);
 	if (Pin == nullptr)
 	{
-		ReportErrorf(TEXT("Cannot find pin '%s'."), *InPinPath);
-		return nullptr;
+		return false;
 	}
 
 	if (!Pin->HasInjectedNodes())
 	{
-		ReportErrorf(TEXT("Pin '%s' has no injected nodes."), *InPinPath);
-		return nullptr;
+		return false;
 	}
 
-	URigVMInjectionInfo* Injection = Pin->InjectionInfos.Last();
 
-	FName NodeName = Injection->Node->GetFName();
-	
-	TMap<FName, FString> DefaultValues;
-	for (URigVMPin* PinOnNode : Injection->Node->GetPins())
-	{
-		if (PinOnNode->GetDirection() == ERigVMPinDirection::Input ||
-			PinOnNode->GetDirection() == ERigVMPinDirection::Visible ||
-			PinOnNode->GetDirection() == ERigVMPinDirection::IO)
-		{
-			FString DefaultValue = PinOnNode->GetDefaultValue();
-			PostProcessDefaultValue(PinOnNode, DefaultValue);
-			DefaultValues.Add(PinOnNode->GetFName(), DefaultValue);
-		}
-	}
+	// 1.- Eject node
+	// 2.- Rewire links
+	// 3.- Remove node
 
 	FRigVMControllerCompileBracketScope CompileScope(this);
 	FRigVMBaseAction Action;
 	if (bSetupUndoRedo)
 	{
-		Action.Title = FString::Printf(TEXT("Eject Node"));
+		Action.Title = FString::Printf(TEXT("Remove Injected Node"));
 		ActionStack->BeginAction(Action);
 	}
 
-	FVector2D Position = Pin->GetNode()->GetPosition() + FVector2D(0.f, 12.f) * float(Pin->GetPinIndex());
-	if (Pin->GetDirection() == ERigVMPinDirection::Output)
+	URigVMInjectionInfo* InjectionInfo = Pin->InjectionInfos.Last();
+	URigVMPin* InputPin = InjectionInfo->InputPin;
+	URigVMPin* OutputPin = InjectionInfo->OutputPin;
+
+	// 1.- Eject node
+	URigVMNode* NodeEjected = EjectNodeFromPin(InPinPath, bSetupUndoRedo);
+	if (!NodeEjected)
 	{
-		Position += FVector2D(250.f, 0.f);
+		ActionStack->CancelAction(Action);
+		return false;
+	}
+
+	// 2.- Rewire links
+	if (bAsInput)
+	{
+		BreakLink(OutputPin, Pin, bSetupUndoRedo);
+		if (InputPin)
+		{
+			TArray<URigVMLink*> Links = InputPin->GetSourceLinks();
+			RewireLinks(InputPin, Pin, true, bSetupUndoRedo, Links);
+		}
 	}
 	else
 	{
-		Position -= FVector2D(250.f, 0.f);
+		BreakLink(Pin, InputPin, bSetupUndoRedo);
+		TArray<URigVMLink*> Links = InputPin->GetTargetLinks();
+		RewireLinks(OutputPin, Pin, false, bSetupUndoRedo, Links);
+	}
+	
+	// 3.- Remove node
+	if (!RemoveNode(NodeEjected))
+	{
+		ActionStack->CancelAction(Action);
+		return false;
 	}
 
-	URigVMNode* EjectedNode = nullptr;
-	if (URigVMUnitNode* UnitNode = Cast<URigVMUnitNode>(Injection->Node))
+	if (!bSuspendNotifications)
 	{
-		EjectedNode = AddUnitNode(UnitNode->GetScriptStruct(), UnitNode->GetMethodName(), Position, FString(), bSetupUndoRedo);
+		Graph->MarkPackageDirty();
 	}
-	else if (URigVMVariableNode* VariableNode = Cast<URigVMVariableNode>(Injection->Node))
+	
+	if (bSetupUndoRedo)
 	{
-		EjectedNode = AddVariableNode(VariableNode->GetVariableName(), VariableNode->GetCPPType(), VariableNode->GetCPPTypeObject(), true, VariableNode->GetDefaultValue(), Position, FString(), bSetupUndoRedo);
+		ActionStack->EndAction(Action);
 	}
 
-	if (!EjectedNode)
+	if (bPrintPythonCommand)
 	{
-		ReportErrorf(TEXT("Could not eject node."));
-		if (bSetupUndoRedo)
+		const FString GraphName = GetSanitizedGraphName(GetGraph()->GetGraphName());
+		
+		RigVMPythonUtils::Print(GetGraphOuterName(), 
+							FString::Printf(TEXT("blueprint.get_controller_by_name('%s').remove_injected_node('%s', %s)"),
+											*GraphName,
+											*GetSanitizedPinPath(InPinPath),
+											(bAsInput) ? TEXT("True") : TEXT("False")));
+	}
+
+	return true;
+}
+
+URigVMInjectionInfo* URigVMController::InjectNodeIntoPin(const FString& InPinPath, bool bAsInput, const FName& InInputPinName, const FName& InOutputPinName, bool bSetupUndoRedo)
+{
+	if (!IsValidGraph())
+	{
+		return nullptr;
+	}
+
+	URigVMGraph* Graph = GetGraph();
+	check(Graph);
+
+	URigVMPin* Pin = Graph->FindPin(InPinPath);
+	if (!Pin)
+	{
+		return nullptr;
+	}
+
+	return InjectNodeIntoPin(Pin, bAsInput, InInputPinName, InOutputPinName, bSetupUndoRedo);
+}
+
+URigVMInjectionInfo* URigVMController::InjectNodeIntoPin(URigVMPin* InPin, bool bAsInput, const FName& InInputPinName, const FName& InOutputPinName, bool bSetupUndoRedo)
+{
+	if (!IsValidGraph())
+	{
+		return nullptr;
+	}
+
+	URigVMGraph* Graph = GetGraph();
+	check(Graph);
+
+	if (Graph->IsA<URigVMFunctionLibrary>())
+	{
+		ReportError(TEXT("Cannot inject nodes in function library graphs."));
+		return nullptr;
+	}
+
+	URigVMPin* PinForLink = InPin->GetPinForLink();
+
+	URigVMNode* NodeToInject = nullptr;
+	TArray<URigVMPin*> ConnectedPins = bAsInput ? PinForLink->GetLinkedSourcePins(true) : PinForLink->GetLinkedTargetPins(true);
+	if (ConnectedPins.Num() < 1)
+	{
+		ReportErrorf(TEXT("Cannot find node connected to pin '%s' as %s."), *InPin->GetPinPath(), bAsInput ? TEXT("input") : TEXT("output"));
+		return nullptr;
+	}
+
+	NodeToInject = ConnectedPins[0]->GetNode();
+	for (int32 i = 1; i < ConnectedPins.Num(); ++i)
+	{
+		if (ConnectedPins[i]->GetNode() != NodeToInject)
 		{
-			ActionStack->CancelAction(Action);
+			ReportErrorf(TEXT("Found more than one node connected to pin '%s' as %s."), *InPin->GetPinPath(), bAsInput ? TEXT("input") : TEXT("output"));
+			return nullptr;
+		}
+	}
+
+	URigVMPin* InputPin = nullptr;
+	URigVMPin* OutputPin = nullptr;
+	if (NodeToInject->IsA<URigVMUnitNode>())
+	{
+		InputPin = NodeToInject->FindPin(InInputPinName.ToString());
+		if (!InputPin)
+		{
+			ReportErrorf(TEXT("Could not find pin '%s' in node %s."), *InInputPinName.ToString(), *NodeToInject->GetNodePath());
+			return nullptr;
 		}
+	}
+	OutputPin = NodeToInject->FindPin(InOutputPinName.ToString());
+	if (!OutputPin)
+	{
+		ReportErrorf(TEXT("Could not find pin '%s' in node %s."), *InOutputPinName.ToString(), *NodeToInject->GetNodePath());
 		return nullptr;
 	}
 
-	for (const TPair<FName, FString>& Pair : DefaultValues)
+	FRigVMControllerCompileBracketScope CompileScope(this);
+	FRigVMBaseAction Action;
+	if (bSetupUndoRedo)
+	{
+		Action.Title = FString::Printf(TEXT("Inject Node"));
+		ActionStack->BeginAction(Action);
+	}
+	
+	
+	URigVMInjectionInfo* InjectionInfo = NewObject<URigVMInjectionInfo>(InPin);
+	{
+		Notify(ERigVMGraphNotifType::NodeRemoved, NodeToInject);
+		
+		// re-parent the unit node to be under the injection info
+		RenameObject(NodeToInject, nullptr, InjectionInfo);
+		
+		InjectionInfo->Node = NodeToInject;
+		InjectionInfo->bInjectedAsInput = bAsInput;
+		InjectionInfo->InputPin = InputPin;
+		InjectionInfo->OutputPin = OutputPin;
+	
+		InPin->InjectionInfos.Add(InjectionInfo);
+
+		Notify(ERigVMGraphNotifType::NodeAdded, NodeToInject);
+	}
+
+	// Notify the change in links (after the node is injected)
 	{
-		if (Pair.Value.IsEmpty())
+		TArray<URigVMLink*> NewLinks;
+		if (bAsInput)
 		{
-			continue;
+			if (InputPin)
+			{
+				NewLinks = InputPin->GetSourceLinks();
+			}
 		}
-		if (URigVMPin* PinOnNode = EjectedNode->FindPin(Pair.Key.ToString()))
+		else
 		{
-			SetPinDefaultValue(PinOnNode, Pair.Value, true, bSetupUndoRedo, false);
+			NewLinks = OutputPin->GetTargetLinks();
+		}
+		for (URigVMLink* Link : NewLinks)
+		{
+			Notify(ERigVMGraphNotifType::LinkAdded, Link);
 		}
 	}
 
-	TArray<URigVMLink*> PreviousLinks;
-	if (Injection->InputPin)
+	if (!bSuspendNotifications)
 	{
-		Injection->InputPin->GetSourceLinks(true);
+		Graph->MarkPackageDirty();
 	}
-	if (Injection->OutputPin)
+
+	if (bSetupUndoRedo)
 	{
-		PreviousLinks.Append(Injection->OutputPin->GetTargetLinks(true));
+		ActionStack->AddAction(FRigVMInjectNodeIntoPinAction(InjectionInfo));
+		ActionStack->EndAction(Action);
 	}
-	for (URigVMLink* PreviousLink : PreviousLinks)
+
+	return InjectionInfo;
+}
+
+URigVMNode* URigVMController::EjectNodeFromPin(const FString& InPinPath, bool bSetupUndoRedo, bool bPrintPythonCommand)
+{
+	if (!IsValidGraph())
 	{
-		PreviousLink->PrepareForCopy();
-		PreviousLink->SourcePin = PreviousLink->TargetPin = nullptr;
+		return nullptr;
 	}
 
-	RemoveNode(Injection->Node, bSetupUndoRedo);
+	URigVMGraph* Graph = GetGraph();
+	check(Graph);
+
+	URigVMPin* Pin = Graph->FindPin(InPinPath);
+	if (!Pin)
+	{
+		return nullptr;
+	}
 
-	FString OldNodeNamePrefix = NodeName.ToString() + TEXT(".");
-	FString NewNodeNamePrefix = EjectedNode->GetName() + TEXT(".");
+	return EjectNodeFromPin(Pin, bSetupUndoRedo, bPrintPythonCommand);
+}
 
-	for (URigVMLink* PreviousLink : PreviousLinks)
+URigVMNode* URigVMController::EjectNodeFromPin(URigVMPin* InPin, bool bSetupUndoRedo, bool bPrintPythonCommand)
+{
+	if (!IsValidGraph())
 	{
-		FString SourcePinPath = PreviousLink->SourcePinPath;
-		if (SourcePinPath.StartsWith(OldNodeNamePrefix))
-		{
-			SourcePinPath = NewNodeNamePrefix + SourcePinPath.RightChop(OldNodeNamePrefix.Len());
-		}
-		FString TargetPinPath = PreviousLink->TargetPinPath;
-		if (TargetPinPath.StartsWith(OldNodeNamePrefix))
+		return nullptr;
+	}
+
+	URigVMGraph* Graph = GetGraph();
+	check(Graph);
+
+	if (Graph->IsA<URigVMFunctionLibrary>())
+	{
+		ReportError(TEXT("Cannot eject nodes in function library graphs."));
+		return nullptr;
+	}
+
+	if (!InPin->HasInjectedNodes())
+	{
+		ReportErrorf(TEXT("Pin '%s' has no injected nodes."), *InPin->GetPinPath());
+		return nullptr;
+	}
+
+	URigVMInjectionInfo* Injection = InPin->InjectionInfos.Last();
+
+	
+	FRigVMControllerCompileBracketScope CompileScope(this);
+	FRigVMInverseAction InverseAction;
+	if (bSetupUndoRedo)
+	{
+		InverseAction.Title = TEXT("Eject node");
+
+		ActionStack->BeginAction(InverseAction);
+		ActionStack->AddAction(FRigVMInjectNodeIntoPinAction(Injection));
+	}
+
+	FVector2D Position = InPin->GetNode()->GetPosition() + FVector2D(0.f, 12.f) * float(InPin->GetPinIndex());
+	if (InPin->GetDirection() == ERigVMPinDirection::Output)
+	{
+		Position += FVector2D(250.f, 0.f);
+	}
+	else
+	{
+		Position -= FVector2D(250.f, 0.f);
+	}
+
+
+	URigVMNode* NodeToEject = Injection->Node;
+	URigVMPin* InputPin = Injection->InputPin;
+	URigVMPin* OutputPin = Injection->OutputPin;
+	Notify(ERigVMGraphNotifType::NodeRemoved, NodeToEject);
+	if (Injection->bInjectedAsInput)
+	{
+		if (InputPin)
 		{
-			TargetPinPath = NewNodeNamePrefix + TargetPinPath.RightChop(OldNodeNamePrefix.Len());
+			Notify(ERigVMGraphNotifType::LinkRemoved, InputPin->GetSourceLinks(true)[0]);
 		}
-
-		URigVMPin* SourcePin = Graph->FindPin(SourcePinPath);
-		URigVMPin* TargetPin = Graph->FindPin(TargetPinPath);
-		AddLink(SourcePin, TargetPin, bSetupUndoRedo);
 	}
+	else
+	{
+		Notify(ERigVMGraphNotifType::LinkRemoved, OutputPin->GetTargetLinks(true)[0]);
+	}
+	
+	
+	RenameObject(NodeToEject, nullptr, Graph);
+	SetNodePosition(NodeToEject, Position, false);
+	InPin->InjectionInfos.Remove(Injection);
+	DestroyObject(Injection);
 
-	TArray<FName> NodeNamesToSelect;
-	NodeNamesToSelect.Add(EjectedNode->GetFName());
-	SetNodeSelection(NodeNamesToSelect, bSetupUndoRedo);
+	Notify(ERigVMGraphNotifType::NodeAdded, NodeToEject);
+	if (InputPin)
+	{
+		Notify(ERigVMGraphNotifType::LinkAdded, InputPin->GetSourceLinks(true)[0]);
+	}
+	Notify(ERigVMGraphNotifType::LinkAdded, OutputPin->GetTargetLinks(true)[0]);
+		
+	if (!bSuspendNotifications)
+	{
+		Graph->MarkPackageDirty();
+	}
 
 	if (bSetupUndoRedo)
 	{
-		ActionStack->EndAction(Action);
+		ActionStack->EndAction(InverseAction);
 	}
 
 	if (bPrintPythonCommand)
@@ -2353,10 +2548,10 @@ URigVMNode* URigVMController::EjectNodeFromPin(const FString& InPinPath, bool bS
 		RigVMPythonUtils::Print(GetGraphOuterName(), 
 							FString::Printf(TEXT("blueprint.get_controller_by_name('%s').eject_node_from_pin('%s')"),
 											*GraphName,
-											*GetSanitizedPinPath(InPinPath)));
+											*GetSanitizedPinPath(InPin->GetPinPath())));
 	}
 
-	return EjectedNode;
+	return NodeToEject;
 }
 
 
@@ -4743,6 +4938,16 @@ bool URigVMController::RemoveNode(URigVMNode* InNode, bool bSetupUndoRedo, bool
 	URigVMGraph* Graph = GetGraph();
 	check(Graph);
 
+	if (InNode->IsInjected())
+	{
+		URigVMInjectionInfo* InjectionInfo = InNode->GetInjectionInfo();
+		if (InjectionInfo->GetPin()->GetInjectedNodes().Last() != InjectionInfo)
+		{
+			ReportErrorf(TEXT("Cannot remove injected node %s as it is not the last injection on the pin"), *InNode->GetNodePath());
+			return false;
+		}
+	}
+
 	if (bSetupUndoRedo)
 	{
 		// don't allow deletion of function entry / return nodes
@@ -4829,22 +5034,26 @@ bool URigVMController::RemoveNode(URigVMNode* InNode, bool bSetupUndoRedo, bool
 		URigVMPin* Pin = InjectionInfo->GetPin();
 		check(Pin);
 
-		Pin->InjectionInfos.Remove(InjectionInfo);
-
+		if (!EjectNodeFromPin(Pin->GetPinPath(), bSetupUndoRedo))
+		{
+			ActionStack->CancelAction(Action);
+			return false;
+		}
+		
 		if (InjectionInfo->bInjectedAsInput)
 		{
 			if (InjectionInfo->InputPin)
 			{
-				URigVMPin* LastInputPin = Pin;
-				RewireLinks(InjectionInfo->InputPin, LastInputPin, true, false);
+				URigVMPin* LastInputPin = Pin->GetPinForLink();
+				RewireLinks(InjectionInfo->InputPin, LastInputPin, true, bSetupUndoRedo);
 			}
 		}
 		else
 		{
 			if (InjectionInfo->OutputPin)
 			{
-				URigVMPin* LastOutputPin = Pin;
-				RewireLinks(InjectionInfo->OutputPin, LastOutputPin, false, false);
+				URigVMPin* LastOutputPin = Pin->GetPinForLink();
+				RewireLinks(InjectionInfo->OutputPin, LastOutputPin, false, bSetupUndoRedo);
 			}
 		}
 	}
@@ -4930,9 +5139,9 @@ bool URigVMController::RemoveNode(URigVMNode* InNode, bool bSetupUndoRedo, bool
 		for (URigVMPin* Pin : InNode->GetPins())
 		{
 			TArray<URigVMInjectionInfo*> InjectedNodes = Pin->GetInjectedNodes();
-			for (URigVMInjectionInfo* InjectedNode : InjectedNodes)
+			for (int32 i=InjectedNodes.Num()-1; i >= 0; --i)
 			{
-				RemoveNode(InjectedNode->Node, bSetupUndoRedo, bRecursive);
+				RemoveNode(InjectedNodes[i]->Node, bSetupUndoRedo, bRecursive);
 			}
 
 			BreakAllLinks(Pin, true, bSetupUndoRedo);
@@ -6907,6 +7116,7 @@ bool URigVMController::BindPinToVariable(URigVMPin* InPin, const FString& InNewB
 		ActionStack->BeginAction(Action);
 	}
 
+	// Unbind any other variables, remove any other injections, and break all links to the input pin
 	{
 		if (InPin->IsBoundToVariable())
 		{
@@ -6915,16 +7125,11 @@ bool URigVMController::BindPinToVariable(URigVMPin* InPin, const FString& InNewB
 		TArray<URigVMInjectionInfo*> Infos = InPin->GetInjectedNodes();
 		for (URigVMInjectionInfo* Info : Infos)
 		{
-			RemoveNode(Info->Node, bSetupUndoRedo);
+			RemoveInjectedNode(Info->GetPin()->GetPinPath(), Info->bInjectedAsInput, bSetupUndoRedo);
 		}
 		BreakAllLinks(InPin, true, bSetupUndoRedo);
 	}
 
-	if (bSetupUndoRedo)
-	{
-		ActionStack->AddAction(FRigVMSetPinBoundVariableAction(InPin, InNewBoundVariablePath));
-	}		
-	
 	// Create variable node
 	URigVMVariableNode* VariableNode = nullptr;
 	{
@@ -6933,7 +7138,7 @@ bool URigVMController::BindPinToVariable(URigVMPin* InPin, const FString& InNewB
 			FString CPPType;
 			UObject* CPPTypeObject;
 			RigVMTypeUtils::CPPTypeFromExternalVariable(Variable, CPPType, &CPPTypeObject);
-			VariableNode = AddVariableNode(*VariableName, CPPType, CPPTypeObject, true, FString(), FVector2D::ZeroVector, InVariableNodeName, false);
+			VariableNode = AddVariableNode(*VariableName, CPPType, CPPTypeObject, true, FString(), FVector2D::ZeroVector, InVariableNodeName, bSetupUndoRedo);
 		}
 		if (VariableNode == nullptr)
 		{
@@ -6946,35 +7151,33 @@ bool URigVMController::BindPinToVariable(URigVMPin* InPin, const FString& InNewB
 	}
 	
 	URigVMPin* ValuePin = VariableNode->FindPin(URigVMVariableNode::ValueName);
-	if (!SegmentPath.IsEmpty())
+	// Connect value pin to input pin
 	{
-		ValuePin = ValuePin->FindSubPin(SegmentPath);
-	}
+		if (!SegmentPath.IsEmpty())
+		{
+			ValuePin = ValuePin->FindSubPin(SegmentPath);
+		}
 
-	// Add injection 
-	URigVMInjectionInfo* InjectionInfo = NewObject<URigVMInjectionInfo>(InPin);
-	{
-		// re-parent the unit node to be under the injection info
-		RenameObject(VariableNode, nullptr, InjectionInfo);
-		
-		InjectionInfo->Node = VariableNode;
-		InjectionInfo->bInjectedAsInput = true;
-		InjectionInfo->InputPin = nullptr;
-		InjectionInfo->OutputPin = ValuePin;
-	
-		InPin->InjectionInfos.Add(InjectionInfo);
-		Notify(ERigVMGraphNotifType::NodeAdded, VariableNode);
+		{		
+			if (!AddLink(ValuePin, InPin, bSetupUndoRedo))
+			{
+				if (bSetupUndoRedo)
+				{
+					ActionStack->CancelAction(Action);
+				}
+				return false;
+			}
+		}
 	}
 
-	{		
-		if (!AddLink(ValuePin, InPin, false))
+	// Inject into pin
+	if (!InjectNodeIntoPin(InPin->GetPinPath(), true, FName(), ValuePin->GetFName(), bSetupUndoRedo))
+	{
+		if (bSetupUndoRedo)
 		{
-			if (bSetupUndoRedo)
-			{
-				ActionStack->CancelAction(Action);
-			}
-			return false;
+			ActionStack->CancelAction(Action);
 		}
+		return false;
 	}
 	
 	if (bSetupUndoRedo)
@@ -7044,20 +7247,7 @@ bool URigVMController::UnbindPinFromVariable(URigVMPin* InPin, bool bSetupUndoRe
 		ActionStack->BeginAction(Action);
 	}
 
-	if (bSetupUndoRedo)
-	{
-		ActionStack->AddAction(FRigVMSetPinBoundVariableAction(InPin, FString()));
-	}
-
-	URigVMInjectionInfo* Info = nullptr;
-	for (URigVMInjectionInfo* InjectionInfo : InPin->GetInjectedNodes())
-	{
-		if (InjectionInfo->Node->IsA<URigVMVariableNode>())
-		{
-			BreakAllLinks(InPin, true, false);
-			RemoveNode(InjectionInfo->Node, false);
-		}
-	}	
+	RemoveInjectedNode(InPin->GetPinPath(), true, bSetupUndoRedo);
 
 	if (bSetupUndoRedo)
 	{
@@ -12582,8 +12772,8 @@ void URigVMController::RewireLinks(URigVMPin* InOldPin, URigVMPin* InNewPin, boo
 			URigVMPin* NewPin = SegmentPath.IsEmpty() ? InNewPin : InNewPin->FindSubPin(SegmentPath);
 			check(NewPin);
 
-			BreakLink(Link->GetSourcePin(), Link->GetTargetPin(), false);
-			AddLink(Link->GetSourcePin(), NewPin, false);
+			BreakLink(Link->GetSourcePin(), Link->GetTargetPin(), bSetupUndoRedo);
+			AddLink(Link->GetSourcePin(), NewPin, bSetupUndoRedo);
 		}
 	}
 	else
@@ -12600,8 +12790,8 @@ void URigVMController::RewireLinks(URigVMPin* InOldPin, URigVMPin* InNewPin, boo
 			URigVMPin* NewPin = SegmentPath.IsEmpty() ? InNewPin : InNewPin->FindSubPin(SegmentPath);
 			check(NewPin);
 
-			BreakLink(Link->GetSourcePin(), Link->GetTargetPin(), false);
-			AddLink(NewPin, Link->GetTargetPin(), false);
+			BreakLink(Link->GetSourcePin(), Link->GetTargetPin(), bSetupUndoRedo);
+			AddLink(NewPin, Link->GetTargetPin(), bSetupUndoRedo);
 		}
 	}
 }
diff --git a/Engine/Source/Developer/RigVMDeveloper/Private/RigVMModel/RigVMControllerActions.cpp b/Engine/Source/Developer/RigVMDeveloper/Private/RigVMModel/RigVMControllerActions.cpp
--- a/Engine/Source/Developer/RigVMDeveloper/Private/RigVMModel/RigVMControllerActions.cpp
+++ b/Engine/Source/Developer/RigVMDeveloper/Private/RigVMModel/RigVMControllerActions.cpp
@@ -618,18 +618,16 @@ bool FRigVMAddPrototypeNodeAction::Redo(URigVMController* InController)
 	return false;
 }
 
-FRigVMAddInjectedNodeAction::FRigVMAddInjectedNodeAction()
+FRigVMInjectNodeIntoPinAction::FRigVMInjectNodeIntoPinAction()
 	: PinPath()
 	, bAsInput(false)
-	, ScriptStructPath()
-	, MethodName(NAME_None)
 	, InputPinName(NAME_None)
 	, OutputPinName(NAME_None)
 	, NodePath()
 {
 }
 
-FRigVMAddInjectedNodeAction::FRigVMAddInjectedNodeAction(URigVMInjectionInfo* InInjectionInfo)
+FRigVMInjectNodeIntoPinAction::FRigVMInjectNodeIntoPinAction(URigVMInjectionInfo* InInjectionInfo)
 	: PinPath(InInjectionInfo->GetPin()->GetPinPath())
 	, bAsInput(InInjectionInfo->bInjectedAsInput)
 	, NodePath(InInjectionInfo->Node->GetName())
@@ -642,35 +640,23 @@ FRigVMAddInjectedNodeAction::FRigVMAddInjectedNodeAction(URigVMInjectionInfo* In
 	{
 		OutputPinName = InInjectionInfo->OutputPin->GetFName();
 	}
-	if (URigVMUnitNode* UnitNode = Cast<URigVMUnitNode>(InInjectionInfo->Node))
-	{
-		ScriptStructPath = UnitNode->GetScriptStruct()->GetPathName();
-		MethodName = UnitNode->GetMethodName();
-	}
-	else
-	{
-		checkNoEntry();
-	}
 }
 
-bool FRigVMAddInjectedNodeAction::Undo(URigVMController* InController)
+bool FRigVMInjectNodeIntoPinAction::Undo(URigVMController* InController)
 {
 	if (!FRigVMBaseAction::Undo(InController))
 	{
 		return false;
 	}
-	return InController->RemoveNodeByName(*NodePath, false, true);
+	return InController->EjectNodeFromPin(*PinPath, false, true) != nullptr;
 }
 
-bool FRigVMAddInjectedNodeAction::Redo(URigVMController* InController)
+bool FRigVMInjectNodeIntoPinAction::Redo(URigVMController* InController)
 {
 #if WITH_EDITOR
-	if (!ScriptStructPath.IsEmpty())
+	if (URigVMInjectionInfo* InjectionInfo = InController->InjectNodeIntoPin(PinPath, bAsInput, InputPinName, OutputPinName, false))
 	{
-		if (URigVMInjectionInfo* InjectionInfo = InController->AddInjectedNodeFromStructPath(PinPath, bAsInput, ScriptStructPath, MethodName, InputPinName, OutputPinName, NodePath, false))
-		{
-			return FRigVMBaseAction::Redo(InController);
-		}
+		return FRigVMBaseAction::Redo(InController);
 	}
 #endif
 	return false;
@@ -682,23 +668,7 @@ FRigVMRemoveNodeAction::FRigVMRemoveNodeAction(URigVMNode* InNode, URigVMControl
 
 	if (URigVMInjectionInfo* InjectionInfo = InNode->GetInjectionInfo())
 	{
-		if (URigVMUnitNode* UnitNode = Cast<URigVMUnitNode>(InjectionInfo->Node))
-		{
-			InverseAction.AddAction(FRigVMAddInjectedNodeAction(InjectionInfo));
-		}
-		else if (URigVMVariableNode* VariableNode = Cast<URigVMVariableNode>(InjectionInfo->Node))
-		{
-			InverseAction.AddAction(FRigVMSetPinBoundVariableAction(InjectionInfo->GetPin(), InjectionInfo->GetPin()->GetBoundVariablePath()));
-		}
-
-		for (URigVMPin* Pin : InNode->GetPins())
-		{
-			if (Pin->GetDirection() == ERigVMPinDirection::Input ||
-				Pin->GetDirection() == ERigVMPinDirection::Visible)
-			{
-				InverseAction.AddAction(FRigVMSetPinDefaultValueAction(Pin, Pin->GetDefaultValue()));
-			}
-		}
+		ensure(false);
 	}
 	else if (URigVMUnitNode* UnitNode = Cast<URigVMUnitNode>(InNode))
 	{
@@ -1459,46 +1429,6 @@ bool FRigVMChangePinTypeAction::Redo(URigVMController* InController)
 	return FRigVMBaseAction::Redo(InController);
 }
 
-FRigVMSetPinBoundVariableAction::FRigVMSetPinBoundVariableAction(URigVMPin* InPin, const FString& InNewBoundVariablePath)
-	: PinPath(InPin->GetPinPath())
-	, OldBoundVariablePath(InPin->GetBoundVariablePath())
-	, NewBoundVariablePath(InNewBoundVariablePath)
-{
-	
-}
-
-bool FRigVMSetPinBoundVariableAction::Merge(const FRigVMBaseAction* Other)
-{
-	if (!FRigVMBaseAction::Merge(Other))
-	{
-		return false;
-	}
-
-	const FRigVMSetPinBoundVariableAction* Action = (const FRigVMSetPinBoundVariableAction*)Other;
-	if (PinPath != Action->PinPath)
-	{
-		return false;
-	}
-
-	NewBoundVariablePath = Action->NewBoundVariablePath;
-	return true;
-}
-
-bool FRigVMSetPinBoundVariableAction::Undo(URigVMController* InController)
-{
-	if (!FRigVMBaseAction::Undo(InController))
-	{
-		return false;
-	}
-	return InController->BindPinToVariable(PinPath, OldBoundVariablePath, false);
-}
-
-bool FRigVMSetPinBoundVariableAction::Redo(URigVMController* InController)
-{
-	InController->BindPinToVariable(PinPath, NewBoundVariablePath, false);
-	return FRigVMBaseAction::Redo(InController);
-}
-
 FRigVMImportNodeFromTextAction::FRigVMImportNodeFromTextAction()
 	: Position(FVector2D::ZeroVector)
 	, NodePath()
diff --git a/Engine/Source/Developer/RigVMDeveloper/Public/RigVMModel/RigVMController.h b/Engine/Source/Developer/RigVMDeveloper/Public/RigVMModel/RigVMController.h
--- a/Engine/Source/Developer/RigVMDeveloper/Public/RigVMModel/RigVMController.h
+++ b/Engine/Source/Developer/RigVMDeveloper/Public/RigVMModel/RigVMController.h
@@ -290,6 +290,11 @@ public:
 	UFUNCTION(BlueprintCallable, Category = RigVMController)
 	URigVMInjectionInfo* AddInjectedNodeFromStructPath(const FString& InPinPath, bool bAsInput, const FString& InScriptStructPath, const FName& InMethodName, const FName& InInputPinName, const FName& InOutputPinName, const FString& InNodeName = TEXT(""), bool bSetupUndoRedo = true);
 
+	// Removes an injected node
+	// This causes a NodeRemoved modified event.
+	UFUNCTION(BlueprintCallable, Category = RigVMController)
+	bool RemoveInjectedNode(const FString& InPinPath, bool bAsInput, bool bSetupUndoRedo = true, bool bPrintPythonCommand = false);
+
 	// Ejects the last injected node on a pin
 	UFUNCTION(BlueprintCallable, Category = RigVMController)
 	URigVMNode* EjectNodeFromPin(const FString& InPinPath, bool bSetupUndoRedo = true, bool bPrintPythonCommand = false);
@@ -810,6 +815,9 @@ private:
 	bool UnbindPinFromVariable(URigVMPin* InPin, bool bSetupUndoRedo);
 	bool MakeBindingsFromVariableNode(URigVMVariableNode* InNode, bool bSetupUndoRedo);
 	bool PromotePinToVariable(URigVMPin* InPin, bool bCreateVariableNode, const FVector2D& InNodePosition, bool bSetupUndoRedo);
+	URigVMInjectionInfo* InjectNodeIntoPin(const FString& InPinPath, bool bAsInput, const FName& InInputPinName, const FName& InOutputPinName, bool bSetupUndoRedo = true);
+	URigVMInjectionInfo* InjectNodeIntoPin(URigVMPin* InPin, bool bAsInput, const FName& InInputPinName, const FName& InOutputPinName, bool bSetupUndoRedo = true);
+	URigVMNode* EjectNodeFromPin(URigVMPin* InPin, bool bSetupUndoRedo = true, bool bPrintPythonCommands = false);
 
 	// try to reconnect source and target pins after a node deletion
 	void RelinkSourceAndTargetPins(URigVMNode* RigNode, bool bSetupUndoRedo = true);
@@ -961,6 +969,7 @@ private:
 	friend struct FRigVMControllerObjectFactory;
 	friend struct FRigVMAddRerouteNodeAction;
 	friend struct FRigVMChangePinTypeAction;
+	friend struct FRigVMInjectNodeIntoPinAction;
 	friend class FRigVMParserAST;
 	friend class FRigVMControllerCompileBracketScope;
 };
diff --git a/Engine/Source/Developer/RigVMDeveloper/Public/RigVMModel/RigVMControllerActions.h b/Engine/Source/Developer/RigVMDeveloper/Public/RigVMModel/RigVMControllerActions.h
--- a/Engine/Source/Developer/RigVMDeveloper/Public/RigVMModel/RigVMControllerActions.h
+++ b/Engine/Source/Developer/RigVMDeveloper/Public/RigVMModel/RigVMControllerActions.h
@@ -619,18 +619,18 @@ public:
 };
 
 /**
- * An action adding an injected node to the graph.
+ * An action injecting a node into a pin
  */
 USTRUCT()
-struct FRigVMAddInjectedNodeAction : public FRigVMBaseAction
+struct FRigVMInjectNodeIntoPinAction : public FRigVMBaseAction
 {
 	GENERATED_BODY()
 
 public:
 
-	FRigVMAddInjectedNodeAction();
-	FRigVMAddInjectedNodeAction(URigVMInjectionInfo* InInjectionInfo);
-	virtual ~FRigVMAddInjectedNodeAction() {};
+	FRigVMInjectNodeIntoPinAction();
+	FRigVMInjectNodeIntoPinAction(URigVMInjectionInfo* InInjectionInfo);
+	virtual ~FRigVMInjectNodeIntoPinAction() {};
 	virtual bool Undo(URigVMController* InController) override;
 	virtual bool Redo(URigVMController* InController) override;
 
@@ -640,12 +640,6 @@ public:
 	UPROPERTY()
 	bool bAsInput;
 
-	UPROPERTY()
-	FString ScriptStructPath;
-
-	UPROPERTY()
-	FName MethodName;
-
 	UPROPERTY()
 	FName InputPinName;
 
@@ -1229,33 +1223,6 @@ public:
 	bool bRemoveSubPins;
 };
 
-/**
- * An action changing a pin's bound variable
- */
-USTRUCT()
-struct FRigVMSetPinBoundVariableAction : public FRigVMBaseAction
-{
-	GENERATED_BODY()
-
-public:
-
-	FRigVMSetPinBoundVariableAction() {}
-	FRigVMSetPinBoundVariableAction(URigVMPin* InPin, const FString& InNewBoundVariablePath);
-	virtual ~FRigVMSetPinBoundVariableAction() {};
-	virtual bool Merge(const FRigVMBaseAction* Other);
-	virtual bool Undo(URigVMController* InController) override;
-	virtual bool Redo(URigVMController* InController) override;
-
-	UPROPERTY()
-	FString PinPath;
-
-	UPROPERTY()
-	FString OldBoundVariablePath;
-
-	UPROPERTY()
-	FString NewBoundVariablePath;
-};
-
 /**
  * An action to add a node from a text buffer
  */
