Control Rig: Fix crash when removing interpolation node

diff --git a/Engine/Source/Developer/RigVMDeveloper/Private/RigVMModel/RigVMController.cpp b/Engine/Source/Developer/RigVMDeveloper/Private/RigVMModel/RigVMController.cpp
--- a/Engine/Source/Developer/RigVMDeveloper/Private/RigVMModel/RigVMController.cpp
+++ b/Engine/Source/Developer/RigVMDeveloper/Private/RigVMModel/RigVMController.cpp
@@ -2510,12 +2510,20 @@ URigVMNode* URigVMController::EjectNodeFromPin(URigVMPin* InPin, bool bSetupUndo
 	{
 		if (InputPin)
 		{
-			Notify(ERigVMGraphNotifType::LinkRemoved, InputPin->GetSourceLinks(true)[0]);
+			TArray<URigVMLink*> SourceLinks = InputPin->GetSourceLinks(true);
+			if (SourceLinks.Num() > 0)
+			{
+				Notify(ERigVMGraphNotifType::LinkRemoved, SourceLinks[0]);
+			}
 		}
 	}
 	else
 	{
-		Notify(ERigVMGraphNotifType::LinkRemoved, OutputPin->GetTargetLinks(true)[0]);
+		TArray<URigVMLink*> TargetLinks = OutputPin->GetTargetLinks(true);
+		if (TargetLinks.Num() > 0)
+		{
+			Notify(ERigVMGraphNotifType::LinkRemoved, TargetLinks[0]);
+		}
 	}
 	
 	
@@ -2527,9 +2535,17 @@ URigVMNode* URigVMController::EjectNodeFromPin(URigVMPin* InPin, bool bSetupUndo
 	Notify(ERigVMGraphNotifType::NodeAdded, NodeToEject);
 	if (InputPin)
 	{
-		Notify(ERigVMGraphNotifType::LinkAdded, InputPin->GetSourceLinks(true)[0]);
+		TArray<URigVMLink*> SourceLinks = InputPin->GetSourceLinks(true);
+		if (SourceLinks.Num() > 0)
+		{
+			Notify(ERigVMGraphNotifType::LinkAdded, SourceLinks[0]);
+		}
+	}
+	TArray<URigVMLink*> TargetLinks = OutputPin->GetTargetLinks(true);
+	if (TargetLinks.Num() > 0)
+	{
+		Notify(ERigVMGraphNotifType::LinkAdded, TargetLinks[0]);
 	}
-	Notify(ERigVMGraphNotifType::LinkAdded, OutputPin->GetTargetLinks(true)[0]);
 		
 	if (!bSuspendNotifications)
 	{
