Control Rig: Fix crash on copy paste nodes with variable binding

diff --git a/Engine/Plugins/Animation/ControlRig/Source/ControlRigEditor/Private/Editor/ControlRigEditor.cpp b/Engine/Plugins/Animation/ControlRig/Source/ControlRigEditor/Private/Editor/ControlRigEditor.cpp
--- a/Engine/Plugins/Animation/ControlRig/Source/ControlRigEditor/Private/Editor/ControlRigEditor.cpp
+++ b/Engine/Plugins/Animation/ControlRig/Source/ControlRigEditor/Private/Editor/ControlRigEditor.cpp
@@ -2349,11 +2349,18 @@ void FControlRigEditor::PasteNodes()
 		FBox2D Bounds;
 		Bounds.bIsValid = false;
 
+		TArray<FName> NodesToSelect;
 		for (const FName& NodeName : NodeNames)
 		{
 			const URigVMNode* Node = GetFocusedModel()->FindNodeByName(NodeName);
 			check(Node);
 
+			if (Node->IsInjected())
+			{
+				continue;
+			}
+			NodesToSelect.Add(NodeName);
+
 			FVector2D Position = Node->GetPosition();
 			FVector2D Size = Node->GetSize();
 
@@ -2366,7 +2373,7 @@ void FControlRigEditor::PasteNodes()
 			Bounds += Position + Size;
 		}
 
-		for (const FName& NodeName : NodeNames)
+		for (const FName& NodeName : NodesToSelect)
 		{
 			const URigVMNode* Node = GetFocusedModel()->FindNodeByName(NodeName);
 			check(Node);
@@ -2375,7 +2382,7 @@ void FControlRigEditor::PasteNodes()
 			GetFocusedController()->SetNodePositionByName(NodeName, PasteLocation + Position - Bounds.GetCenter(), true, false, true);
 		}
 
-		GetFocusedController()->SetNodeSelection(NodeNames);
+		GetFocusedController()->SetNodeSelection(NodesToSelect);
 		GetFocusedController()->CloseUndoBracket();
 	}
 	else
diff --git a/Engine/Source/Developer/RigVMDeveloper/Private/RigVMModel/Nodes/RigVMVariableNode.cpp b/Engine/Source/Developer/RigVMDeveloper/Private/RigVMModel/Nodes/RigVMVariableNode.cpp
--- a/Engine/Source/Developer/RigVMDeveloper/Private/RigVMModel/Nodes/RigVMVariableNode.cpp
+++ b/Engine/Source/Developer/RigVMDeveloper/Private/RigVMModel/Nodes/RigVMVariableNode.cpp
@@ -118,3 +118,13 @@ FRigVMGraphVariableDescription URigVMVariableNode::GetVariableDescription() cons
 	Variable.DefaultValue = GetDefaultValue();
 	return Variable;
 }
+
+URigVMPin* URigVMVariableNode::GetVariableNamePin() const
+{
+	return FindPin(VariableName);
+}
+
+URigVMPin* URigVMVariableNode::GetValuePin() const
+{
+	return FindPin(ValueName);
+}
diff --git a/Engine/Source/Developer/RigVMDeveloper/Private/RigVMModel/RigVMController.cpp b/Engine/Source/Developer/RigVMDeveloper/Private/RigVMModel/RigVMController.cpp
--- a/Engine/Source/Developer/RigVMDeveloper/Private/RigVMModel/RigVMController.cpp
+++ b/Engine/Source/Developer/RigVMDeveloper/Private/RigVMModel/RigVMController.cpp
@@ -554,31 +554,34 @@ TArray<FString> URigVMController::GetAddNodePythonCommands(URigVMNode* Node) con
 	}
 	else if (const URigVMVariableNode* VariableNode = Cast<URigVMVariableNode>(Node))
 	{
-		const FString VariableName = GetSanitizedVariableName(VariableNode->GetVariableName().ToString());
-
-		// add_variable_node(variable_name, cpp_type, cpp_type_object, is_getter, default_value, position=[0.0, 0.0], node_name='', undo=True)
-		if (VariableNode->GetVariableDescription().CPPTypeObject)
+		if (!VariableNode->IsInjected())
 		{
-			Commands.Add(FString::Printf(TEXT("blueprint.get_controller_by_name('%s').add_variable_node_from_object_path('%s', '%s', '%s', %s, '%s', %s, '%s')"),
-					*GraphName,
-					*VariableName,
-					*VariableNode->GetVariableDescription().CPPType,
-					*VariableNode->GetVariableDescription().CPPTypeObject->GetPathName(),
-					VariableNode->IsGetter() ? TEXT("True") : TEXT("False"),
-					*VariableNode->GetVariableDescription().DefaultValue,
-					*RigVMPythonUtils::Vector2DToPythonString(VariableNode->GetPosition()),
-					*NodeName));	
-		}
-		else
-		{
-			Commands.Add(FString::Printf(TEXT("blueprint.get_controller_by_name('%s').add_variable_node('%s', '%s', None, %s, '%s', %s, '%s')"),
-					*GraphName,
-					*VariableName,
-					*VariableNode->GetVariableDescription().CPPType,
-					VariableNode->IsGetter() ? TEXT("True") : TEXT("False"),
-					*VariableNode->GetVariableDescription().DefaultValue,
-					*RigVMPythonUtils::Vector2DToPythonString(VariableNode->GetPosition()),
-					*NodeName));	
+			const FString VariableName = GetSanitizedVariableName(VariableNode->GetVariableName().ToString());
+
+			// add_variable_node(variable_name, cpp_type, cpp_type_object, is_getter, default_value, position=[0.0, 0.0], node_name='', undo=True)
+			if (VariableNode->GetVariableDescription().CPPTypeObject)
+			{
+				Commands.Add(FString::Printf(TEXT("blueprint.get_controller_by_name('%s').add_variable_node_from_object_path('%s', '%s', '%s', %s, '%s', %s, '%s')"),
+						*GraphName,
+						*VariableName,
+						*VariableNode->GetVariableDescription().CPPType,
+						*VariableNode->GetVariableDescription().CPPTypeObject->GetPathName(),
+						VariableNode->IsGetter() ? TEXT("True") : TEXT("False"),
+						*VariableNode->GetVariableDescription().DefaultValue,
+						*RigVMPythonUtils::Vector2DToPythonString(VariableNode->GetPosition()),
+						*NodeName));	
+			}
+			else
+			{
+				Commands.Add(FString::Printf(TEXT("blueprint.get_controller_by_name('%s').add_variable_node('%s', '%s', None, %s, '%s', %s, '%s')"),
+						*GraphName,
+						*VariableName,
+						*VariableNode->GetVariableDescription().CPPType,
+						VariableNode->IsGetter() ? TEXT("True") : TEXT("False"),
+						*VariableNode->GetVariableDescription().DefaultValue,
+						*RigVMPythonUtils::Vector2DToPythonString(VariableNode->GetPosition()),
+						*NodeName));	
+			}
 		}
 	}
 	else if (const URigVMParameterNode* ParameterNode = Cast<URigVMParameterNode>(Node))
@@ -2820,7 +2823,10 @@ TArray<FName> URigVMController::ImportNodesFromText(const FString& InText, bool
 
 		if (bSetupUndoRedo)
 		{
-			ActionStack->AddAction(FRigVMRemoveNodeAction(CreatedNode, this));
+			if (!CreatedNode->IsInjected() || !CreatedNode->IsA<URigVMVariableNode>())
+			{
+				ActionStack->AddAction(FRigVMRemoveNodeAction(CreatedNode, this));
+			}
 		}
 
 		// find all nodes affected by this
@@ -2934,7 +2940,18 @@ TArray<FName> URigVMController::ImportNodesFromText(const FString& InText, bool
 						if (bWasBinded)
 						{
 							VariableNodeName = TargetPin->GetBoundVariableNode()->GetName();
-							BindingPath = TargetPin->GetBoundVariablePath();	
+							BindingPath = TargetPin->GetBoundVariablePath();
+
+							// The current situation is that the outter pin has an injection info, and the injected node exists
+							// but the injected node is not linked to the outter pin. BreakAllLinks will try to unbind the outter pin,
+							// for that to be successful, the binding needs to be complete
+							// Connect it so that the unbound is successful
+							if (!SourcePin->IsLinkedTo(TargetPin))
+							{
+								Graph->Links.Add(CreatedLink);
+								SourcePin->Links.Add(CreatedLink);
+								TargetPin->Links.Add(CreatedLink);
+							}
 						}
 						
 						BreakAllLinksRecursive(TargetPin, true, true, bSetupUndoRedo);
@@ -2951,11 +2968,19 @@ TArray<FName> URigVMController::ImportNodesFromText(const FString& InText, bool
 							Graph->Links.Add(CreatedLink);
 							SourcePin->Links.Add(CreatedLink);
 							TargetPin->Links.Add(CreatedLink);
-						}
 
-						if (bSetupUndoRedo)
-						{
-							ActionStack->AddAction(FRigVMAddLinkAction(SourcePin, TargetPin));
+							if (bSetupUndoRedo)
+							{
+								ActionStack->AddAction(FRigVMAddLinkAction(SourcePin, TargetPin));
+								if (SourcePin->GetNode()->IsInjected())
+								{
+									ActionStack->AddAction(FRigVMInjectNodeIntoPinAction(SourcePin->GetTypedOuter<URigVMInjectionInfo>()));	
+								}
+								if (TargetPin->GetNode()->IsInjected())
+								{
+									ActionStack->AddAction(FRigVMInjectNodeIntoPinAction(TargetPin->GetTypedOuter<URigVMInjectionInfo>()));	
+								}
+							}
 						}
 						Notify(ERigVMGraphNotifType::LinkAdded, CreatedLink);
 						continue;
@@ -7174,7 +7199,8 @@ bool URigVMController::BindPinToVariable(URigVMPin* InPin, const FString& InNewB
 			ValuePin = ValuePin->FindSubPin(SegmentPath);
 		}
 
-		{		
+		{
+			GetGraph()->ClearAST(true, false);
 			if (!AddLink(ValuePin, InPin, bSetupUndoRedo))
 			{
 				if (bSetupUndoRedo)
diff --git a/Engine/Source/Developer/RigVMDeveloper/Private/RigVMModel/RigVMControllerActions.cpp b/Engine/Source/Developer/RigVMDeveloper/Private/RigVMModel/RigVMControllerActions.cpp
--- a/Engine/Source/Developer/RigVMDeveloper/Private/RigVMModel/RigVMControllerActions.cpp
+++ b/Engine/Source/Developer/RigVMDeveloper/Private/RigVMModel/RigVMControllerActions.cpp
@@ -648,7 +648,7 @@ bool FRigVMInjectNodeIntoPinAction::Undo(URigVMController* InController)
 	{
 		return false;
 	}
-	return InController->EjectNodeFromPin(*PinPath, false, true) != nullptr;
+	return InController->EjectNodeFromPin(*PinPath, false) != nullptr;
 }
 
 bool FRigVMInjectNodeIntoPinAction::Redo(URigVMController* InController)
@@ -666,11 +666,7 @@ FRigVMRemoveNodeAction::FRigVMRemoveNodeAction(URigVMNode* InNode, URigVMControl
 {
 	FRigVMInverseAction InverseAction;
 
-	if (URigVMInjectionInfo* InjectionInfo = InNode->GetInjectionInfo())
-	{
-		ensure(false);
-	}
-	else if (URigVMUnitNode* UnitNode = Cast<URigVMUnitNode>(InNode))
+	if (URigVMUnitNode* UnitNode = Cast<URigVMUnitNode>(InNode))
 	{
 		InverseAction.AddAction(FRigVMAddUnitNodeAction(UnitNode));
 		for (URigVMPin* Pin : UnitNode->GetPins())
@@ -755,6 +751,19 @@ FRigVMRemoveNodeAction::FRigVMRemoveNodeAction(URigVMNode* InNode, URigVMControl
 			ExpansionAction.OldIsExpanded = false;
 			InverseAction.AddAction(ExpansionAction);
 		}
+
+		if (Pin->HasInjectedNodes())
+		{
+			if (URigVMVariableNode* VariableNode = Cast<URigVMVariableNode>(Pin->GetInjectedNodes()[0]->Node))
+			{
+				FRigVMAddVariableNodeAction AddVariableNodeAction(VariableNode);
+				FRigVMAddLinkAction AddLinkAction(VariableNode->GetValuePin(), Pin);
+				FRigVMInjectNodeIntoPinAction InjectAction(Pin->GetInjectedNodes()[0]);
+				InverseAction.AddAction(AddVariableNodeAction);
+				InverseAction.AddAction(AddLinkAction);
+				InverseAction.AddAction(InjectAction);
+			}
+		}
 	}
 
 	InverseActionKey.Set<FRigVMInverseAction>(InverseAction);
diff --git a/Engine/Source/Developer/RigVMDeveloper/Public/RigVMModel/Nodes/RigVMVariableNode.h b/Engine/Source/Developer/RigVMDeveloper/Public/RigVMModel/Nodes/RigVMVariableNode.h
--- a/Engine/Source/Developer/RigVMDeveloper/Public/RigVMModel/Nodes/RigVMVariableNode.h
+++ b/Engine/Source/Developer/RigVMDeveloper/Public/RigVMModel/Nodes/RigVMVariableNode.h
@@ -67,7 +67,13 @@ private:
 	static const FString VariableName;
 	static const FString ValueName;
 
+	URigVMPin* GetVariableNamePin() const;
+	URigVMPin* GetValuePin() const;
+	
+
 	friend class URigVMController;
+	friend struct FRigVMRemoveNodeAction;
+	friend class URigVMPin;
 	friend class URigVMCompiler;
 	friend class FRigVMVarExprAST;
 	friend class FRigVMParserAST;
