Control Rig: Added initial support for user defined structs in control rig

diff --git a/Engine/Plugins/Animation/ControlRig/Source/ControlRigDeveloper/Private/ControlRigBlueprint.cpp b/Engine/Plugins/Animation/ControlRig/Source/ControlRigDeveloper/Private/ControlRigBlueprint.cpp
--- a/Engine/Plugins/Animation/ControlRig/Source/ControlRigDeveloper/Private/ControlRigBlueprint.cpp
+++ b/Engine/Plugins/Animation/ControlRig/Source/ControlRigDeveloper/Private/ControlRigBlueprint.cpp
@@ -451,32 +451,12 @@ void UControlRigBlueprint::PostLoad()
 
 #if WITH_EDITOR
 
-		TArray<URigVMGraph*> GraphsToDetach;
-		GraphsToDetach.Add(GetModel());
-		GraphsToDetach.Add(GetLocalFunctionLibrary());
-
-		if (ensure(IsInGameThread()))
-		{
-			for (URigVMGraph* GraphToDetach : GraphsToDetach)
-			{
-				URigVMController* Controller = GetOrCreateController(GraphToDetach);
-				// temporarily disable default value validation during load time, serialized values should always be accepted
-				TGuardValue<bool> PerGraphDisablePinDefaultValueValidation(Controller->bValidatePinDefaults, false);
-				Controller->DetachLinksFromPinObjects();
-				TArray<URigVMNode*> Nodes = GraphToDetach->GetNodes();
-				for (URigVMNode* Node : Nodes)
-				{
-					Controller->RepopulatePinsOnNode(Node, true, false, true);
-				}
-			}
-			SetupPinRedirectorsForBackwardsCompatibility();
-		}
-
-		for (URigVMGraph* GraphToDetach : GraphsToDetach)
-		{
-			URigVMController* Controller = GetOrCreateController(GraphToDetach);
-			Controller->ReattachLinksToPinObjects(true /* follow redirectors */, nullptr, false, true);
-		}
+		// refresh the graph such that the pin hierarchies matches their CPPTypeObject
+		// this step is needed everytime we open a BP in the editor, b/c even after load
+		// model data can change while the Control Rig BP is not opened
+		// for example, if a user defined struct changed after BP load,
+		// any pin that references the struct needs to be regenerated
+		RefreshAllModels();
 
 		// perform backwards compat value upgrades
 		TArray<URigVMGraph*> GraphsToValidate = GetAllModels();
@@ -786,6 +766,38 @@ void UControlRigBlueprint::DecrementVMRecompileBracket()
 	}
 }
 
+void UControlRigBlueprint::RefreshAllModels()
+{
+	TGuardValue<bool> IsCompilingGuard(bIsCompiling, true);
+	
+	TArray<URigVMGraph*> GraphsToDetach;
+	GraphsToDetach.Add(GetModel());
+	GraphsToDetach.Add(GetLocalFunctionLibrary());
+
+	if (ensure(IsInGameThread()))
+	{
+		for (URigVMGraph* GraphToDetach : GraphsToDetach)
+		{
+			URigVMController* Controller = GetOrCreateController(GraphToDetach);
+			// temporarily disable default value validation during load time, serialized values should always be accepted
+			TGuardValue<bool> PerGraphDisablePinDefaultValueValidation(Controller->bValidatePinDefaults, false);
+			Controller->DetachLinksFromPinObjects();
+			TArray<URigVMNode*> Nodes = GraphToDetach->GetNodes();
+			for (URigVMNode* Node : Nodes)
+			{
+				Controller->RepopulatePinsOnNode(Node, true, false, true);
+			}
+		}
+		SetupPinRedirectorsForBackwardsCompatibility();
+	}
+
+	for (URigVMGraph* GraphToDetach : GraphsToDetach)
+	{
+		URigVMController* Controller = GetOrCreateController(GraphToDetach);
+		Controller->ReattachLinksToPinObjects(true /* follow redirectors */, nullptr, false, true);
+	}
+}
+
 void UControlRigBlueprint::HandleReportFromCompiler(EMessageSeverity::Type InSeverity, UObject* InSubject, const FString& InMessage)
 {
 	UObject* SubjectForMessage = InSubject;
diff --git a/Engine/Plugins/Animation/ControlRig/Source/ControlRigDeveloper/Public/ControlRigBlueprint.h b/Engine/Plugins/Animation/ControlRig/Source/ControlRigDeveloper/Public/ControlRigBlueprint.h
--- a/Engine/Plugins/Animation/ControlRig/Source/ControlRigDeveloper/Public/ControlRigBlueprint.h
+++ b/Engine/Plugins/Animation/ControlRig/Source/ControlRigDeveloper/Public/ControlRigBlueprint.h
@@ -267,6 +267,12 @@ public:
 	void IncrementVMRecompileBracket();
 	void DecrementVMRecompileBracket();
 
+	// this is needed since even after load
+	// model data can change while the Control Rig BP is not opened
+	// for example, if a user defined struct changed after BP load,
+	// any pin that references the struct needs to be regenerated
+	void RefreshAllModels();
+
 	UFUNCTION(BlueprintCallable, Category = "Control Rig Blueprint")
 	void RequestControlRigInit();
 
diff --git a/Engine/Plugins/Animation/ControlRig/Source/ControlRigEditor/Private/ControlRigEditorModule.cpp b/Engine/Plugins/Animation/ControlRig/Source/ControlRigEditor/Private/ControlRigEditorModule.cpp
--- a/Engine/Plugins/Animation/ControlRig/Source/ControlRigEditor/Private/ControlRigEditorModule.cpp
+++ b/Engine/Plugins/Animation/ControlRig/Source/ControlRigEditor/Private/ControlRigEditorModule.cpp
@@ -122,6 +122,8 @@
 #include "Channels/SCurveEditorKeyBarView.h"
 #include "ControlRigSpaceChannelCurveModel.h"
 #include "ControlRigSpaceChannelEditors.h"
+#include "UserDefinedStructure/UserDefinedStructEditorData.h"
+#include "UObject/FieldIterator.h"
 
 #define LOCTEXT_NAMESPACE "ControlRigEditorModule"
 
@@ -2181,6 +2183,175 @@ void FControlRigEditorModule::GetContextMenuActions(const UControlRigGraphSchema
 	}
 }
 
+void FControlRigEditorModule::PreChange(const UUserDefinedStruct* Changed,
+	FStructureEditorUtils::EStructureEditorChangeInfo ChangedType)
+{
+	// the following is similar to
+	// FUserDefinedStructureCompilerInner::ReplaceStructWithTempDuplicate()
+	// it is necessary since existing rigs need to be kept valid until after PreBPCompile
+	// there are other systems, such as sequencer, that might need to evaluate the rig
+	// for one last time during PreBPCompile
+	// Overrall sequence of events
+	// PreStructChange --1--> PostStructChange
+	//                              --2--> PreBPCompile --3--> PostBPCompile
+	
+	UUserDefinedStruct* StructureToReinstance = (UUserDefinedStruct*)Changed;
+	if (StructureToReinstance)
+	{
+		UUserDefinedStruct* DuplicatedStruct = NULL;
+		{
+			const FString ReinstancedName = FString::Printf(TEXT("STRUCT_REINST_%s"), *StructureToReinstance->GetName());
+			const FName UniqueName = MakeUniqueObjectName(GetTransientPackage(), UUserDefinedStruct::StaticClass(), FName(*ReinstancedName));
+
+			TGuardValue<bool> IsDuplicatingClassForReinstancing(GIsDuplicatingClassForReinstancing, true);
+			DuplicatedStruct = (UUserDefinedStruct*)StaticDuplicateObject(StructureToReinstance, GetTransientPackage(), UniqueName, ~RF_Transactional); 
+		}
+
+		DuplicatedStruct->Guid = StructureToReinstance->Guid;
+		DuplicatedStruct->Bind();
+		DuplicatedStruct->StaticLink(true);
+		DuplicatedStruct->PrimaryStruct = StructureToReinstance;
+		DuplicatedStruct->Status = EUserDefinedStructureStatus::UDSS_Duplicate;
+		DuplicatedStruct->SetFlags(RF_Transient);
+		DuplicatedStruct->AddToRoot();
+
+		CastChecked<UUserDefinedStructEditorData>(DuplicatedStruct->EditorData)->RecreateDefaultInstance();
+
+		// List of unique classes and structs to regenerate bytecode and property referenced objects list
+		TSet<UStruct*> StructsToRegenerateReferencesFor;
+
+		for (TAllFieldsIterator<FStructProperty> FieldIt(RF_NoFlags, EInternalObjectFlags::PendingKill); FieldIt; ++FieldIt)
+		{
+			FStructProperty* StructProperty = *FieldIt;
+			if (StructProperty && (StructureToReinstance == StructProperty->Struct))
+			{
+				// make sure variable properties on the BP is patched
+				// since active rig instance still references it
+				if (UControlRigBlueprintGeneratedClass* OwnerClass = Cast<UControlRigBlueprintGeneratedClass>(StructProperty->GetOwnerClass()))
+				{
+					if (UControlRigBlueprint* FoundBlueprint = Cast<UControlRigBlueprint>(OwnerClass->ClassGeneratedBy))
+					{
+						StructProperty->Struct = DuplicatedStruct;
+						StructsToRegenerateReferencesFor.Add(OwnerClass);
+					}
+				}
+				// similar story, VM instructions reference properties on the GeneratorClass
+				if (URigVMMemoryStorageGeneratorClass* OwnerClass = Cast<URigVMMemoryStorageGeneratorClass>(StructProperty->GetOwnerStruct()))
+				{
+					StructProperty->Struct = DuplicatedStruct;
+					StructsToRegenerateReferencesFor.Add(OwnerClass);
+				}
+			}
+		}
+
+		// Make sure we update the list of objects referenced by structs after we replaced the struct in FStructProperties
+		for (UStruct* Struct : StructsToRegenerateReferencesFor)
+		{
+			Struct->CollectBytecodeAndPropertyReferencedObjects();
+			
+			// refresh these since VM caching references them
+			if (URigVMMemoryStorageGeneratorClass* GeneratorClass = Cast<URigVMMemoryStorageGeneratorClass>(Struct))
+			{
+				GeneratorClass->RefreshLinkedProperties();
+				GeneratorClass->RefreshPropertyPaths();	
+			}
+		}
+
+		// as rigs are re-instanced, the duplicated struct will be GCed
+		DuplicatedStruct->RemoveFromRoot();
+	}
+	
+	// in the future we could only invalidate caches on affected rig instances, it shouldn't make too much of a difference though
+	for (TObjectIterator<UControlRig> It(RF_Transient | RF_ClassDefaultObject, /** bIncludeDerivedClasses */ true, /** InternalExcludeFlags */ EInternalObjectFlags::PendingKill); It; ++It)
+	{
+		UControlRig* Rig = *It;
+		// rebuild property list and property path list
+		Rig->GetVM()->InvalidateCachedMemory();
+	}
+}
+
+void FControlRigEditorModule::PostChange(const UUserDefinedStruct* Changed,
+                                   FStructureEditorUtils::EStructureEditorChangeInfo ChangedType)
+{
+	TArray<UControlRigBlueprint*> BlueprintsToRefresh;
+	for (TObjectIterator<URigVMPin> It(RF_Transient | RF_ClassDefaultObject, /** bIncludeDerivedClasses */ true, /** InternalExcludeFlags */ EInternalObjectFlags::PendingKill); It; ++It)
+	{
+		URigVMPin* Pin = *It;
+		// GetCPPTypeObject also makes sure the pin's type information is update to date
+		if (Pin && Pin->GetCPPTypeObject() == Changed)
+		{
+			if (UControlRigBlueprint* RigBlueprint = Pin->GetTypedOuter<UControlRigBlueprint>())
+			{
+				BlueprintsToRefresh.AddUnique(RigBlueprint);
+				
+				// this pin is part of a function definition
+				// update all BP that uses this function
+				if (Pin->GetGraph() == RigBlueprint->GetLocalFunctionLibrary())
+				{
+					TArray< TSoftObjectPtr<URigVMFunctionReferenceNode> > References;
+					References = RigBlueprint->FunctionLibrary->GetReferencesForFunction(Pin->GetNode()->GetFName());
+					
+					for (const TSoftObjectPtr<URigVMFunctionReferenceNode>& Reference : References)
+					{
+						URigVMFunctionReferenceNode* RefNode = Reference.LoadSynchronous();
+						if (!RefNode)
+						{
+							continue;
+						}
+						
+						if (UControlRigBlueprint* FunctionUserBlueprint = RefNode->GetTypedOuter<UControlRigBlueprint>())
+						{
+							BlueprintsToRefresh.AddUnique(FunctionUserBlueprint);
+						}
+					}	
+				}
+			}
+		}
+	}
+
+	for (UControlRigBlueprint* RigBlueprint : BlueprintsToRefresh)
+	{
+		// refresh all pins
+		RigBlueprint->RefreshAllModels();
+		// reflect changes in the editor
+		RigBlueprint->RebuildGraphFromModel();
+		RigBlueprint->MarkPackageDirty();
+	}
+	
+	for (UControlRigBlueprint* RigBlueprint : BlueprintsToRefresh)
+	{
+		// this should make sure variables in BP are updated with the latest struct object
+		// otherwise RigVMCompiler validation would complain about variable type - pin type mismatch
+		FCompilerResultsLog	ResultsLog;
+		FKismetEditorUtilities::CompileBlueprint(RigBlueprint, EBlueprintCompileOptions::None, &ResultsLog);
+		
+		// BP compiler always initialize the new CDO by copying from the old CDO,
+		// however, in case that a BP variable type has changed, the data old CDO would be invalid because
+		// while the old memory container still references the temp duplicated struct we created during PreChange()
+		// registers that reference the BP variable would be referencing the new struct as a result of
+		// FKismetCompilerContext::CompileClassLayout, so type mismatch would invalidate relevant copy operations
+		// so to simplify things, here we just reset all rigs upon error
+		if (ResultsLog.NumErrors > 0)
+		{
+			UControlRigBlueprintGeneratedClass* RigClass = RigBlueprint->GetControlRigBlueprintGeneratedClass();
+			UControlRig* CDO = Cast<UControlRig>(RigClass->GetDefaultObject(true /* create if needed */));
+			if (CDO->GetVM() != nullptr)
+			{
+				CDO->GetVM()->Reset();
+			}
+			TArray<UObject*> ArchetypeInstances;
+			CDO->GetArchetypeInstances(ArchetypeInstances);
+			for (UObject* Instance : ArchetypeInstances)
+			{
+				if (UControlRig* InstanceRig = Cast<UControlRig>(Instance))
+				{
+					InstanceRig->GetVM()->Reset();
+				}
+			}
+		}
+	}
+}
+
 IMPLEMENT_MODULE(FControlRigEditorModule, ControlRigEditor)
 
 #undef LOCTEXT_NAMESPACE
diff --git a/Engine/Plugins/Animation/ControlRig/Source/ControlRigEditor/Private/ControlRigEditorModule.h b/Engine/Plugins/Animation/ControlRig/Source/ControlRigEditor/Private/ControlRigEditorModule.h
--- a/Engine/Plugins/Animation/ControlRig/Source/ControlRigEditor/Private/ControlRigEditorModule.h
+++ b/Engine/Plugins/Animation/ControlRig/Source/ControlRigEditor/Private/ControlRigEditorModule.h
@@ -56,6 +56,9 @@ public:
 	virtual FConnectionDrawingPolicy* CreateConnectionDrawingPolicy(int32 InBackLayerID, int32 InFrontLayerID, float InZoomFactor, const FSlateRect& InClippingRect, class FSlateWindowElementList& InDrawElements, class UEdGraph* InGraphObj) override;
 	virtual void GetContextMenuActions(const UControlRigGraphSchema* Schema, class UToolMenu* Menu, class UGraphNodeContextMenuContext* Context) const override;
 
+	/* FStructureEditorUtils::INotifyOnStructChanged Interface, used to respond to changes to user defined structs */
+	virtual void PreChange(const UUserDefinedStruct* Changed, FStructureEditorUtils::EStructureEditorChangeInfo ChangedType) override;
+	virtual void PostChange(const UUserDefinedStruct* Changed, FStructureEditorUtils::EStructureEditorChangeInfo ChangedType) override;
 private:
 	/** Handle a new animation controller blueprint being created */
 	void HandleNewBlueprintCreated(UBlueprint* InBlueprint);
diff --git a/Engine/Plugins/Animation/ControlRig/Source/ControlRigEditor/Private/Editor/ControlRigEditor.cpp b/Engine/Plugins/Animation/ControlRig/Source/ControlRigEditor/Private/Editor/ControlRigEditor.cpp
--- a/Engine/Plugins/Animation/ControlRig/Source/ControlRigEditor/Private/Editor/ControlRigEditor.cpp
+++ b/Engine/Plugins/Animation/ControlRig/Source/ControlRigEditor/Private/Editor/ControlRigEditor.cpp
@@ -337,6 +337,11 @@ void FControlRigEditor::InitControlRigEditor(const EToolkitMode::Type Mode, cons
 
 	CommonInitialization(ControlRigBlueprints, false);
 
+	// user-defined-struct can change even after load
+	// refresh the models such that pins are updated to match
+	// the latest struct member layout
+	InControlRigBlueprint->RefreshAllModels();
+
 	{
 		TArray<UEdGraph*> EdGraphs;
 		InControlRigBlueprint->GetAllGraphs(EdGraphs);
diff --git a/Engine/Plugins/Animation/ControlRig/Source/ControlRigEditor/Private/Graph/SControlRigGraphNode.cpp b/Engine/Plugins/Animation/ControlRig/Source/ControlRigEditor/Private/Graph/SControlRigGraphNode.cpp
--- a/Engine/Plugins/Animation/ControlRig/Source/ControlRigEditor/Private/Graph/SControlRigGraphNode.cpp
+++ b/Engine/Plugins/Animation/ControlRig/Source/ControlRigEditor/Private/Graph/SControlRigGraphNode.cpp
@@ -96,6 +96,8 @@ void SControlRigGraphNode::Construct( const FArguments& InArgs )
 	}
 
 	TArray<URigVMPin*> RootModelPins = ModelNode->GetPins();
+	// orphaned pins are appended to the end of pin list on each side of the node
+	RootModelPins.Append(ModelNode->GetOrphanedPins());
 	TArray<URigVMPin*> ModelPins;
 
 	// sort model pins
diff --git a/Engine/Plugins/Animation/ControlRig/Source/ControlRigEditor/Public/IControlRigEditorModule.h b/Engine/Plugins/Animation/ControlRig/Source/ControlRigEditor/Public/IControlRigEditorModule.h
--- a/Engine/Plugins/Animation/ControlRig/Source/ControlRigEditor/Public/IControlRigEditorModule.h
+++ b/Engine/Plugins/Animation/ControlRig/Source/ControlRigEditor/Public/IControlRigEditorModule.h
@@ -8,6 +8,7 @@
 #include "Framework/MultiBox/MultiBoxExtender.h"
 #include "Toolkits/AssetEditorToolkit.h"
 #include "IControlRigEditor.h"
+#include "Kismet2/StructureEditorUtils.h"
 
 DECLARE_LOG_CATEGORY_EXTERN(LogControlRigEditor, Log, All);
 
@@ -17,7 +18,7 @@ class UControlRigGraphNode;
 class UControlRigGraphSchema;
 class FConnectionDrawingPolicy;
 
-class IControlRigEditorModule : public IModuleInterface, public IHasMenuExtensibility, public IHasToolBarExtensibility
+class IControlRigEditorModule : public IModuleInterface, public IHasMenuExtensibility, public IHasToolBarExtensibility, public FStructureEditorUtils::INotifyOnStructChanged
 {
 public:
 
diff --git a/Engine/Source/Developer/RigVMDeveloper/Private/RigVMCompiler/RigVMCompiler.cpp b/Engine/Source/Developer/RigVMDeveloper/Private/RigVMCompiler/RigVMCompiler.cpp
--- a/Engine/Source/Developer/RigVMDeveloper/Private/RigVMCompiler/RigVMCompiler.cpp
+++ b/Engine/Source/Developer/RigVMDeveloper/Private/RigVMCompiler/RigVMCompiler.cpp
@@ -172,9 +172,28 @@ bool URigVMCompiler::Compile(URigVMGraph* InGraph, URigVMController* InControlle
 	for(int32 GraphIndex=0; GraphIndex<VisitedGraphs.Num(); GraphIndex++)
 	{
 		URigVMGraph* VisitedGraph = VisitedGraphs[GraphIndex];
+		
+		{
+			FRigVMControllerGraphGuard Guard(InController, VisitedGraph, false);
+			// make sure variables are up to date before validating other things.
+			// that is, make sure their cpp type and type object agree with each other
+			InController->EnsureLocalVariableValidity();
+		}
+		
 		for(URigVMNode* ModelNode : VisitedGraph->GetNodes())
 		{
 			FRigVMControllerGraphGuard Guard(InController, VisitedGraph, false);
+
+			// make sure pins are up to date before validating other things.
+			// that is, make sure their cpp type and type object agree with each other
+			for(URigVMPin* Pin : ModelNode->Pins)
+			{
+				if(!URigVMController::EnsurePinValidity(Pin, true))
+				{
+					return false;
+				}
+			}
+			
 			if(!InController->RemoveUnusedOrphanedPins(ModelNode, true))
 			{
 				static const FString LinkedMessage = TEXT("Node @@ uses pins that no longer exist. Please rewire the links and re-compile.");
@@ -335,6 +354,13 @@ bool URigVMCompiler::Compile(URigVMGraph* InGraph, URigVMController* InControlle
 	{
 		for (URigVMLibraryNode* LibraryNode : FunctionLibrary->GetFunctions())
 		{
+			{
+				FRigVMControllerGraphGuard Guard(InController, LibraryNode->GetContainedGraph(), false);
+				// make sure variables are up to date before validating other things.
+				// that is, make sure their cpp type and type object agree with each other
+				InController->EnsureLocalVariableValidity();
+			}
+
 			for (FRigVMGraphVariableDescription& Variable : LibraryNode->GetContainedGraph()->LocalVariables)
 			{
 				FString Path = FString::Printf(TEXT("LocalVariableDefault::%s|%s::Const"), *LibraryNode->GetFName().ToString(), *Variable.Name.ToString());
diff --git a/Engine/Source/Developer/RigVMDeveloper/Private/RigVMModel/RigVMController.cpp b/Engine/Source/Developer/RigVMDeveloper/Private/RigVMModel/RigVMController.cpp
--- a/Engine/Source/Developer/RigVMDeveloper/Private/RigVMModel/RigVMController.cpp
+++ b/Engine/Source/Developer/RigVMDeveloper/Private/RigVMModel/RigVMController.cpp
@@ -17,6 +17,7 @@
 #include "Algo/Sort.h"
 #include "RigVMPythonUtils.h"
 #include "RigVMTypeUtils.h"
+#include "Engine/UserDefinedStruct.h"
 
 #if WITH_EDITOR
 #include "Exporters/Exporter.h"
@@ -1035,6 +1036,10 @@ URigVMVariableNode* URigVMController::AddVariableNode(const FName& InVariableNam
 	{
 		ValuePin->CPPType = ExternalVariable.TypeName.ToString();
 		ValuePin->CPPTypeObject = ExternalVariable.TypeObject;
+		if (ValuePin->CPPTypeObject)
+		{
+			ValuePin->CPPTypeObjectPath = *ValuePin->CPPTypeObject->GetPathName();
+		}
 		ValuePin->bIsDynamicArray = ExternalVariable.bIsArray;
 
 		if(ValuePin->bIsDynamicArray && !RigVMTypeUtils::IsArrayType(ValuePin->CPPType))
@@ -8483,21 +8488,19 @@ bool URigVMController::ChangeExposedPinType(const FName& InPinName, const FStrin
 
 	if (URigVMFunctionEntryNode* EntryNode = Graph->GetEntryNode())
 	{
-		TArray<URigVMLink*> Links = EntryNode->GetLinks();
+		const TArray<URigVMLink*> Links = EntryNode->GetLinks();
 		DetachLinksFromPinObjects(&Links, true);
-		RepopulatePinsOnNode(EntryNode, true, true);
-		ReattachLinksToPinObjects(false, &Links, true);
-		
+		RepopulatePinsOnNode(EntryNode, true, true, bSetupOrphanPins);
+		ReattachLinksToPinObjects(true, &Links, true, bSetupOrphanPins);
 		RemoveUnusedOrphanedPins(EntryNode, true);
 	}
 	
 	if (URigVMFunctionReturnNode* ReturnNode = Graph->GetReturnNode())
 	{
-		TArray<URigVMLink*> Links = ReturnNode->GetLinks();
+		const TArray<URigVMLink*> Links = ReturnNode->GetLinks();
 		DetachLinksFromPinObjects(&Links, true);
-		RepopulatePinsOnNode(ReturnNode, true, true);
-		ReattachLinksToPinObjects(false, &Links, true);
-		
+		RepopulatePinsOnNode(ReturnNode, true, true, bSetupOrphanPins);
+		ReattachLinksToPinObjects(true, &Links, true, bSetupOrphanPins);
 		RemoveUnusedOrphanedPins(ReturnNode, true);
 	}
 
@@ -10077,6 +10080,11 @@ URigVMSelectNode* URigVMController::AddSelectNode(const FString& InCPPType, cons
 	ResultPin->Direction = ERigVMPinDirection::Output;
 	AddNodePin(Node, ResultPin);
 
+	if (ResultPin->IsStruct())
+	{
+		AddPinsForStruct(ResultPin->GetScriptStruct(), Node, ResultPin, ResultPin->Direction, FString(), false);
+	}
+
 	Graph->Nodes.Add(Node);
 
 	Notify(ERigVMGraphNotifType::NodeAdded, Node);
@@ -11555,7 +11563,7 @@ int32 URigVMController::ReattachLinksToPinObjects(bool bFollowCoreRedirectors, c
 			if (URigVMCollapseNode* CollapseNode = Cast<URigVMCollapseNode>(Node))
 			{
 				FRigVMControllerGraphGuard GraphGuard(this, CollapseNode->GetContainedGraph(), false);
-				ReattachLinksToPinObjects(bFollowCoreRedirectors, nullptr);
+				ReattachLinksToPinObjects(bFollowCoreRedirectors, nullptr, bNotify, bSetupOrphanedPins);
 			}
 		}
 	}
@@ -11767,6 +11775,9 @@ void URigVMController::RepopulatePinsOnNode(URigVMNode* InNode, bool bFollowCore
 	URigVMCollapseNode* CollapseNode = Cast<URigVMCollapseNode>(InNode);
 	URigVMFunctionReferenceNode* FunctionRefNode = Cast<URigVMFunctionReferenceNode>(InNode);
 	URigVMVariableNode* VariableNode = Cast<URigVMVariableNode>(InNode);
+	URigVMIfNode* IfNode = Cast<URigVMIfNode>(InNode);
+	URigVMSelectNode* SelectNode = Cast<URigVMSelectNode>(InNode);
+	URigVMArrayNode* ArrayNode = Cast<URigVMArrayNode>(InNode);
 
 	TGuardValue<bool> EventuallySuspendNotifs(bSuspendNotifications, !bNotify);
 	FScopeLock Lock(&PinPathCoreRedirectorsLock);
@@ -11937,6 +11948,43 @@ void URigVMController::RepopulatePinsOnNode(URigVMNode* InNode, bool bFollowCore
 	}
 	else if (CollapseNode)
 	{
+		// PinStates were already saved earlier, and will be applied later, no need to do it here 
+
+		// since we are replacing existing pins, their names have to be saved
+		// and assigned to the new pins only after we removed the old ones
+		TArray<TTuple<URigVMPin*, FName>> NewRootPinInfos;
+		for (URigVMPin* RootPin : InNode->Pins)
+		{
+			URigVMPin* NewRootPin = NewObject<URigVMPin>(InNode);
+			ConfigurePinFromPin(NewRootPin, RootPin);
+			EnsurePinValidity(NewRootPin, false);
+		
+			NewRootPinInfos.Add(TTuple<URigVMPin*, FName>(NewRootPin, RootPin->GetFName()));
+		}
+		
+		RemovePinsDuringRepopulate(InNode, InNode->Pins, bNotify, bSetupOrphanedPins);
+		
+		for (TTuple<URigVMPin*, FName> NewRootPinInfo : NewRootPinInfos)
+		{
+			RenameObject(NewRootPinInfo.Key, *NewRootPinInfo.Value.ToString(), InNode);
+			AddNodePin(InNode, NewRootPinInfo.Key);
+		}
+		
+		for (URigVMPin* Pin : InNode->Pins)
+		{
+			if (Pin->IsStruct())
+			{
+				AddPinsForStruct(Pin->GetScriptStruct(), InNode, Pin, Pin->GetDirection(), FString(), false);
+			}
+			Notify(ERigVMGraphNotifType::PinAdded, Pin);
+		}
+
+		if (CollapseNode->GetOuter()->IsA<URigVMFunctionLibrary>())
+		{
+			// no need to notify since the function library graph is invisible anyway
+			RemoveUnusedOrphanedPins(CollapseNode, false);
+		}
+
 		FRigVMControllerGraphGuard GraphGuard(this, CollapseNode->GetContainedGraph(), false);
 		// need to get a copy of the node array since the following function could remove nodes from the graph
 		// we don't want to remove elements from the array we are iterating over.
@@ -11961,7 +12009,8 @@ void URigVMController::RepopulatePinsOnNode(URigVMNode* InNode, bool bFollowCore
 			{
 				URigVMPin* NewPin = NewObject<URigVMPin>(InNode, ReferencedPin->GetFName());
 				ConfigurePinFromPin(NewPin, ReferencedPin);
-
+				EnsurePinValidity(NewPin, false);
+				
 				AddNodePin(InNode, NewPin);
 
 				if (NewPin->IsStruct())
@@ -11975,7 +12024,39 @@ void URigVMController::RepopulatePinsOnNode(URigVMNode* InNode, bool bFollowCore
 			ApplyPinStates(InNode, ReferencedPinStates);
 		}
 	}
-
+	else if (IfNode || SelectNode || ArrayNode)
+	{
+		// PinStates were already saved earlier, and will be applied later, no need to do it here
+	 
+		// since we are replacing existing pins, their names have to be saved
+		// and assigned to the new pins only after we removed old ones	
+		TArray<TTuple<URigVMPin*, FName>> NewRootPinInfos;
+		for (URigVMPin* RootPin : InNode->Pins)
+		{
+			URigVMPin* NewRootPin = NewObject<URigVMPin>(InNode);
+			ConfigurePinFromPin(NewRootPin, RootPin);
+			EnsurePinValidity(NewRootPin, false);
+		
+			NewRootPinInfos.Add(TTuple<URigVMPin*, FName>(NewRootPin, RootPin->GetFName()));
+		}
+		
+		RemovePinsDuringRepopulate(InNode, InNode->Pins, bNotify, bSetupOrphanedPins);
+		
+		for (TTuple<URigVMPin*, FName> NewRootPinInfo : NewRootPinInfos)
+		{
+			RenameObject(NewRootPinInfo.Key, *NewRootPinInfo.Value.ToString(), InNode);
+			AddNodePin(InNode, NewRootPinInfo.Key);
+		}
+		
+		for (URigVMPin* Pin : InNode->Pins)
+		{
+			if (Pin->IsStruct())
+			{
+				AddPinsForStruct(Pin->GetScriptStruct(), InNode, Pin, Pin->GetDirection(), FString(), false);
+			}
+			Notify(ERigVMGraphNotifType::PinAdded, Pin);
+		}	
+	}
 	else
 	{
 		return;
@@ -12664,9 +12745,20 @@ bool URigVMController::ChangePinType(URigVMPin* InPin, const FString& InCPPType,
 	}
 
 	UObject* CPPTypeObject = URigVMPin::FindObjectFromCPPTypeObjectPath(InCPPTypeObjectPath.ToString());
-	if (InPin->CPPType == InCPPType && InPin->CPPTypeObject == CPPTypeObject)
+
+	// always refresh pin type if it is a user defined struct, whose internal layout can change at anytime
+	bool bForceRefresh = false;
+	if (CPPTypeObject && CPPTypeObject->IsA<UUserDefinedStruct>())
 	{
-		return true;
+		bForceRefresh = true;
+	}
+
+	if (!bForceRefresh)
+	{
+		if (InPin->CPPType == InCPPType && InPin->CPPTypeObject == CPPTypeObject)
+		{
+			return true;
+		}
 	}
 
 	return ChangePinType(InPin, InCPPType, CPPTypeObject, bSetupUndoRedo, bSetupOrphanPins, bBreakLinks, bRemoveSubPins);
@@ -12909,6 +13001,8 @@ bool URigVMController::EnsurePinValidity(URigVMPin* InPin, bool bRecursive)
 	// check if the CPPTypeObject is set up correctly.
 	if(FRigVMPropertyDescription::RequiresCPPTypeObject(InPin->GetCPPType()))
 	{
+		// GetCPPTypeObject attempts to update pin type information to the latest
+		// without testing for redirector
 		if(InPin->GetCPPTypeObject() == nullptr)
 		{
 			// try to find the CPPTypeObject by name
@@ -12966,6 +13060,19 @@ void URigVMController::ValidatePin(URigVMPin* InPin)
 	FRigVMPropertyDescription(InPin->GetFName(), InPin->GetCPPType(), InPin->GetCPPTypeObject(), InPin->GetDefaultValue());
 }
 
+void URigVMController::EnsureLocalVariableValidity()
+{
+	if (URigVMGraph* Graph = GetGraph())
+	{
+		for (FRigVMGraphVariableDescription& Variable : Graph->LocalVariables)
+		{
+			// CPPType can become invalid when the type object is defined by
+			// an asset that have changed name or asset path, user defined struct is one possibility
+			Variable.CPPType = PostProcessCPPType(Variable.CPPType, Variable.CPPTypeObject);
+		}
+	}
+}
+
 FRigVMExternalVariable URigVMController::GetVariableByName(const FName& InExternalVariableName, const bool bIncludeInputArguments)
 {
 	TArray<FRigVMExternalVariable> Variables = GetAllVariables(bIncludeInputArguments);
diff --git a/Engine/Source/Developer/RigVMDeveloper/Private/RigVMModel/RigVMPin.cpp b/Engine/Source/Developer/RigVMDeveloper/Private/RigVMModel/RigVMPin.cpp
--- a/Engine/Source/Developer/RigVMDeveloper/Private/RigVMModel/RigVMPin.cpp
+++ b/Engine/Source/Developer/RigVMDeveloper/Private/RigVMModel/RigVMPin.cpp
@@ -1167,7 +1167,7 @@ bool URigVMPin::IsOrphanPin() const
 	return false;
 }
 
-void URigVMPin::UpdateCPPTypeObjectIfRequired() const
+void URigVMPin::UpdateTypeInformationIfRequired() const
 {
 	if (CPPTypeObject == nullptr)
 	{
@@ -1177,11 +1177,18 @@ void URigVMPin::UpdateCPPTypeObjectIfRequired() const
 			MutableThis->CPPTypeObject = FindObjectFromCPPTypeObjectPath(CPPTypeObjectPath.ToString());
 		}
 	}
+
+	if (CPPTypeObject)
+	{
+		// refresh the type string 
+		URigVMPin* MutableThis = (URigVMPin*)this;
+		MutableThis->CPPType = URigVMController::PostProcessCPPType(CPPType, CPPTypeObject);
+	}
 }
 
 UObject* URigVMPin::GetCPPTypeObject() const
 {
-	UpdateCPPTypeObjectIfRequired();
+	UpdateTypeInformationIfRequired();
 	return CPPTypeObject;
 }
 
diff --git a/Engine/Source/Developer/RigVMDeveloper/Public/RigVMModel/RigVMController.h b/Engine/Source/Developer/RigVMDeveloper/Public/RigVMModel/RigVMController.h
--- a/Engine/Source/Developer/RigVMDeveloper/Public/RigVMModel/RigVMController.h
+++ b/Engine/Source/Developer/RigVMDeveloper/Public/RigVMModel/RigVMController.h
@@ -888,6 +888,10 @@ private:
 	static bool EnsurePinValidity(URigVMPin* InPin, bool bRecursive);
 	static void ValidatePin(URigVMPin* InPin);
 
+	// recreate the CPP type strings for variables that reference a type object
+	// they can get out of sync when the variable references a user defined struct
+	void EnsureLocalVariableValidity();
+	
 	FRigVMExternalVariable GetVariableByName(const FName& InExternalVariableName, const bool bIncludeInputArguments = false);
 	TArray<FRigVMExternalVariable> GetAllVariables(const bool bIncludeInputArguments = false);
 
diff --git a/Engine/Source/Developer/RigVMDeveloper/Public/RigVMModel/RigVMPin.h b/Engine/Source/Developer/RigVMDeveloper/Public/RigVMModel/RigVMPin.h
--- a/Engine/Source/Developer/RigVMDeveloper/Public/RigVMModel/RigVMPin.h
+++ b/Engine/Source/Developer/RigVMDeveloper/Public/RigVMModel/RigVMPin.h
@@ -396,7 +396,7 @@ public:
 
 private:
 
-	void UpdateCPPTypeObjectIfRequired() const;
+	void UpdateTypeInformationIfRequired() const;
 	void SetNameFromIndex();
 
 	void GetExposedPinChainImpl(TArray<const URigVMPin*>& OutExposedPins, TArray<const URigVMPin*>& VisitedPins) const;
@@ -432,7 +432,11 @@ private:
 	UPROPERTY()
 	FString CPPType;
 
-	UPROPERTY(transient)
+	// serialize object ptr here to keep track of the latest version of the type object,
+	// type object can reference assets like user defined struct, which can be renamed
+	// or moved to new locations, serializing the type object with the pin
+	// ensure automatic update whenever those things happen
+	UPROPERTY()
 	TObjectPtr<UObject> CPPTypeObject;
 
 	UPROPERTY()
diff --git a/Engine/Source/Runtime/RigVM/Public/RigVMCore/RigVM.h b/Engine/Source/Runtime/RigVM/Public/RigVMCore/RigVM.h
--- a/Engine/Source/Runtime/RigVM/Public/RigVMCore/RigVM.h
+++ b/Engine/Source/Runtime/RigVM/Public/RigVMCore/RigVM.h
@@ -747,12 +747,14 @@ public:
 	uint32 GetNumExecutions() const { return NumExecutions; }
 	const FRigVMExecuteContext& GetContext() const { return Context; }
 
-	void InvalidateCachedMemory();
-
 private:
 
 	void ResolveFunctionsIfRequired();
 	void RefreshInstructionsIfRequired();
+public:
+	void InvalidateCachedMemory();
+	
+private:
 #if UE_RIGVM_UCLASS_BASED_STORAGE_DISABLED
 	void CacheMemoryHandlesIfRequired(FRigVMMemoryContainerPtrArray InMemory);
 #else
diff --git a/Engine/Source/Runtime/RigVM/Public/RigVMCore/RigVMMemoryStorage.h b/Engine/Source/Runtime/RigVM/Public/RigVMCore/RigVMMemoryStorage.h
--- a/Engine/Source/Runtime/RigVM/Public/RigVMCore/RigVMMemoryStorage.h
+++ b/Engine/Source/Runtime/RigVM/Public/RigVMCore/RigVMMemoryStorage.h
@@ -274,11 +274,13 @@ protected:
 	// Adds a single property to a class.
 	static FProperty* AddProperty(URigVMMemoryStorageGeneratorClass* InClass, const FRigVMPropertyDescription& InProperty, FField** LinkToProperty = nullptr);
 
-private:
+public:
 
 	void RefreshLinkedProperties();
 	void RefreshPropertyPaths();
 
+private:
+
 	// The type of memory of this class
 	ERigVMMemoryType MemoryType;
 
