Control Rig: Check array type for array nodes

diff --git a/Engine/Source/Developer/RigVMDeveloper/Private/RigVMModel/RigVMController.cpp b/Engine/Source/Developer/RigVMDeveloper/Private/RigVMModel/RigVMController.cpp
--- a/Engine/Source/Developer/RigVMDeveloper/Private/RigVMModel/RigVMController.cpp
+++ b/Engine/Source/Developer/RigVMDeveloper/Private/RigVMModel/RigVMController.cpp
@@ -10309,16 +10309,23 @@ URigVMArrayNode* URigVMController::AddArrayNode(ERigVMOpCode InOpCode, const FSt
 		return nullptr;
 	}
 
+	FString CPPType = InCPPType;
+	if(CPPType.StartsWith(RigVMTypeUtils::TArrayPrefix))
+	{
+		// turn "TArray<bool>" to "bool"
+		CPPType = CPPType.Mid(FCString::Strlen(RigVMTypeUtils::TArrayPrefix)).LeftChop(1).TrimStartAndEnd();
+	}
+
 	if (InCPPTypeObject == nullptr)
 	{
-		InCPPTypeObject = URigVMCompiler::GetScriptStructForCPPType(InCPPType);
+		InCPPTypeObject = URigVMCompiler::GetScriptStructForCPPType(CPPType);
 	}
 	if (InCPPTypeObject == nullptr)
 	{
-		InCPPTypeObject = URigVMPin::FindObjectFromCPPTypeObjectPath<UObject>(InCPPType);
+		InCPPTypeObject = URigVMPin::FindObjectFromCPPTypeObjectPath<UObject>(CPPType);
 	}
 
-	FString CPPType = PostProcessCPPType(InCPPType, InCPPTypeObject);
+	CPPType = PostProcessCPPType(CPPType, InCPPTypeObject);
 	
 	const FString Name = GetValidNodeName(InNodeName.IsEmpty() ? FString(TEXT("ArrayNode")) : InNodeName);
 	URigVMArrayNode* Node = NewObject<URigVMArrayNode>(Graph, *Name);
