Control Rig: Fix undo pasting collapse node reaching ensure

diff --git a/Engine/Source/Developer/RigVMDeveloper/Private/RigVMModel/RigVMController.cpp b/Engine/Source/Developer/RigVMDeveloper/Private/RigVMModel/RigVMController.cpp
--- a/Engine/Source/Developer/RigVMDeveloper/Private/RigVMModel/RigVMController.cpp
+++ b/Engine/Source/Developer/RigVMDeveloper/Private/RigVMModel/RigVMController.cpp
@@ -2843,7 +2843,7 @@ TArray<FName> URigVMController::ImportNodesFromText(const FString& InText, bool
 			if(URigVMCollapseNode* CollapseNode = Cast<URigVMCollapseNode>(SubNodes[SubNodeIndex]))
 			{
 				{
-					FRigVMControllerGraphGuard GraphGuard(this, CollapseNode->GetContainedGraph(), bSetupUndoRedo);
+					FRigVMControllerGraphGuard GraphGuard(this, CollapseNode->GetContainedGraph(), false);
 					ReattachLinksToPinObjects();
 				}
 				
