Control Rig: React to pin container type changes correctly

diff --git a/Engine/Source/Developer/RigVMDeveloper/Private/RigVMModel/RigVMController.cpp b/Engine/Source/Developer/RigVMDeveloper/Private/RigVMModel/RigVMController.cpp
--- a/Engine/Source/Developer/RigVMDeveloper/Private/RigVMModel/RigVMController.cpp
+++ b/Engine/Source/Developer/RigVMDeveloper/Private/RigVMModel/RigVMController.cpp
@@ -10310,10 +10310,9 @@ URigVMArrayNode* URigVMController::AddArrayNode(ERigVMOpCode InOpCode, const FSt
 	}
 
 	FString CPPType = InCPPType;
-	if(CPPType.StartsWith(RigVMTypeUtils::TArrayPrefix))
+	if(RigVMTypeUtils::IsArrayType(CPPType))
 	{
-		// turn "TArray<bool>" to "bool"
-		CPPType = CPPType.Mid(FCString::Strlen(RigVMTypeUtils::TArrayPrefix)).LeftChop(1).TrimStartAndEnd();
+		CPPType = RigVMTypeUtils::BaseTypeFromArrayType(CPPType);
 	}
 
 	if (InCPPTypeObject == nullptr)
@@ -12847,6 +12846,7 @@ bool URigVMController::ChangePinType(URigVMPin* InPin, const FString& InCPPType,
 	InPin->CPPType = InCPPType;
 	InPin->CPPTypeObjectPath = CPPTypeObjectPath;
 	InPin->CPPTypeObject = InCPPTypeObject;
+	InPin->bIsDynamicArray = RigVMTypeUtils::IsArrayType(InCPPType);
 	// we might want to use GetPinInitialDefaultValue here for a better default value
 	InPin->DefaultValue = FString();
 
