RigVMCompiler: Fold assignments for array pins

diff --git a/Engine/Source/Developer/RigVMDeveloper/Private/RigVMCompiler/RigVMAST.cpp b/Engine/Source/Developer/RigVMDeveloper/Private/RigVMCompiler/RigVMAST.cpp
--- a/Engine/Source/Developer/RigVMDeveloper/Private/RigVMCompiler/RigVMAST.cpp
+++ b/Engine/Source/Developer/RigVMDeveloper/Private/RigVMCompiler/RigVMAST.cpp
@@ -1485,9 +1485,10 @@ void FRigVMParserAST::FoldAssignments()
 			}
 		}
 		
-		// non-input pins on anything but a reroute node should be skipped
+		// non-input pins on anything but a reroute / array node should be skipped
 		if (TargetPin->GetDirection() != ERigVMPinDirection::Input &&
-			Cast<URigVMRerouteNode>(TargetPin->GetNode()) == nullptr)
+			(Cast<URigVMRerouteNode>(TargetPin->GetNode()) == nullptr) &&
+			(Cast<URigVMArrayNode>(TargetPin->GetNode()) == nullptr))
 		{
 			continue;
 		}
