RigVM: Allow links to be recreate between entry + return nodes

diff --git a/Engine/Source/Developer/RigVMDeveloper/Private/RigVMModel/RigVMController.cpp b/Engine/Source/Developer/RigVMDeveloper/Private/RigVMModel/RigVMController.cpp
--- a/Engine/Source/Developer/RigVMDeveloper/Private/RigVMModel/RigVMController.cpp
+++ b/Engine/Source/Developer/RigVMDeveloper/Private/RigVMModel/RigVMController.cpp
@@ -7610,6 +7610,14 @@ bool URigVMController::AddLink(URigVMPin* OutputPin, URigVMPin* InputPin, bool b
 		FString FailureReason;
 		if (!Graph->CanLink(OutputPin, InputPin, &FailureReason, GetCurrentByteCode()))
 		{
+			if(OutputPin->IsExecuteContext() && InputPin->IsExecuteContext())
+			{
+				if(OutputPin->GetNode()->IsA<URigVMFunctionEntryNode>() &&
+					InputPin->GetNode()->IsA<URigVMFunctionReturnNode>())
+				{
+					return false;
+				}
+			}
 			ReportErrorf(TEXT("Cannot link '%s' to '%s': %s."), *OutputPin->GetPinPath(), *InputPin->GetPinPath(), *FailureReason, GetCurrentByteCode());
 			return false;
 		}
