RigVM: During expansion wire subpin links to reroute nodes

diff --git a/Engine/Source/Developer/RigVMDeveloper/Private/RigVMModel/RigVMController.cpp b/Engine/Source/Developer/RigVMDeveloper/Private/RigVMModel/RigVMController.cpp
--- a/Engine/Source/Developer/RigVMDeveloper/Private/RigVMModel/RigVMController.cpp
+++ b/Engine/Source/Developer/RigVMDeveloper/Private/RigVMModel/RigVMController.cpp
@@ -4353,7 +4353,20 @@ TArray<URigVMNode*> URigVMController::ExpandLibraryNode(URigVMLibraryNode* InNod
 				const TArray<FString>& LibraryPinLinks = *LibraryPinLinksPtr;
 				ensure(LibraryPinLinks.Num() == 1);
 
-				Local::UpdateRemappedSourcePins(LibraryPinPath, LibraryPinLinks[0], RemappedSourcePinsForInputs);
+				const FString SourcePinPath = LibraryPinPath;
+				FString TargetPinPath = LibraryPinLinks[0];
+
+				// if the pin on the library node is represented by a reroute
+				// we need to remap to that instead.
+				if(URigVMPin** ReroutedPinPtr = ReroutedInputPins.Find(SourcePinPath))
+				{
+					if(URigVMPin* ReroutedPin = *ReroutedPinPtr)
+					{
+						TargetPinPath = ReroutedPin->GetPinPath();
+					}
+				}
+
+				Local::UpdateRemappedSourcePins(SourcePinPath, TargetPinPath, RemappedSourcePinsForInputs);
 			}
 		}
 		if (LibraryPin->GetDirection() == ERigVMPinDirection::Output ||
@@ -4364,7 +4377,20 @@ TArray<URigVMNode*> URigVMController::ExpandLibraryNode(URigVMLibraryNode* InNod
 				const TArray<FString>& LibraryPinLinks = *LibraryPinLinksPtr;
 				ensure(LibraryPinLinks.Num() == 1);
 
-				Local::UpdateRemappedSourcePins(LibraryPinPath, LibraryPinLinks[0], RemappedSourcePinsForOutputs);
+				const FString SourcePinPath = LibraryPinPath;
+				FString TargetPinPath = LibraryPinLinks[0];
+
+				// if the pin on the library node is represented by a reroute
+				// we need to remap to that instead.
+				if(URigVMPin** ReroutedPinPtr = ReroutedOutputPins.Find(SourcePinPath))
+				{
+					if(URigVMPin* ReroutedPin = *ReroutedPinPtr)
+					{
+						TargetPinPath = ReroutedPin->GetPinPath();
+					}
+				}
+
+				Local::UpdateRemappedSourcePins(SourcePinPath, TargetPinPath, RemappedSourcePinsForOutputs);
 			}
 		}
 	}
