RigVM: Rely on import text for undo / redo of Collapse, Expand

diff --git a/Engine/Source/Developer/RigVMDeveloper/Private/RigVMModel/RigVMController.cpp b/Engine/Source/Developer/RigVMDeveloper/Private/RigVMModel/RigVMController.cpp
--- a/Engine/Source/Developer/RigVMDeveloper/Private/RigVMModel/RigVMController.cpp
+++ b/Engine/Source/Developer/RigVMDeveloper/Private/RigVMModel/RigVMController.cpp
@@ -3581,14 +3581,16 @@ URigVMCollapseNode* URigVMController::CollapseNodes(const TArray<URigVMNode*>& I
 
 	FRigVMControllerCompileBracketScope CompileScope(this);
 	FRigVMCollapseNodesAction CollapseAction;
-	CollapseAction.Title = TEXT("Collapse Nodes");
+
+	FString CollapseNodeName = GetValidNodeName(InCollapseNodeName.IsEmpty() ? FString(TEXT("CollapseNode")) : InCollapseNodeName);
 
 	if (bSetupUndoRedo)
 	{
+		CollapseAction = FRigVMCollapseNodesAction(this, Nodes, CollapseNodeName); 
+		CollapseAction.Title = TEXT("Collapse Nodes");
 		ActionStack->BeginAction(CollapseAction);
 	}
 
-	FString CollapseNodeName = GetValidNodeName(InCollapseNodeName.IsEmpty() ? FString(TEXT("CollapseNode")) : InCollapseNodeName);
 	URigVMCollapseNode* CollapseNode = NewObject<URigVMCollapseNode>(Graph, *CollapseNodeName);
 	CollapseNode->ContainedGraph = NewObject<URigVMGraph>(CollapseNode, TEXT("ContainedGraph"));
 	CollapseNode->Position = Center;
@@ -3934,11 +3936,6 @@ URigVMCollapseNode* URigVMController::CollapseNodes(const TArray<URigVMNode*>& I
 
 	if (bSetupUndoRedo)
 	{
-		CollapseAction.LibraryNodePath = CollapseNode->GetName();
-		for (URigVMNode* InNode : InNodes)
-		{
-			CollapseAction.CollapsedNodesPaths.Add(InNode->GetName());
-		}
 		ActionStack->EndAction(CollapseAction);
 	}
 
@@ -3970,10 +3967,11 @@ TArray<URigVMNode*> URigVMController::ExpandLibraryNode(URigVMLibraryNode* InNod
 
 	FRigVMControllerCompileBracketScope CompileScope(this);
 	FRigVMExpandNodeAction ExpandAction;
-	ExpandAction.Title = FString::Printf(TEXT("Expand '%s' Node"), *InNode->GetName());
 
 	if (bSetupUndoRedo)
 	{
+		ExpandAction = FRigVMExpandNodeAction(this, InNode);
+		ExpandAction.Title = FString::Printf(TEXT("Expand '%s' Node"), *InNode->GetName());
 		ActionStack->BeginAction(ExpandAction);
 	}
 
@@ -4630,10 +4628,6 @@ TArray<URigVMNode*> URigVMController::ExpandLibraryNode(URigVMLibraryNode* InNod
 	}
 
 	// l) remove the library node from the graph
-	if (bSetupUndoRedo)
-	{
-		ExpandAction.LibraryNodePath = InNode->GetName();
-	}
 	RemoveNode(InNode, false, true);
 
 	if (bSetupUndoRedo)
diff --git a/Engine/Source/Developer/RigVMDeveloper/Private/RigVMModel/RigVMControllerActions.cpp b/Engine/Source/Developer/RigVMDeveloper/Private/RigVMModel/RigVMControllerActions.cpp
--- a/Engine/Source/Developer/RigVMDeveloper/Private/RigVMModel/RigVMControllerActions.cpp
+++ b/Engine/Source/Developer/RigVMDeveloper/Private/RigVMModel/RigVMControllerActions.cpp
@@ -1480,13 +1480,29 @@ FRigVMCollapseNodesAction::FRigVMCollapseNodesAction()
 {
 }
 
-FRigVMCollapseNodesAction::FRigVMCollapseNodesAction(const TArray<URigVMNode*>& InNodes, const FString& InNodePath)
+FRigVMCollapseNodesAction::FRigVMCollapseNodesAction(URigVMController* InController, const TArray<URigVMNode*>& InNodes, const FString& InNodePath)
 	: LibraryNodePath(InNodePath)
 {
+	TArray<FName> NodesToExport;
 	for (URigVMNode* InNode : InNodes)
 	{
+		NodesToExport.Add(InNode->GetFName());
 		CollapsedNodesPaths.Add(InNode->GetName());
+
+		// find the links external to the nodes to be collapsed
+		TArray<URigVMLink*> Links = InNode->GetLinks();
+		for(URigVMLink* Link : Links)
+		{
+			if(InNodes.Contains(Link->GetSourcePin()->GetNode()) &&
+				InNodes.Contains(Link->GetTargetPin()->GetNode()))
+			{
+				continue;
+			}
+			CollapsedNodesLinks.Add(Link->GetPinPathRepresentation());
+		}
 	}
+
+	CollapsedNodesContent = InController->ExportNodesToText(NodesToExport);
 }
 
 bool FRigVMCollapseNodesAction::Undo(URigVMController* InController)
@@ -1495,8 +1511,29 @@ bool FRigVMCollapseNodesAction::Undo(URigVMController* InController)
 	{
 		return false;
 	}
-	TArray<URigVMNode*> ExpandedNodes = InController->ExpandLibraryNode(*LibraryNodePath, false);
-	return ExpandedNodes.Num() == CollapsedNodesPaths.Num();
+
+	// remove the library node
+	if(!InController->RemoveNodeByName(*LibraryNodePath, false, true, false))
+	{
+		return false;
+	}
+
+	const TArray<FName> RecoveredNodes = InController->ImportNodesFromText(CollapsedNodesContent, false, false);
+	if(RecoveredNodes.Num() != CollapsedNodesPaths.Num())
+	{
+		return false;
+	}
+
+	for(const FString& CollapsedNodesLink : CollapsedNodesLinks)
+	{
+		FString Source, Target;
+		if(URigVMLink::SplitPinPathRepresentation(CollapsedNodesLink, Source, Target))
+		{
+			InController->AddLink(Source, Target, false, false);
+		}
+	}
+
+	return true;
 }
 
 bool FRigVMCollapseNodesAction::Redo(URigVMController* InController)
@@ -1522,9 +1559,18 @@ FRigVMExpandNodeAction::FRigVMExpandNodeAction()
 {
 }
 
-FRigVMExpandNodeAction::FRigVMExpandNodeAction(URigVMLibraryNode* InLibraryNode)
+FRigVMExpandNodeAction::FRigVMExpandNodeAction(URigVMController* InController, URigVMLibraryNode* InLibraryNode)
 	: LibraryNodePath(InLibraryNode->GetName())
 {
+	TArray<FName> NodesToExport;
+	NodesToExport.Add(InLibraryNode->GetFName());
+	LibraryNodeContent = InController->ExportNodesToText(NodesToExport);
+
+	TArray<URigVMLink*> Links = InLibraryNode->GetLinks();
+	for(URigVMLink* Link : Links)
+	{
+		LibraryNodeLinks.Add(Link->GetPinPathRepresentation());
+	}
 }
 
 bool FRigVMExpandNodeAction::Undo(URigVMController* InController)
@@ -1534,14 +1580,31 @@ bool FRigVMExpandNodeAction::Undo(URigVMController* InController)
 		return false;
 	}
 
-	TArray<FName> NodeNames;
+	// remove the expanded nodes
 	for (const FString& NodePath : ExpandedNodePaths)
 	{
-		NodeNames.Add(*NodePath);
+		if(!InController->RemoveNodeByName(*NodePath, false, true, false))
+		{
+			return false;
+		}
 	}
 
-	URigVMLibraryNode* LibraryNode = InController->CollapseNodes(NodeNames, *LibraryNodePath, false);
-	return LibraryNode != nullptr;
+	const TArray<FName> RecoveredNodes = InController->ImportNodesFromText(LibraryNodeContent, false, false);
+	if(RecoveredNodes.Num() != 1)
+	{
+		return false;
+	}
+
+	for(const FString& LibraryNodeLink : LibraryNodeLinks)
+	{
+		FString Source, Target;
+		if(URigVMLink::SplitPinPathRepresentation(LibraryNodeLink, Source, Target))
+		{
+			InController->AddLink(Source, Target, false, false);
+		}
+	}
+
+	return true;
 }
 
 bool FRigVMExpandNodeAction::Redo(URigVMController* InController)
diff --git a/Engine/Source/Developer/RigVMDeveloper/Private/RigVMModel/RigVMLink.cpp b/Engine/Source/Developer/RigVMDeveloper/Private/RigVMModel/RigVMLink.cpp
--- a/Engine/Source/Developer/RigVMDeveloper/Private/RigVMModel/RigVMLink.cpp
+++ b/Engine/Source/Developer/RigVMDeveloper/Private/RigVMModel/RigVMLink.cpp
@@ -73,6 +73,11 @@ FString URigVMLink::GetPinPathRepresentation()
 	return FString::Printf(TEXT("%s -> %s"), *GetSourcePin()->GetPinPath(), *GetTargetPin()->GetPinPath());
 }
 
+bool URigVMLink::SplitPinPathRepresentation(const FString& InString, FString& OutSource, FString& OutTarget)
+{
+	return InString.Split(TEXT(" -> "), &OutSource, &OutTarget);
+}
+
 void URigVMLink::PrepareForCopy()
 {
 	if (URigVMPin* CurrentSourcePin = GetSourcePin())
diff --git a/Engine/Source/Developer/RigVMDeveloper/Public/RigVMModel/RigVMControllerActions.h b/Engine/Source/Developer/RigVMDeveloper/Public/RigVMModel/RigVMControllerActions.h
--- a/Engine/Source/Developer/RigVMDeveloper/Public/RigVMModel/RigVMControllerActions.h
+++ b/Engine/Source/Developer/RigVMDeveloper/Public/RigVMModel/RigVMControllerActions.h
@@ -1260,7 +1260,7 @@ struct FRigVMCollapseNodesAction : public FRigVMBaseAction
 public:
 
 	FRigVMCollapseNodesAction();
-	FRigVMCollapseNodesAction(const TArray<URigVMNode*>& InNodes, const FString& InNodePath);
+	FRigVMCollapseNodesAction(URigVMController* InController, const TArray<URigVMNode*>& InNodes, const FString& InNodePath);
 	virtual ~FRigVMCollapseNodesAction() {};
 	virtual bool Undo(URigVMController* InController) override;
 	virtual bool Redo(URigVMController* InController) override;
@@ -1268,8 +1268,14 @@ public:
 	UPROPERTY()
 	FString LibraryNodePath;
 
+	UPROPERTY()
+	FString CollapsedNodesContent;
+
 	UPROPERTY()
 	TArray<FString> CollapsedNodesPaths;
+
+	UPROPERTY()
+	TArray<FString> CollapsedNodesLinks;
 };
 
 /**
@@ -1283,7 +1289,7 @@ struct FRigVMExpandNodeAction : public FRigVMBaseAction
 public:
 
 	FRigVMExpandNodeAction();
-	FRigVMExpandNodeAction(URigVMLibraryNode* InLibraryNode);
+	FRigVMExpandNodeAction(URigVMController* InController, URigVMLibraryNode* InLibraryNode);
 	virtual ~FRigVMExpandNodeAction() {};
 	virtual bool Undo(URigVMController* InController) override;
 	virtual bool Redo(URigVMController* InController) override;
@@ -1291,6 +1297,12 @@ public:
 	UPROPERTY()
 	FString LibraryNodePath;
 
+	UPROPERTY()
+	FString LibraryNodeContent;
+
+	UPROPERTY()
+	TArray<FString> LibraryNodeLinks;
+
 	UPROPERTY()
 	TArray<FString> ExpandedNodePaths;
 };
diff --git a/Engine/Source/Developer/RigVMDeveloper/Public/RigVMModel/RigVMLink.h b/Engine/Source/Developer/RigVMDeveloper/Public/RigVMModel/RigVMLink.h
--- a/Engine/Source/Developer/RigVMDeveloper/Public/RigVMModel/RigVMLink.h
+++ b/Engine/Source/Developer/RigVMDeveloper/Public/RigVMModel/RigVMLink.h
@@ -49,9 +49,15 @@ public:
 
 	// Returns a string representation of the Link,
 	// for example: "NodeA.Color.R -> NodeB.Translation.X"
+	// note: can be split again using SplitPinPathRepresentation
 	UFUNCTION(BlueprintCallable, Category = RigVMLink)
 	FString GetPinPathRepresentation();
 
+	// Splits a pin path represenation of a link
+	// for example: "NodeA.Color.R -> NodeB.Translation.X"
+	// into its two pin paths
+	static bool SplitPinPathRepresentation(const FString& InString, FString& OutSource, FString& OutTarget);
+
 private:
 
 	void PrepareForCopy();
