RigVM: Fix AST source pin traverse to support function ref nodes

diff --git a/Engine/Source/Developer/RigVMDeveloper/Private/RigVMCompiler/RigVMAST.cpp b/Engine/Source/Developer/RigVMDeveloper/Private/RigVMCompiler/RigVMAST.cpp
--- a/Engine/Source/Developer/RigVMDeveloper/Private/RigVMCompiler/RigVMAST.cpp
+++ b/Engine/Source/Developer/RigVMDeveloper/Private/RigVMCompiler/RigVMAST.cpp
@@ -220,7 +220,10 @@ int32 FRigVMExprAST::GetMinChildIndexWithinParent(const FRigVMExprAST* InParentE
 
 void FRigVMExprAST::AddParent(FRigVMExprAST* InParent)
 {
+	ensure(IsValid());
+	ensure(InParent->IsValid());
 	ensure(InParent != this);
+	
 	if (Parents.Contains(InParent))
 	{
 		return;
@@ -232,6 +235,9 @@ void FRigVMExprAST::AddParent(FRigVMExprAST* InParent)
 
 void FRigVMExprAST::RemoveParent(FRigVMExprAST* InParent)
 {
+	ensure(IsValid());
+	ensure(InParent->IsValid());
+	
 	if (Parents.Remove(InParent) > 0)
 	{
 		InParent->Children.Remove(this);
@@ -240,11 +246,18 @@ void FRigVMExprAST::RemoveParent(FRigVMExprAST* InParent)
 
 void FRigVMExprAST::RemoveChild(FRigVMExprAST* InChild)
 {
+	ensure(IsValid());
+	ensure(InChild->IsValid());
+	
 	InChild->RemoveParent(this);
 }
 
 void FRigVMExprAST::ReplaceParent(FRigVMExprAST* InCurrentParent, FRigVMExprAST* InNewParent)
 {
+	ensure(IsValid());
+	ensure(InCurrentParent->IsValid());
+	ensure(InNewParent->IsValid());
+	
 	for (int32 ParentIndex = 0; ParentIndex < Parents.Num(); ParentIndex++)
 	{
 		if (Parents[ParentIndex] == InCurrentParent)
@@ -258,6 +271,10 @@ void FRigVMExprAST::ReplaceParent(FRigVMExprAST* InCurrentParent, FRigVMExprAST*
 
 void FRigVMExprAST::ReplaceChild(FRigVMExprAST* InCurrentChild, FRigVMExprAST* InNewChild)
 {
+	ensure(IsValid());
+	ensure(InCurrentChild->IsValid());
+	ensure(InNewChild->IsValid());
+
 	for (int32 ChildIndex = 0; ChildIndex < Children.Num(); ChildIndex++)
 	{
 		if (Children[ChildIndex] == InCurrentChild)
@@ -803,6 +820,12 @@ FRigVMParserAST::~FRigVMParserAST()
 	}
 	Expressions.Empty();
 
+	for (FRigVMExprAST* Expression : DeletedExpressions)
+	{
+		delete(Expression);
+	}
+	DeletedExpressions.Empty();
+
 	// root expressions are a subset of the
 	// expressions array, so no cleanup necessary
 	RootExpressions.Empty();
@@ -2642,23 +2665,21 @@ void FRigVMParserAST::RemoveExpressions(TArray<FRigVMExprAST*> InExprs)
 			FRigVMExprAST* Expr = Expressions[ExpressionIndex];
 			RemainingExpressions.Add(Expr);
 
-			for(int32 ChildIndex = 0; ChildIndex < Expr->Children.Num(); ChildIndex++)
+			for(int32 ChildIndex = Expr->Children.Num() - 1; ChildIndex >= 0; ChildIndex--)
 			{
 				FRigVMExprAST* ChildExpr = Expr->Children[ChildIndex];
 				if(bRemoveExpression[ChildExpr->GetIndex()])
 				{
 					Expr->Children.RemoveAt(ChildIndex);
-					break;
 				}
 			}
 
-			for(int32 ParentIndex = 0; ParentIndex < Expr->Parents.Num(); ParentIndex++)
+			for(int32 ParentIndex = Expr->Parents.Num() - 1; ParentIndex >= 0; ParentIndex--)
 			{
 				FRigVMExprAST* ParentExpr = Expr->Parents[ParentIndex];
 				if(bRemoveExpression[ParentExpr->GetIndex()])
 				{
 					Expr->Parents.RemoveAt(ParentIndex);
-					break;
 				}
 			}
 		}
@@ -2682,7 +2703,8 @@ void FRigVMParserAST::RemoveExpressions(TArray<FRigVMExprAST*> InExprs)
 
 	for(int32 ExpressionIndex = ExpressionsToRemove.Num() - 1; ExpressionIndex >= 0; ExpressionIndex--)
 	{
-		delete ExpressionsToRemove[ExpressionIndex];
+		ExpressionsToRemove[ExpressionIndex]->Index = INDEX_NONE;
+		DeletedExpressions.Add(ExpressionsToRemove[ExpressionIndex]);
 	}
 
 	RefreshExprIndices();
@@ -2936,7 +2958,10 @@ void FRigVMParserAST::Inline(URigVMGraph* InGraph, const TArray<FRigVMASTProxy>&
 				{
 					if (URigVMFunctionEntryNode* EntryNode = Cast<URigVMFunctionEntryNode>(ChildPin->GetNode()))
 					{
-						if (URigVMCollapseNode* OuterNode = Cast<URigVMCollapseNode>(EntryNode->GetGraph()->GetOuter()))
+						// rather than relying on the other we are going to query what's in the call stack.
+						// for collapse nodes that's not a different, but for function ref nodes the outer
+						// node is the definition and not the reference node - which sits in the callstack.
+						if (URigVMLibraryNode* OuterNode = (InPinProxy.GetParent().GetSubject<URigVMLibraryNode>()))
 						{
 							if (URigVMPin* OuterPin = OuterNode->FindPin(ChildPin->GetName()))
 							{
diff --git a/Engine/Source/Developer/RigVMDeveloper/Public/RigVMCompiler/RigVMAST.h b/Engine/Source/Developer/RigVMDeveloper/Public/RigVMCompiler/RigVMAST.h
--- a/Engine/Source/Developer/RigVMDeveloper/Public/RigVMCompiler/RigVMAST.h
+++ b/Engine/Source/Developer/RigVMDeveloper/Public/RigVMCompiler/RigVMAST.h
@@ -101,6 +101,9 @@ public:
 	// @return the index of this expression within the parser's storage
 	int32 GetIndex() const { return Index; }
 
+	// returns true if the expressoin is valid
+	FORCEINLINE bool IsValid() const { return GetIndex() != INDEX_NONE; }
+
 	// returns the parent of this expression
 	// @return the parent of this expression
 	const FRigVMExprAST* GetParent() const;
@@ -1405,6 +1408,7 @@ private:
 	TMap<FRigVMASTProxy, int32> NodeExpressionIndex;
 	TArray<FRigVMExprAST*> Expressions;
 	TArray<FRigVMExprAST*> RootExpressions;
+	TArray<FRigVMExprAST*> DeletedExpressions;
 	FRigVMBlockExprAST* ObsoleteBlock;
 
 	TArray<FRigVMASTProxy> NodeProxies;
