RigVM: Change execution order for variable composition

diff --git a/Engine/Source/Developer/RigVMDeveloper/Private/RigVMCompiler/RigVMCompiler.cpp b/Engine/Source/Developer/RigVMDeveloper/Private/RigVMCompiler/RigVMCompiler.cpp
--- a/Engine/Source/Developer/RigVMDeveloper/Private/RigVMCompiler/RigVMCompiler.cpp
+++ b/Engine/Source/Developer/RigVMDeveloper/Private/RigVMCompiler/RigVMCompiler.cpp
@@ -315,23 +315,24 @@ bool URigVMCompiler::Compile(URigVMGraph* InGraph, URigVMController* InControlle
 		OutOperands->Add(Hash, Operand);
 	}
 
-	TSharedPtr<FRigVMParserAST> AST = InAST;
-	if (!AST.IsValid())
+	FRigVMCompilerWorkData WorkData;
+
+	WorkData.AST = InAST;
+	if (!WorkData.AST.IsValid())
 	{
-		AST = MakeShareable(new FRigVMParserAST(InGraph, InController, Settings.ASTSettings, InExternalVariables, UserData));
-		InGraph->RuntimeAST = AST;
+		WorkData.AST = MakeShareable(new FRigVMParserAST(InGraph, InController, Settings.ASTSettings, InExternalVariables, UserData));
+		InGraph->RuntimeAST = WorkData.AST;
 #if UE_BUILD_DEBUG
 		//UE_LOG(LogRigVMDeveloper, Display, TEXT("%s"), *AST->DumpDot());
 #endif
 	}
-	ensure(AST.IsValid());
+	ensure(WorkData.AST.IsValid());
 
-	FRigVMCompilerWorkData WorkData;
 	WorkData.VM = OutVM;
 	WorkData.PinPathToOperand = OutOperands;
 	WorkData.RigVMUserData = UserData[0];
 	WorkData.bSetupMemory = true;
-	WorkData.ProxySources = &AST->SharedOperandPins;
+	WorkData.ProxySources = &WorkData.AST->SharedOperandPins;
 
 	// tbd: do we need this only when we have no pins?
 	//if(!WorkData.WatchedPins.IsEmpty())
@@ -388,7 +389,7 @@ bool URigVMCompiler::Compile(URigVMGraph* InGraph, URigVMController* InControlle
 			auto AddDefaultValueOperand = [&](URigVMPin* Pin)
 			{
 				FRigVMASTProxy PinProxy = FRigVMASTProxy::MakeFromUObject(Pin);
-				FRigVMVarExprAST* TempVarExpr = AST->MakeExpr<FRigVMVarExprAST>(FRigVMExprAST::EType::Literal, PinProxy);
+				FRigVMVarExprAST* TempVarExpr = WorkData.AST->MakeExpr<FRigVMVarExprAST>(FRigVMExprAST::EType::Literal, PinProxy);
 				FRigVMOperand Operand = FindOrAddRegister(TempVarExpr, WorkData, false);
 
 #if UE_RIGVM_UCLASS_BASED_STORAGE_DISABLED
@@ -441,7 +442,7 @@ bool URigVMCompiler::Compile(URigVMGraph* InGraph, URigVMController* InControlle
 #endif
 
 	// define all parameters independent from sorted nodes
-	for (const FRigVMExprAST* Expr : AST->Expressions)
+	for (const FRigVMExprAST* Expr : WorkData.AST->Expressions)
 	{
 		if (Expr->IsA(FRigVMExprAST::EType::Literal))
 		{
@@ -463,7 +464,7 @@ bool URigVMCompiler::Compile(URigVMGraph* InGraph, URigVMController* InControlle
 	}
 
 	WorkData.ExprComplete.Reset();
-	for (FRigVMExprAST* RootExpr : *AST)
+	for (FRigVMExprAST* RootExpr : *WorkData.AST)
 	{
 		TraverseExpression(RootExpr, WorkData);
 	}
@@ -476,7 +477,7 @@ bool URigVMCompiler::Compile(URigVMGraph* InGraph, URigVMController* InControlle
 		if (URigVMParameterNode* ParameterNode = Cast<URigVMParameterNode>(Node))
 		{
 			FRigVMASTProxy ParameterNodeProxy = FRigVMASTProxy::MakeFromUObject(ParameterNode);
-			const FRigVMExprAST* ParameterExpr = AST->GetExprForSubject(ParameterNodeProxy);
+			const FRigVMExprAST* ParameterExpr = WorkData.AST->GetExprForSubject(ParameterNodeProxy);
 			if (!ParameterExpr || ParameterExpr->IsA(FRigVMExprAST::NoOp))
 			{
 				FName Name = ParameterNode->GetParameterName();
@@ -511,7 +512,7 @@ bool URigVMCompiler::Compile(URigVMGraph* InGraph, URigVMController* InControlle
 					}
 					FRigVMASTProxy PinProxy = FRigVMASTProxy::MakeFromUObject(ModelPin);
 					FRigVMVarExprAST TempVarExpr(FRigVMExprAST::EType::Var, PinProxy);
-					TempVarExpr.ParserPtr = AST.Get();
+					TempVarExpr.ParserPtr = WorkData.AST.Get();
 
 					FindOrAddRegister(&TempVarExpr, WorkData, true);
 				}
@@ -546,7 +547,7 @@ bool URigVMCompiler::Compile(URigVMGraph* InGraph, URigVMController* InControlle
 
 	WorkData.bSetupMemory = false;
 	WorkData.ExprComplete.Reset();
-	for (FRigVMExprAST* RootExpr : *AST)
+	for (FRigVMExprAST* RootExpr : *WorkData.AST)
 	{
 		TraverseExpression(RootExpr, WorkData);
 	}
@@ -579,7 +580,7 @@ bool URigVMCompiler::Compile(URigVMGraph* InGraph, URigVMController* InControlle
 #else
 		for(URigVMPin* WatchedPin : WorkData.WatchedPins)
 		{
-			MarkDebugWatch(true, WatchedPin, WorkData.VM, WorkData.PinPathToOperand, AST);
+			MarkDebugWatch(true, WatchedPin, WorkData.VM, WorkData.PinPathToOperand, WorkData.AST);
 		}
 #endif
 	}
@@ -1135,29 +1136,7 @@ void URigVMCompiler::TraverseAssign(const FRigVMAssignExprAST* InExpr, FRigVMCom
 		FRigVMCopyOp CopyOp = WorkData.VM->GetCopyOpForOperands(Source, Target);
 		if(CopyOp.IsValid())
 		{
-			WorkData.VM->GetByteCode().AddCopyOp(CopyOp);
-
-			int32 InstructionIndex = WorkData.VM->GetByteCode().GetNumInstructions() - 1;
-			if (Settings.SetupNodeInstructionIndex)
-			{
-				if (URigVMPin* SourcePin = InExpr->GetSourcePin())
-				{
-					if (URigVMVariableNode* VariableNode = Cast<URigVMVariableNode>(SourcePin->GetNode()))
-					{
-						const FRigVMCallstack Callstack = SourceExpr->GetProxy().GetSibling(VariableNode).GetCallstack();
-						WorkData.VM->GetByteCode().SetSubject(InstructionIndex, Callstack.GetCallPath(), Callstack.GetStack());
-					}
-				}
-
-				if (URigVMPin* TargetPin = InExpr->GetTargetPin())
-				{
-					if (URigVMVariableNode* VariableNode = Cast<URigVMVariableNode>(TargetPin->GetNode()))
-					{
-						const FRigVMCallstack Callstack = TargetExpr->GetProxy().GetSibling(VariableNode).GetCallstack();
-						WorkData.VM->GetByteCode().SetSubject(InstructionIndex, Callstack.GetCallPath(), Callstack.GetStack());
-					}
-				}
-			}
+			AddCopyOperator(CopyOp, InExpr, SourceExpr, TargetExpr, WorkData);
 		}
 	}
 }
@@ -1821,6 +1800,107 @@ void URigVMCompiler::TraverseArray(const FRigVMArrayExprAST* InExpr, FRigVMCompi
 	}
 }
 
+void URigVMCompiler::AddCopyOperator(const FRigVMCopyOp& InOp, const FRigVMAssignExprAST* InAssignExpr,
+	const FRigVMVarExprAST* InSourceExpr, const FRigVMVarExprAST* InTargetExpr,  FRigVMCompilerWorkData& WorkData,
+	bool bDelayCopyOperations)
+{
+	if(bDelayCopyOperations)
+	{
+		// if this is a full literal copy, let's delay it.
+		// to maintain the execution order we want nodes which compose a value
+		// to delay their reset to the default value, which happens prior to
+		// computing dependencies.
+		// so for example an external variable of FVector may need to be reset
+		// to a literal value prior to the rest of the composition, for example
+		// if there's a float link only on the Y component. the execution order
+		// desired is this:
+		//
+		// * Run all dependent branches
+		// * Copy the literal value into the variable
+		// * Copy the parts into the variable (like the Y component).
+		// 
+		// By delaying the copy operator until right before the very first composition
+		// copy operator we ensure the desired execution order
+		if(InOp.Target.GetRegisterOffset() == INDEX_NONE && 
+			InOp.Source.GetMemoryType() == ERigVMMemoryType::Literal &&
+			InOp.Source.GetRegisterOffset() == INDEX_NONE)
+		{
+			if(URigVMPin* Pin = InTargetExpr->GetPin())
+			{
+				if(URigVMPin* RootPin = Pin->GetRootPin())
+				{
+					const FRigVMASTProxy RootPinProxy = InTargetExpr->GetProxy().GetSibling(RootPin);
+
+					// if the root pin has only links on its subpins
+					if(WorkData.AST->GetSourceLinks(RootPinProxy, false).Num() == 0)
+					{
+						if(WorkData.AST->GetSourceLinks(RootPinProxy, true).Num() > 0)
+						{					
+							FRigVMCompilerWorkData::FCopyOpInfo DeferredCopyOp;
+							DeferredCopyOp.Op = InOp;
+							DeferredCopyOp.AssignExpr = InAssignExpr;
+							DeferredCopyOp.SourceExpr = InSourceExpr;
+							DeferredCopyOp.TargetExpr = InTargetExpr;
+				
+							const FRigVMOperand Key(InOp.Target.GetMemoryType(), InOp.Target.GetRegisterIndex());
+							WorkData.DeferredCopyOps.FindOrAdd(Key) = DeferredCopyOp;
+							return;
+						}
+					}
+				}
+			}
+		}
+		
+		bDelayCopyOperations = false;
+	}
+
+	// loop up a potentially delayed copy operation which needs to happen
+	// just prior to this one and inject it as well.
+	if(!bDelayCopyOperations)
+	{
+		const FRigVMOperand DeferredKey(InOp.Target.GetMemoryType(), InOp.Target.GetRegisterIndex());
+		const FRigVMCompilerWorkData::FCopyOpInfo* DeferredCopyOpPtr = WorkData.DeferredCopyOps.Find(DeferredKey);
+		if(DeferredCopyOpPtr != nullptr)
+		{
+			FRigVMCompilerWorkData::FCopyOpInfo CopyOpInfo = *DeferredCopyOpPtr;
+			WorkData.DeferredCopyOps.Remove(DeferredKey);
+			AddCopyOperator(CopyOpInfo, WorkData, false);
+		}
+	}
+	
+	WorkData.VM->GetByteCode().AddCopyOp(InOp);
+
+	int32 InstructionIndex = WorkData.VM->GetByteCode().GetNumInstructions() - 1;
+	if (Settings.SetupNodeInstructionIndex)
+	{
+		if (URigVMPin* SourcePin = InAssignExpr->GetSourcePin())
+		{
+			if (URigVMVariableNode* VariableNode = Cast<URigVMVariableNode>(SourcePin->GetNode()))
+			{
+				const FRigVMCallstack Callstack = InSourceExpr->GetProxy().GetSibling(VariableNode).GetCallstack();
+				WorkData.VM->GetByteCode().SetSubject(InstructionIndex, Callstack.GetCallPath(), Callstack.GetStack());
+			}
+		}
+
+		if (URigVMPin* TargetPin = InAssignExpr->GetTargetPin())
+		{
+			if (URigVMVariableNode* VariableNode = Cast<URigVMVariableNode>(TargetPin->GetNode()))
+			{
+				const FRigVMCallstack Callstack = InTargetExpr->GetProxy().GetSibling(VariableNode).GetCallstack();
+				WorkData.VM->GetByteCode().SetSubject(InstructionIndex, Callstack.GetCallPath(), Callstack.GetStack());
+			}
+		}
+	}
+}
+
+void URigVMCompiler::AddCopyOperator(
+	const FRigVMCompilerWorkData::FCopyOpInfo& CopyOpInfo,
+	FRigVMCompilerWorkData& WorkData,
+	bool bDelayCopyOperations)
+{
+	AddCopyOperator(CopyOpInfo.Op, CopyOpInfo.AssignExpr, CopyOpInfo.SourceExpr, CopyOpInfo.TargetExpr, WorkData, bDelayCopyOperations);
+}
+
 void URigVMCompiler::InitializeLocalVariables(const FRigVMExprAST* InExpr, FRigVMCompilerWorkData& WorkData)
 {
 	// Initialize local variables if we are entering a new graph
diff --git a/Engine/Source/Developer/RigVMDeveloper/Public/RigVMCompiler/RigVMCompiler.h b/Engine/Source/Developer/RigVMDeveloper/Public/RigVMCompiler/RigVMCompiler.h
--- a/Engine/Source/Developer/RigVMDeveloper/Public/RigVMCompiler/RigVMCompiler.h
+++ b/Engine/Source/Developer/RigVMDeveloper/Public/RigVMCompiler/RigVMCompiler.h
@@ -113,6 +113,19 @@ public:
 	int32 FindOrAddPropertyPath(const FRigVMOperand& InOperand, const FString& InHeadCPPType, const FString& InSegmentPath);
 	
 #endif
+
+	TSharedPtr<FRigVMParserAST> AST;
+	
+	struct FCopyOpInfo
+	{
+		FRigVMCopyOp Op;
+		const FRigVMAssignExprAST* AssignExpr;
+		const FRigVMVarExprAST* SourceExpr;
+		const FRigVMVarExprAST* TargetExpr;
+	};
+
+	// operators that have been delayed for injection into the bytecode
+	TMap<FRigVMOperand, FCopyOpInfo> DeferredCopyOps;
 };
 
 UCLASS(BlueprintType)
@@ -178,6 +191,19 @@ private:
 	void TraverseSelect(const FRigVMSelectExprAST* InExpr, FRigVMCompilerWorkData& WorkData);
 	void TraverseArray(const FRigVMArrayExprAST* InExpr, FRigVMCompilerWorkData& WorkData);
 
+	void AddCopyOperator(
+		const FRigVMCopyOp& InOp,
+		const FRigVMAssignExprAST* InAssignExpr,
+		const FRigVMVarExprAST* InSourceExpr,
+		const FRigVMVarExprAST* InTargetExpr,
+		FRigVMCompilerWorkData& WorkData,
+		bool bDelayCopyOperations = true);
+
+	void AddCopyOperator(
+		const FRigVMCompilerWorkData::FCopyOpInfo& CopyOpInfo,
+		FRigVMCompilerWorkData& WorkData,
+		bool bDelayCopyOperations = true);
+
 	void InitializeLocalVariables(const FRigVMExprAST* InExpr, FRigVMCompilerWorkData& WorkData);
 
 	FRigVMOperand FindOrAddRegister(const FRigVMVarExprAST* InVarExpr, FRigVMCompilerWorkData& WorkData, bool bIsDebugValue = false);
@@ -206,6 +232,6 @@ private:
 	{
 		ReportError(FString::Printf(Fmt, Args...));
 	}
-
+	
 	friend class FRigVMCompilerImportErrorContext;
 };
diff --git a/Engine/Source/Runtime/RigVM/Private/RigVMCore/RigVM.cpp b/Engine/Source/Runtime/RigVM/Private/RigVMCore/RigVM.cpp
--- a/Engine/Source/Runtime/RigVM/Private/RigVMCore/RigVM.cpp
+++ b/Engine/Source/Runtime/RigVM/Private/RigVMCore/RigVM.cpp
@@ -3467,6 +3467,14 @@ FString URigVM::GetOperandLabel(const FRigVMOperand& InOperand, TFunction<FStrin
 
 #else
 
+	if(InOperand.GetMemoryType() == ERigVMMemoryType::External)
+	{
+		if(ExternalVariables.IsValidIndex(InOperand.GetRegisterIndex()))
+		{
+			return ExternalVariables[InOperand.GetRegisterIndex()].Name.ToString();
+		}
+	}
+	
 	URigVMMemoryStorage* Memory = GetMemoryByType(InOperand.GetMemoryType());
 	if(Memory == nullptr)
 	{
