Control Rig: Fix for array pin losing bIsDynamicArray Flag on editor open

diff --git a/Engine/Source/Developer/RigVMDeveloper/Private/RigVMCompiler/RigVMAST.cpp b/Engine/Source/Developer/RigVMDeveloper/Private/RigVMCompiler/RigVMAST.cpp
--- a/Engine/Source/Developer/RigVMDeveloper/Private/RigVMCompiler/RigVMAST.cpp
+++ b/Engine/Source/Developer/RigVMDeveloper/Private/RigVMCompiler/RigVMAST.cpp
@@ -1510,8 +1510,7 @@ void FRigVMParserAST::FoldAssignments()
 		
 		// non-input pins on anything but a reroute / array node should be skipped
 		if (TargetPin->GetDirection() != ERigVMPinDirection::Input &&
-			(Cast<URigVMRerouteNode>(TargetPin->GetNode()) == nullptr) &&
-			(Cast<URigVMArrayNode>(TargetPin->GetNode()) == nullptr))
+			(Cast<URigVMRerouteNode>(TargetPin->GetNode()) == nullptr))
 		{
 			continue;
 		}
diff --git a/Engine/Source/Developer/RigVMDeveloper/Private/RigVMModel/RigVMController.cpp b/Engine/Source/Developer/RigVMDeveloper/Private/RigVMModel/RigVMController.cpp
--- a/Engine/Source/Developer/RigVMDeveloper/Private/RigVMModel/RigVMController.cpp
+++ b/Engine/Source/Developer/RigVMDeveloper/Private/RigVMModel/RigVMController.cpp
@@ -11265,12 +11265,16 @@ void URigVMController::ConfigurePinFromProperty(FProperty* InProperty, URigVMPin
 
 void URigVMController::ConfigurePinFromPin(URigVMPin* InOutPin, URigVMPin* InPin)
 {
+	// it is important we copy things that define the identity of the pin
+	// things that defines the state of the pin is copied during GetPinState()
+	// though addmittedly these two functions have overlaps currently
 	InOutPin->bIsConstant = InPin->bIsConstant;
 	InOutPin->Direction = InPin->Direction;
 	InOutPin->CPPType = InPin->CPPType;
 	InOutPin->CPPTypeObjectPath = InPin->CPPTypeObjectPath;
 	InOutPin->CPPTypeObject = InPin->CPPTypeObject;
 	InOutPin->DefaultValue = InPin->DefaultValue;
+	InOutPin->bIsDynamicArray = InPin->bIsDynamicArray;
 }
 
 bool URigVMController::ShouldStructBeUnfolded(const UStruct* Struct)
diff --git a/Engine/Source/Developer/RigVMDeveloper/Public/RigVMModel/RigVMPin.h b/Engine/Source/Developer/RigVMDeveloper/Public/RigVMModel/RigVMPin.h
--- a/Engine/Source/Developer/RigVMDeveloper/Public/RigVMModel/RigVMPin.h
+++ b/Engine/Source/Developer/RigVMDeveloper/Public/RigVMModel/RigVMPin.h
@@ -414,6 +414,10 @@ private:
 
 #endif
 
+	// if new members are added to the pin in the future 
+	// it is important to search for all existing usages of all members
+	// to make sure things are copied/initialized properly
+	
 	UPROPERTY()
 	ERigVMPinDirection Direction;
 
