RigVM: Catch model / variable out of sync issues

diff --git a/Engine/Source/Developer/RigVMDeveloper/Private/RigVMModel/RigVMController.cpp b/Engine/Source/Developer/RigVMDeveloper/Private/RigVMModel/RigVMController.cpp
--- a/Engine/Source/Developer/RigVMDeveloper/Private/RigVMModel/RigVMController.cpp
+++ b/Engine/Source/Developer/RigVMDeveloper/Private/RigVMModel/RigVMController.cpp
@@ -11907,6 +11907,42 @@ void URigVMController::RepopulatePinsOnNode(URigVMNode* InNode, bool bFollowCore
 		}
 		check(ValuePin);
 		EnsurePinValidity(ValuePin, false);
+
+		if(VariableNode)
+		{
+			// this includes local variables for validation
+			const TArray<FRigVMExternalVariable> ExternalVariables = GetAllVariables(false);
+			const FRigVMGraphVariableDescription VariableDescription = VariableNode->GetVariableDescription();
+			const FRigVMExternalVariable CurrentExternalVariable = VariableDescription.ToExternalVariable();
+
+			for(const FRigVMExternalVariable& ExternalVariable : ExternalVariables)
+			{
+				if(ExternalVariable.Name == CurrentExternalVariable.Name)
+				{
+					if(ExternalVariable.TypeName != CurrentExternalVariable.TypeName ||
+						ExternalVariable.TypeObject != CurrentExternalVariable.TypeObject ||
+						ExternalVariable.bIsArray != CurrentExternalVariable.bIsArray)
+					{
+						FString CPPType;
+						UObject* CPPTypeObject;
+						
+						if(RigVMTypeUtils::CPPTypeFromExternalVariable(ExternalVariable, CPPType, &CPPTypeObject))
+						{
+							RefreshVariableNode(VariableNode->GetFName(), ExternalVariable.Name, CPPType, ExternalVariable.TypeObject, false, bSetupOrphanedPins);
+						}
+						else
+						{
+							ReportErrorf(
+								TEXT("Control Rig '%s', Type of Variable '%s' cannot be resolved."),
+								*InNode->GetOutermost()->GetPathName(),
+								*ExternalVariable.Name.ToString()
+							);
+						}
+					}
+					break;
+				}
+			}
+		}
 		
 		RemovePinsDuringRepopulate(InNode, ValuePin->SubPins, bNotify, bSetupOrphanedPins);
 
