Control Rig: Fixed single apostrophes in comments causes Python Script output not to work

diff --git a/Engine/Source/Developer/RigVMDeveloper/Private/RigVMModel/RigVMController.cpp b/Engine/Source/Developer/RigVMDeveloper/Private/RigVMModel/RigVMController.cpp
--- a/Engine/Source/Developer/RigVMDeveloper/Private/RigVMModel/RigVMController.cpp
+++ b/Engine/Source/Developer/RigVMDeveloper/Private/RigVMModel/RigVMController.cpp
@@ -619,7 +619,7 @@ TArray<FString> URigVMController::GetAddNodePythonCommands(URigVMNode* Node) con
 		// add_comment_node(comment_text, position=[0.0, 0.0], size=[400.0, 300.0], color=[0.0, 0.0, 0.0, 0.0], node_name='', undo=True)
 		Commands.Add(FString::Printf(TEXT("blueprint.get_controller_by_name('%s').add_comment_node('%s', %s, %s, %s, '%s')"),
 					*GraphName,
-					*CommentNode->GetCommentText(),
+					*CommentNode->GetCommentText().ReplaceCharWithEscapedChar(),
 					*RigVMPythonUtils::Vector2DToPythonString(CommentNode->GetPosition()),
 					*RigVMPythonUtils::Vector2DToPythonString(CommentNode->GetSize()),
 					*RigVMPythonUtils::LinearColorToPythonString(CommentNode->GetNodeColor()),
