Control Rig: Fix Mannequin_FootPlant_CtrlRig to prevent crash in LWC

diff --git a/Engine/Source/Developer/RigVMDeveloper/Private/RigVMCompiler/RigVMAST.cpp b/Engine/Source/Developer/RigVMDeveloper/Private/RigVMCompiler/RigVMAST.cpp
--- a/Engine/Source/Developer/RigVMDeveloper/Private/RigVMCompiler/RigVMAST.cpp
+++ b/Engine/Source/Developer/RigVMDeveloper/Private/RigVMCompiler/RigVMAST.cpp
@@ -1548,6 +1548,31 @@ void FRigVMParserAST::FoldAssignments()
 			continue;
 		}
 
+		// To prevent bad assignments in LWC for VMs compiled in non LWC, we do not allow folding of assignments
+		// to/from external variables of type float
+		{
+			if (SourcePin->GetCPPType() == TEXT("float") || SourcePin->GetCPPType() == TEXT("TArray<float>"))
+			{
+				if (URigVMVariableNode* SourceVariableNode = Cast<URigVMVariableNode>(SourcePin->GetNode()))
+				{
+					if (!SourceVariableNode->IsInputArgument() && !SourceVariableNode->IsLocalVariable())
+					{
+						continue;
+					}
+				}
+			}
+			if (TargetPin->GetCPPType() == TEXT("float") || TargetPin->GetCPPType() == TEXT("TArray<float>"))
+			{
+				if (URigVMVariableNode* TargetVariableNode = Cast<URigVMVariableNode>(TargetPin->GetNode()))
+				{
+					if (!TargetVariableNode->IsInputArgument() && !TargetVariableNode->IsLocalVariable())
+					{
+						continue;
+					}
+				}
+			}
+		}
+
 		FRigVMExprAST* Child = AssignExpr->Children[0];
 		AssignExpr->RemoveParent(Parent);
 		Child->RemoveParent(AssignExpr);
