Control Rig: Refresh pasted nodes to update any old configuration of the node (color, pins, name, etc)

diff --git a/Engine/Source/Developer/RigVMDeveloper/Private/RigVMModel/RigVMController.cpp b/Engine/Source/Developer/RigVMDeveloper/Private/RigVMModel/RigVMController.cpp
--- a/Engine/Source/Developer/RigVMDeveloper/Private/RigVMModel/RigVMController.cpp
+++ b/Engine/Source/Developer/RigVMDeveloper/Private/RigVMModel/RigVMController.cpp
@@ -2838,6 +2838,15 @@ TArray<FName> URigVMController::ImportNodesFromText(const FString& InText, bool
 		TArray<URigVMNode*> SubNodes;
 		SubNodes.Add(CreatedNode);
 
+		// Refresh the unit nodes to account for changes in node color, pin additions, pin order, etc
+		for(int32 SubNodeIndex=0; SubNodeIndex < SubNodes.Num(); SubNodeIndex++)
+		{
+			if (URigVMUnitNode* UnitNode = Cast<URigVMUnitNode>(SubNodes[SubNodeIndex]))
+			{
+				RepopulatePinsOnNode(UnitNode);
+			}
+		}
+
 		for(int32 SubNodeIndex=0; SubNodeIndex < SubNodes.Num(); SubNodeIndex++)
 		{
 			if(URigVMCollapseNode* CollapseNode = Cast<URigVMCollapseNode>(SubNodes[SubNodeIndex]))
