Control Rig: Fix reroutes redo connection failing

diff --git a/Engine/Source/Developer/RigVMDeveloper/Private/RigVMModel/RigVMController.cpp b/Engine/Source/Developer/RigVMDeveloper/Private/RigVMModel/RigVMController.cpp
--- a/Engine/Source/Developer/RigVMDeveloper/Private/RigVMModel/RigVMController.cpp
+++ b/Engine/Source/Developer/RigVMDeveloper/Private/RigVMModel/RigVMController.cpp
@@ -2586,7 +2586,6 @@ bool URigVMController::Undo()
 		return false;
 	}
 
-	TGuardValue<bool> GuardCompactness(bIgnoreRerouteCompactnessChanges, true);
 	return ActionStack->Undo(this);
 }
 
@@ -2597,7 +2596,6 @@ bool URigVMController::Redo()
 		return false;
 	}
 
-	TGuardValue<bool> GuardCompactness(bIgnoreRerouteCompactnessChanges, true);
 	return ActionStack->Redo(this);
 }
 
@@ -7718,8 +7716,11 @@ bool URigVMController::AddLink(URigVMPin* OutputPin, URigVMPin* InputPin, bool b
 	}
 	Notify(ERigVMGraphNotifType::LinkAdded, Link);
 
-	UpdateRerouteNodeAfterChangingLinks(OutputPin, bSetupUndoRedo);
-	UpdateRerouteNodeAfterChangingLinks(InputPin, bSetupUndoRedo);
+	if (bSetupUndoRedo)
+	{
+		UpdateRerouteNodeAfterChangingLinks(OutputPin, bSetupUndoRedo);
+		UpdateRerouteNodeAfterChangingLinks(InputPin, bSetupUndoRedo);
+	}
 
 	//TArray<URigVMNode*> NodesVisited;
 	//PotentiallyResolvePrototypeNode(Cast<URigVMPrototypeNode>(InputPin->GetNode()), bSetupUndoRedo, NodesVisited);
@@ -7878,8 +7879,11 @@ bool URigVMController::BreakLink(URigVMPin* OutputPin, URigVMPin* InputPin, bool
 
 			DestroyObject(Link);
 
-			UpdateRerouteNodeAfterChangingLinks(OutputPin, bSetupUndoRedo);
-			UpdateRerouteNodeAfterChangingLinks(InputPin, bSetupUndoRedo);
+			if (bSetupUndoRedo)
+			{
+				UpdateRerouteNodeAfterChangingLinks(OutputPin, bSetupUndoRedo);
+				UpdateRerouteNodeAfterChangingLinks(InputPin, bSetupUndoRedo);
+			}
 
 			if (bSetupUndoRedo)
 			{
