Control Rig: Move function reference table to the dependent blueprint asset registry searchable data

diff --git a/Engine/Plugins/Animation/ControlRig/Source/ControlRigDeveloper/Private/ControlRigBlueprint.cpp b/Engine/Plugins/Animation/ControlRig/Source/ControlRigDeveloper/Private/ControlRigBlueprint.cpp
--- a/Engine/Plugins/Animation/ControlRig/Source/ControlRigDeveloper/Private/ControlRigBlueprint.cpp
+++ b/Engine/Plugins/Animation/ControlRig/Source/ControlRigDeveloper/Private/ControlRigBlueprint.cpp
@@ -368,6 +368,8 @@ void UControlRigBlueprint::PreSave(FObjectPreSaveContext ObjectSaveContext)
 			}
 		}
 	}
+
+	FunctionReferenceNodeData = GetReferenceNodeData();
 }
 
 void UControlRigBlueprint::PostLoad()
@@ -518,7 +520,10 @@ void UControlRigBlueprint::PostLoad()
 							{
 								if(DependencyBlueprint != this)
 								{
-									DependencyBlueprint->GetLocalFunctionLibrary()->UpdateReferencesForReferenceNode(FunctionReferenceNode);
+									if(URigVMBuildData* BuildData = URigVMController::GetBuildData())
+									{
+										BuildData->UpdateReferencesForFunctionReferenceNode(FunctionReferenceNode);
+									}
 								}
 							}
 						}
@@ -532,11 +537,6 @@ void UControlRigBlueprint::PostLoad()
 #endif
 	}
 
-	if (FunctionLibrary)
-	{
-		FunctionLibrary->ClearInvalidReferences();
-	}
-	
 	// upgrade the gizmo libraries to shape libraries
 	if(GizmoLibrary_DEPRECATED.IsValid() || GetLinkerCustomVersion(FControlRigObjectVersion::GUID) < FControlRigObjectVersion::RenameGizmoToShape)
 	{
@@ -624,6 +624,60 @@ void UControlRigBlueprint::HandlePackageDone(TConstArrayView<UPackage*> InPackag
 
 		ShapeLibrariesToLoadOnPackageLoaded.Reset();
 	}
+
+	if(URigVMBuildData* BuildData = URigVMController::GetBuildData())
+	{
+		if(FunctionLibrary != nullptr)
+		{
+			// for backwards compatibility load the function references from the
+			// model's storage over to the centralized build data
+			if(!FunctionLibrary->FunctionReferences_DEPRECATED.IsEmpty())
+			{
+				// let's also update the asset data of the dependents
+				const FAssetRegistryModule& AssetRegistryModule = FModuleManager::LoadModuleChecked<FAssetRegistryModule>("AssetRegistry");
+				
+				for(const TTuple< TObjectPtr<URigVMLibraryNode>, FRigVMFunctionReferenceArray >& Pair :
+					FunctionLibrary->FunctionReferences_DEPRECATED)
+				{
+					TSoftObjectPtr<URigVMLibraryNode> FunctionKey(Pair.Key);
+						
+					for(int32 ReferenceIndex = 0; ReferenceIndex < Pair.Value.Num(); ReferenceIndex++)
+					{
+						// update the build data
+						BuildData->RegisterFunctionReference(FunctionKey, Pair.Value[ReferenceIndex]);
+
+						// find all control rigs matching the reference node
+						FAssetData AssetData = AssetRegistryModule.Get().GetAssetByObjectPath(
+							Pair.Value[ReferenceIndex].ToSoftObjectPath().GetAssetPathName());
+
+						// if the asset has never been loaded - make sure to load it once and mark as dirty
+						if(AssetData.IsValid() && !AssetData.IsAssetLoaded())
+						{
+							if(UControlRigBlueprint* Dependent = Cast<UControlRigBlueprint>(AssetData.GetAsset()))
+							{
+								if(Dependent != this)
+								{
+									Dependent->MarkPackageDirty();
+								}
+							}
+						}
+					}
+				}
+				
+				FunctionLibrary->FunctionReferences_DEPRECATED.Reset();
+				MarkPackageDirty();
+			}
+		}
+
+		// update the build data from the current function references
+		const TArray<FRigVMReferenceNodeData> ReferenceNodeDatas = GetReferenceNodeData();
+		for(const FRigVMReferenceNodeData& ReferenceNodeData : ReferenceNodeDatas)
+		{
+			BuildData->RegisterFunctionReference(ReferenceNodeData);
+		}
+
+		BuildData->ClearInvalidReferences();
+	}
 	
 	PropagateHierarchyFromBPToInstances();
 	RecompileVM();
@@ -1149,6 +1203,24 @@ void UControlRigBlueprint::RefreshControlRigBreakpoints()
 	}
 }
 
+TArray<FRigVMReferenceNodeData> UControlRigBlueprint::GetReferenceNodeData() const
+{
+	TArray<FRigVMReferenceNodeData> Data;
+	
+	TArray<URigVMGraph*> AllModels = GetAllModels();
+	for (URigVMGraph* ModelToVisit : AllModels)
+	{
+		for(URigVMNode* Node : ModelToVisit->GetNodes())
+		{
+			if(URigVMFunctionReferenceNode* ReferenceNode = Cast<URigVMFunctionReferenceNode>(Node))
+			{
+				Data.Add(FRigVMReferenceNodeData(ReferenceNode));
+			}
+		}
+	}
+	return Data;
+}
+
 #endif
 
 void UControlRigBlueprint::RequestControlRigInit()
@@ -1500,7 +1572,7 @@ URigVMController* UControlRigBlueprint::GetOrCreateController(URigVMGraph* InGra
 		
 		return FName();
 	});
-	
+
 #endif
 
 	Controller->RemoveStaleNodes();
@@ -3488,6 +3560,8 @@ void UControlRigBlueprint::PatchFunctionReferencesOnLoad()
 			}
 		}
 	}
+
+	FunctionReferenceNodeData = GetReferenceNodeData();
 }
 
 #endif
diff --git a/Engine/Plugins/Animation/ControlRig/Source/ControlRigDeveloper/Public/ControlRigBlueprint.h b/Engine/Plugins/Animation/ControlRig/Source/ControlRigDeveloper/Public/ControlRigBlueprint.h
--- a/Engine/Plugins/Animation/ControlRig/Source/ControlRigDeveloper/Public/ControlRigBlueprint.h
+++ b/Engine/Plugins/Animation/ControlRig/Source/ControlRigDeveloper/Public/ControlRigBlueprint.h
@@ -357,6 +357,10 @@ protected:
 	UPROPERTY(AssetRegistrySearchable)
 	TArray<FControlRigPublicFunctionData> PublicFunctions;
 
+	/** Asset searchable information function references in this rig */
+	UPROPERTY(AssetRegistrySearchable)
+	TArray<FRigVMReferenceNodeData> FunctionReferenceNodeData;
+
 	UPROPERTY(BlueprintReadOnly, transient, Category = "VM")
 	TMap<TObjectPtr<URigVMGraph>, TObjectPtr<URigVMController>> Controllers;
 
@@ -689,6 +693,8 @@ public:
 	/** Shape libraries to load during package load completed */ 
 	TArray<FString> ShapeLibrariesToLoadOnPackageLoaded;
 
+	TArray<FRigVMReferenceNodeData> GetReferenceNodeData() const;
+
 #endif
 
 private:
diff --git a/Engine/Plugins/Animation/ControlRig/Source/ControlRigEditor/Private/ControlRigEditorModule.cpp b/Engine/Plugins/Animation/ControlRig/Source/ControlRigEditor/Private/ControlRigEditorModule.cpp
--- a/Engine/Plugins/Animation/ControlRig/Source/ControlRigEditor/Private/ControlRigEditorModule.cpp
+++ b/Engine/Plugins/Animation/ControlRig/Source/ControlRigEditor/Private/ControlRigEditorModule.cpp
@@ -250,6 +250,44 @@ void FControlRigEditorModule::StartupModule()
 	//UThumbnailManager::Get().RegisterCustomRenderer(UControlRigPoseAsset::StaticClass(), UControlRigPoseThumbnailRenderer::StaticClass());
 
 	bFilterAssetBySkeleton = true;
+
+	FFunctionGraphTask::CreateAndDispatchWhenReady([]()
+	{
+		if(URigVMBuildData* BuildData = URigVMController::GetBuildData())
+		{
+			const FArrayProperty* ReferenceNodeDataProperty =
+				CastField<FArrayProperty>(UControlRigBlueprint::StaticClass()->FindPropertyByName(TEXT("FunctionReferenceNodeData")));
+			if(ReferenceNodeDataProperty)
+			{
+				const FAssetRegistryModule& AssetRegistryModule = FModuleManager::LoadModuleChecked<FAssetRegistryModule>("AssetRegistry");
+
+				// find all control rigs in the project
+				TArray<FAssetData> ControlRigAssetDatas;
+				FARFilter ControlRigAssetFilter;
+				ControlRigAssetFilter.ClassNames.Add(UControlRigBlueprint::StaticClass()->GetFName());
+				AssetRegistryModule.Get().GetAssets(ControlRigAssetFilter, ControlRigAssetDatas);
+
+				// loop over all control rigs in the project
+				for(const FAssetData& ControlRigAssetData : ControlRigAssetDatas)
+				{
+					const FString ReferenceNodeDataString =
+						ControlRigAssetData.GetTagValueRef<FString>(ReferenceNodeDataProperty->GetFName());
+					if(ReferenceNodeDataString.IsEmpty())
+					{
+						continue;
+					}
+
+					TArray<FRigVMReferenceNodeData> ReferenceNodeDatas;
+					ReferenceNodeDataProperty->ImportText(*ReferenceNodeDataString, &ReferenceNodeDatas, EPropertyPortFlags::PPF_None, nullptr);
+
+					for(const FRigVMReferenceNodeData& ReferenceNodeData : ReferenceNodeDatas)
+					{
+						BuildData->RegisterFunctionReference(ReferenceNodeData);
+					}
+				}
+			}
+		}
+	}, TStatId(), NULL, ENamedThreads::GameThread);
 }
 
 void FControlRigEditorModule::ShutdownModule()
diff --git a/Engine/Source/Developer/RigVMDeveloper/Private/RigVMModel/RigVMBuildData.cpp b/Engine/Source/Developer/RigVMDeveloper/Private/RigVMModel/RigVMBuildData.cpp
new file mode 100644
--- /dev/null
+++ b/Engine/Source/Developer/RigVMDeveloper/Private/RigVMModel/RigVMBuildData.cpp
@@ -0,0 +1,266 @@
+// Copyright Epic Games, Inc. All Rights Reserved.
+
+#include "RigVMModel/RigVMBuildData.h"
+
+FRigVMReferenceNodeData::FRigVMReferenceNodeData(URigVMFunctionReferenceNode* InReferenceNode)
+{
+	check(InReferenceNode);
+	ReferenceNodePtr = TSoftObjectPtr<URigVMFunctionReferenceNode>(InReferenceNode);
+	ReferenceNodePath = ReferenceNodePtr.ToString();
+	LibraryNodePtr = TSoftObjectPtr<URigVMLibraryNode>(InReferenceNode->GetReferencedNode());
+	ReferencedFunctionPath = LibraryNodePtr.ToString();
+}
+
+TSoftObjectPtr<URigVMFunctionReferenceNode> FRigVMReferenceNodeData::GetReferenceNodeObjectPath()
+{
+	if(ReferenceNodePtr.IsNull())
+	{
+		ReferenceNodePtr = TSoftObjectPtr<URigVMFunctionReferenceNode>(ReferenceNodePath);
+	}
+	return ReferenceNodePtr;
+}
+
+TSoftObjectPtr<URigVMLibraryNode> FRigVMReferenceNodeData::GetReferencedFunctionObjectPath()
+{
+	if(LibraryNodePtr.IsNull())
+	{
+		LibraryNodePtr = TSoftObjectPtr<URigVMFunctionReferenceNode>(ReferencedFunctionPath);
+	}
+	return LibraryNodePtr;
+}
+
+URigVMFunctionReferenceNode* FRigVMReferenceNodeData::GetReferenceNode()
+{
+	if(ReferenceNodePtr.IsNull())
+	{
+		ReferenceNodePtr = TSoftObjectPtr<URigVMFunctionReferenceNode>(ReferenceNodePath);
+	}
+	if(!ReferenceNodePtr.IsValid())
+	{
+		ReferenceNodePtr.LoadSynchronous();
+	}
+	if(ReferenceNodePtr.IsValid())
+	{
+		return ReferenceNodePtr.Get();
+	}
+	return nullptr;
+}
+
+URigVMLibraryNode* FRigVMReferenceNodeData::GetReferencedFunction()
+{
+	if(LibraryNodePtr.IsNull())
+	{
+		LibraryNodePtr = TSoftObjectPtr<URigVMLibraryNode>(ReferencedFunctionPath);
+	}
+	if(!LibraryNodePtr.IsValid())
+	{
+		LibraryNodePtr.LoadSynchronous();
+	}
+	if(LibraryNodePtr.IsValid())
+	{
+		return LibraryNodePtr.Get();
+	}
+	return nullptr;
+}
+
+URigVMBuildData::URigVMBuildData()
+: UObject()
+{
+}
+
+const FRigVMFunctionReferenceArray* URigVMBuildData::FindFunctionReferences(const URigVMLibraryNode* InFunction) const
+{
+	check(InFunction);
+
+	const TSoftObjectPtr<URigVMLibraryNode> Key(InFunction);
+	return FunctionReferences.Find(Key);
+}
+
+void URigVMBuildData::ForEachFunctionReference(const URigVMLibraryNode* InFunction,
+                                             TFunction<void(URigVMFunctionReferenceNode*)> PerReferenceFunction) const
+{
+	check(InFunction);
+	
+	if (const FRigVMFunctionReferenceArray* ReferencesEntry = FindFunctionReferences(InFunction))
+	{
+		for (int32 ReferenceIndex = 0; ReferenceIndex < ReferencesEntry->Num(); ReferenceIndex++)
+		{
+			const TSoftObjectPtr<URigVMFunctionReferenceNode>& Reference = ReferencesEntry->operator [](ReferenceIndex);
+			if (!Reference.IsValid())
+			{
+				Reference.LoadSynchronous();
+			}
+			if (Reference.IsValid())
+			{
+				PerReferenceFunction(Reference.Get());
+			}
+		}
+	}
+}
+
+void URigVMBuildData::ForEachFunctionReferenceSoftPtr(const URigVMLibraryNode* InFunction,
+	TFunction<void(TSoftObjectPtr<URigVMFunctionReferenceNode>)> PerReferenceFunction) const
+{
+	check(InFunction);
+
+	if (const FRigVMFunctionReferenceArray* ReferencesEntry = FindFunctionReferences(InFunction))
+	{
+		for (int32 ReferenceIndex = 0; ReferenceIndex < ReferencesEntry->Num(); ReferenceIndex++)
+		{
+			const TSoftObjectPtr<URigVMFunctionReferenceNode>& Reference = ReferencesEntry->operator [](ReferenceIndex);
+			PerReferenceFunction(Reference);
+		}
+	}
+}
+
+void URigVMBuildData::UpdateReferencesForFunctionReferenceNode(URigVMFunctionReferenceNode* InReferenceNode)
+{
+	check(InReferenceNode);
+
+	if(InReferenceNode->GetOutermost() == GetTransientPackage())
+	{
+		return;
+	}
+	
+	if(const URigVMLibraryNode* Function = InReferenceNode->GetReferencedNode())
+	{
+		const TSoftObjectPtr<URigVMLibraryNode> Key(Function);
+		FRigVMFunctionReferenceArray* ReferencesEntry = FunctionReferences.Find(Key);
+		if (ReferencesEntry == nullptr)
+		{
+			Modify();
+			FunctionReferences.Add(Key);
+			ReferencesEntry = FunctionReferences.Find(Key);
+		}
+
+		const FString ReferenceNodePathName = InReferenceNode->GetPathName();
+		for (int32 ReferenceIndex = 0; ReferenceIndex < ReferencesEntry->Num(); ReferenceIndex++)
+		{
+			const TSoftObjectPtr<URigVMFunctionReferenceNode>& Reference = ReferencesEntry->operator [](ReferenceIndex);
+			if(Reference.ToString() == ReferenceNodePathName)
+			{
+				return;
+			}
+		}
+
+		Modify();
+		ReferencesEntry->FunctionReferences.Add(InReferenceNode);
+		MarkPackageDirty();
+	}
+}
+
+void URigVMBuildData::RegisterFunctionReference(URigVMLibraryNode* InFunction,
+	URigVMFunctionReferenceNode* InReference)
+{
+	if(InFunction == nullptr || InReference == nullptr)
+	{
+		return;
+	}
+
+	const TSoftObjectPtr<URigVMLibraryNode> FunctionKey(InFunction);
+	const TSoftObjectPtr<URigVMFunctionReferenceNode> ReferenceKey(InReference);
+
+	RegisterFunctionReference(FunctionKey, ReferenceKey);
+}
+
+void URigVMBuildData::RegisterFunctionReference(TSoftObjectPtr<URigVMLibraryNode> InFunction,
+	TSoftObjectPtr<URigVMFunctionReferenceNode> InReference)
+{
+	if(InFunction.IsNull() || InReference.IsNull())
+	{
+		return;
+	}
+
+	if(FRigVMFunctionReferenceArray* ReferenceEntry = FunctionReferences.Find(InFunction))
+	{
+		if(ReferenceEntry->FunctionReferences.Contains(InReference))
+		{
+			return;
+		}
+
+		Modify();
+		ReferenceEntry->FunctionReferences.Add(InReference);
+	}
+	else
+	{
+		Modify();
+		FRigVMFunctionReferenceArray NewReferenceEntry;
+		NewReferenceEntry.FunctionReferences.Add(InReference);
+		FunctionReferences.Add(InFunction, NewReferenceEntry);
+	}
+	
+	MarkPackageDirty();
+}
+
+void URigVMBuildData::RegisterFunctionReference(FRigVMReferenceNodeData InReferenceNodeData)
+{
+	RegisterFunctionReference(InReferenceNodeData.GetReferencedFunctionObjectPath(), InReferenceNodeData.GetReferenceNodeObjectPath());
+}
+
+void URigVMBuildData::UnregisterFunctionReference(URigVMLibraryNode* InFunction,
+                                                  URigVMFunctionReferenceNode* InReference)
+{
+	if(InFunction == nullptr || InReference == nullptr)
+	{
+		return;
+	}
+
+	const TSoftObjectPtr<URigVMLibraryNode> FunctionKey(InFunction);
+	const TSoftObjectPtr<URigVMFunctionReferenceNode> ReferenceKey(InReference);
+
+	return UnregisterFunctionReference(FunctionKey, ReferenceKey);
+}
+
+void URigVMBuildData::UnregisterFunctionReference(TSoftObjectPtr<URigVMLibraryNode> InFunction,
+	TSoftObjectPtr<URigVMFunctionReferenceNode> InReference)
+{
+	if(InFunction.IsNull() || InReference.IsNull())
+	{
+		return;
+	}
+
+	if(FRigVMFunctionReferenceArray* ReferenceEntry = FunctionReferences.Find(InFunction))
+	{
+		if(!ReferenceEntry->FunctionReferences.Contains(InReference))
+		{
+			return;
+		}
+
+		Modify();
+		ReferenceEntry->FunctionReferences.Remove(InReference);
+		MarkPackageDirty();
+	}
+}
+
+void URigVMBuildData::ClearInvalidReferences()
+{
+	Modify();
+	
+	// check each function's each reference
+	int32 NumRemoved = 0;
+	for (TTuple<TSoftObjectPtr<URigVMLibraryNode>, FRigVMFunctionReferenceArray>& FunctionReferenceInfo : FunctionReferences)
+	{
+		FRigVMFunctionReferenceArray* ReferencesEntry = &FunctionReferenceInfo.Value;
+
+		static FString sTransientPackagePrefix;
+		if(sTransientPackagePrefix.IsEmpty())
+		{
+			sTransientPackagePrefix = GetTransientPackage()->GetPathName();
+		}
+		static const FString sTempPrefix = TEXT("/Temp/");
+
+		NumRemoved += ReferencesEntry->FunctionReferences.RemoveAll([](TSoftObjectPtr<URigVMFunctionReferenceNode> Referencer)
+		{
+			// ignore keys / references within the transient package
+			const FString ReferencerString = Referencer.ToString();
+			return ReferencerString.StartsWith(sTransientPackagePrefix) || ReferencerString.StartsWith(sTempPrefix);
+		});
+	}
+
+	if (NumRemoved > 0)
+	{
+		MarkPackageDirty();
+	}
+}
+
+
diff --git a/Engine/Source/Developer/RigVMDeveloper/Private/RigVMModel/RigVMController.cpp b/Engine/Source/Developer/RigVMDeveloper/Private/RigVMModel/RigVMController.cpp
--- a/Engine/Source/Developer/RigVMDeveloper/Private/RigVMModel/RigVMController.cpp
+++ b/Engine/Source/Developer/RigVMDeveloper/Private/RigVMModel/RigVMController.cpp
@@ -2875,12 +2875,11 @@ TArray<FName> URigVMController::ImportNodesFromText(const FString& InText, bool
 
 			if (URigVMFunctionReferenceNode* FunctionRefNode = Cast<URigVMFunctionReferenceNode>(SubNode))
 			{
-				if (URigVMFunctionLibrary* FunctionLibrary = FunctionRefNode->GetLibrary())
+				if (URigVMLibraryNode* FunctionDefinition = FunctionRefNode->GetReferencedNode())
 				{
-					if (URigVMLibraryNode* FunctionDefinition = FunctionRefNode->GetReferencedNode())
+					if(URigVMBuildData* BuildData = GetBuildData())
 					{
-						FunctionLibrary->FunctionReferences.FindOrAdd(FunctionDefinition).FunctionReferences.Add(FunctionRefNode);
-						FunctionLibrary->MarkPackageDirty();
+						BuildData->RegisterFunctionReference(FunctionDefinition, FunctionRefNode);
 					}
 				}
 			}
@@ -4961,21 +4960,16 @@ void URigVMController::SetReferencedFunction(URigVMFunctionReferenceNode* InFunc
 		return;
 	}
 	
-	URigVMLibraryNode* ReferencedNode = InFunctionRefNode->GetReferencedNode();
-	
-	if(URigVMFunctionLibrary* OtherLibrary = Cast<URigVMFunctionLibrary>(ReferencedNode->GetOuter()))
+	URigVMLibraryNode* OldReferencedNode = InFunctionRefNode->GetReferencedNode();
+	if(OldReferencedNode != InNewReferencedNode)
 	{
-		if(FRigVMFunctionReferenceArray* OtherReferences = OtherLibrary->FunctionReferences.Find(ReferencedNode))
+		if(URigVMBuildData* BuildData = GetBuildData())
 		{
-			OtherReferences->FunctionReferences.Remove(InFunctionRefNode);
+			BuildData->UnregisterFunctionReference(OldReferencedNode, InFunctionRefNode);
+			BuildData->RegisterFunctionReference(InNewReferencedNode, InFunctionRefNode);
 		}
 	}
 
-	URigVMFunctionLibrary* NewLibrary = Cast<URigVMFunctionLibrary>(InNewReferencedNode->GetOuter());
-	FRigVMFunctionReferenceArray& NewReferences = NewLibrary->FunctionReferences.FindOrAdd(InNewReferencedNode);
-	NewReferences.FunctionReferences.Add(InFunctionRefNode);
-	NewLibrary->MarkPackageDirty();
-
 	InFunctionRefNode->SetReferencedNode(InNewReferencedNode);
 	
 	FRigVMControllerGraphGuard GraphGuard(this, InFunctionRefNode->GetGraph(), false);
@@ -5137,57 +5131,45 @@ bool URigVMController::RemoveNode(URigVMNode* InNode, bool bSetupUndoRedo, bool
 		// If we are removing a reference, remove the function references to this node in the function library
 		if(URigVMFunctionReferenceNode* FunctionReferenceNode = Cast<URigVMFunctionReferenceNode>(LibraryNode))
 		{
-			if (URigVMFunctionLibrary* FunctionLibrary = Cast<URigVMFunctionLibrary>(FunctionReferenceNode->GetLibrary()))
+			if(URigVMBuildData* BuildData = GetBuildData())
 			{
-				FRigVMFunctionReferenceArray* References = FunctionLibrary->FunctionReferences.Find(FunctionReferenceNode->GetReferencedNode());
-				if (References)
-				{
-					References->FunctionReferences.RemoveAll(
-						[FunctionReferenceNode](TSoftObjectPtr<URigVMFunctionReferenceNode>& FunctionReferencePtr) {
-
-							if (!FunctionReferencePtr.IsValid())
-							{
-								FunctionReferencePtr.LoadSynchronous();
-							}
-							if (!FunctionReferencePtr.IsValid())
-							{
-								return true;
-							}
-							return FunctionReferencePtr.Get() == FunctionReferenceNode;
-					});
-				}
+				BuildData->UnregisterFunctionReference(FunctionReferenceNode->GetReferencedNode(), FunctionReferenceNode);
 			}
 		}
 		// If we are removing a function, remove all the references first
 		else if (URigVMFunctionLibrary* FunctionLibrary = Cast<URigVMFunctionLibrary>(LibraryNode->GetGraph()))
 		{
-			const FRigVMFunctionReferenceArray* FunctionReferencesPtr = FunctionLibrary->FunctionReferences.Find(LibraryNode);
-			if (FunctionReferencesPtr)
+			if(URigVMBuildData* BuildData = GetBuildData())
 			{
-				TArray< TSoftObjectPtr<URigVMFunctionReferenceNode> > FunctionReferences = FunctionReferencesPtr->FunctionReferences;
-				for (const TSoftObjectPtr<URigVMFunctionReferenceNode>& FunctionReferencePtr : FunctionReferences)
+				if (const FRigVMFunctionReferenceArray* ReferencesEntry = BuildData->FindFunctionReferences(LibraryNode))
 				{
-					if (!FunctionReferencesPtr->FunctionReferences.Contains(FunctionReferencePtr))
-					{
-						continue;
-					}
-					
-					if (FunctionReferencePtr.IsValid())
+					// make a copy since we'll be modifying the array
+					TArray< TSoftObjectPtr<URigVMFunctionReferenceNode> > FunctionReferences = ReferencesEntry->FunctionReferences;
+					for (const TSoftObjectPtr<URigVMFunctionReferenceNode>& FunctionReferencePtr : FunctionReferences)
 					{
+						if (!ReferencesEntry->FunctionReferences.Contains(FunctionReferencePtr))
 						{
-							FRigVMControllerGraphGuard GraphGuard(this, FunctionReferencePtr->GetGraph(), bSetupUndoRedo);
-							RemoveNode(FunctionReferencePtr.Get());
+							continue;
+						}
+						
+						if (FunctionReferencePtr.IsValid())
+						{
+							{
+								FRigVMControllerGraphGuard GraphGuard(this, FunctionReferencePtr->GetGraph(), bSetupUndoRedo);
+								RemoveNode(FunctionReferencePtr.Get());
 
-							TGuardValue<TSoftObjectPtr<URigVMLibraryNode>> ClearReferencedNodePtr(
-								FunctionReferencePtr->ReferencedNodePtr,
-								TSoftObjectPtr<URigVMLibraryNode>());
+								TGuardValue<TSoftObjectPtr<URigVMLibraryNode>> ClearReferencedNodePtr(
+									FunctionReferencePtr->ReferencedNodePtr,
+									TSoftObjectPtr<URigVMLibraryNode>());
+							}
+							FunctionReferencePtr->ReferencedNodePtr.ResetWeakPtr();
 						}
-						FunctionReferencePtr->ReferencedNodePtr.ResetWeakPtr();
 					}
 				}
-			}
-			FunctionLibrary->FunctionReferences.Remove(LibraryNode);
 
+				BuildData->FunctionReferences.Remove(LibraryNode);
+			}
+			
 			for(const auto& Pair : FunctionLibrary->LocalizedFunctions)
 			{
 				if(Pair.Value == LibraryNode)
@@ -8735,10 +8717,9 @@ URigVMFunctionReferenceNode* URigVMController::AddFunctionReferenceNode(URigVMLi
 
 	Notify(ERigVMGraphNotifType::NodeAdded, FunctionRefNode);
 
-	if (URigVMFunctionLibrary* FunctionLibrary = InFunctionDefinition->GetLibrary())
+	if (URigVMBuildData* BuildData = GetBuildData())
 	{
-		FunctionLibrary->FunctionReferences.FindOrAdd(InFunctionDefinition).FunctionReferences.Add(FunctionRefNode);
-		FunctionLibrary->MarkPackageDirty();
+		BuildData->RegisterFunctionReference(InFunctionDefinition, FunctionRefNode);
 	}
 
 	for (URigVMPin* SourcePin : InFunctionDefinition->Pins)
@@ -11393,6 +11374,20 @@ FProperty* URigVMController::FindPropertyForPin(const FString& InPinPath)
 	return nullptr;
 }
 
+URigVMBuildData* URigVMController::GetBuildData(bool bCreateIfNeeded)
+{
+	static TStrongObjectPtr<URigVMBuildData> sBuildData;
+	if(!sBuildData.IsValid() && bCreateIfNeeded && IsInGameThread())
+	{
+		sBuildData = TStrongObjectPtr<URigVMBuildData>(
+			NewObject<URigVMBuildData>(
+				GetTransientPackage(), 
+				TEXT("RigVMBuildData"), 
+				RF_Transient));
+	}
+	return sBuildData.Get();
+}
+
 int32 URigVMController::DetachLinksFromPinObjects(const TArray<URigVMLink*>* InLinks, bool bNotify)
 {
 	URigVMGraph* Graph = GetGraph();
diff --git a/Engine/Source/Developer/RigVMDeveloper/Private/RigVMModel/RigVMFunctionLibrary.cpp b/Engine/Source/Developer/RigVMDeveloper/Private/RigVMModel/RigVMFunctionLibrary.cpp
--- a/Engine/Source/Developer/RigVMDeveloper/Private/RigVMModel/RigVMFunctionLibrary.cpp
+++ b/Engine/Source/Developer/RigVMDeveloper/Private/RigVMModel/RigVMFunctionLibrary.cpp
@@ -2,6 +2,8 @@
 
 #include "RigVMModel/RigVMFunctionLibrary.h"
 
+#include "RigVMModel/RigVMController.h"
+
 URigVMFunctionLibrary::URigVMFunctionLibrary()
 : URigVMGraph()
 {
@@ -66,6 +68,11 @@ URigVMLibraryNode* URigVMFunctionLibrary::FindFunctionForNode(URigVMNode* InNode
 	return Cast<URigVMLibraryNode>(Subject);
 }
 
+URigVMBuildData* URigVMFunctionLibrary::GetBuildData() const
+{
+	return URigVMController::GetBuildData();
+}
+
 TArray< TSoftObjectPtr<URigVMFunctionReferenceNode> > URigVMFunctionLibrary::GetReferencesForFunction(const FName& InFunctionName)
 {
 	TArray< TSoftObjectPtr<URigVMFunctionReferenceNode> > Result;
@@ -95,21 +102,9 @@ void URigVMFunctionLibrary::ForEachReference(const FName& InFunctionName,
 {
 	if (URigVMLibraryNode* Function = FindFunction(InFunctionName))
 	{
-		const FRigVMFunctionReferenceArray* ReferencesEntry = FunctionReferences.Find(Function);
-		if (ReferencesEntry)
+		if(const URigVMBuildData* BuildData = GetBuildData())
 		{
-			for (int32 ReferenceIndex = 0; ReferenceIndex < ReferencesEntry->Num(); ReferenceIndex++)
-			{
-				const TSoftObjectPtr<URigVMFunctionReferenceNode>& Reference = ReferencesEntry->operator [](ReferenceIndex);
-				if (!Reference.IsValid())
-				{
-					Reference.LoadSynchronous();
-				}
-				if (Reference.IsValid())
-				{
-					PerReferenceFunction(Reference.Get());
-				}
-			}
+			BuildData->ForEachFunctionReference(Function, PerReferenceFunction);
 		}
 	}
 }
@@ -119,50 +114,9 @@ void URigVMFunctionLibrary::ForEachReferenceSoftPtr(const FName& InFunctionName,
 {
 	if (URigVMLibraryNode* Function = FindFunction(InFunctionName))
 	{
-		const FRigVMFunctionReferenceArray* ReferencesEntry = FunctionReferences.Find(Function);
-		if (ReferencesEntry)
-		{
-			for (int32 ReferenceIndex = 0; ReferenceIndex < ReferencesEntry->Num(); ReferenceIndex++)
-			{
-				const TSoftObjectPtr<URigVMFunctionReferenceNode>& Reference = ReferencesEntry->operator [](ReferenceIndex);
-				PerReferenceFunction(Reference);
-			}
-		}
-	}
-}
-
-void URigVMFunctionLibrary::UpdateReferencesForReferenceNode(URigVMFunctionReferenceNode* InReferenceNode)
-{
-	if(InReferenceNode->GetOutermost() == GetTransientPackage())
-	{
-		return;
-	}
-	
-	if(URigVMLibraryNode* Function = InReferenceNode->GetReferencedNode())
-	{
-		if(Function->GetOuter() == this)
+		if(const URigVMBuildData* BuildData = GetBuildData())
 		{
-			FRigVMFunctionReferenceArray* ReferencesEntry = FunctionReferences.Find(Function);
-			if (ReferencesEntry == nullptr)
-			{
-				Modify();
-				FunctionReferences.Add(Function);
-				ReferencesEntry = FunctionReferences.Find(Function);
-			}
-
-			const FString ReferenceNodePathName = InReferenceNode->GetPathName();
-			for (int32 ReferenceIndex = 0; ReferenceIndex < ReferencesEntry->Num(); ReferenceIndex++)
-			{
-				const TSoftObjectPtr<URigVMFunctionReferenceNode>& Reference = ReferencesEntry->operator [](ReferenceIndex);
-				if(Reference.ToString() == ReferenceNodePathName)
-				{
-					return;
-				}
-			}
-
-			Modify();
-			ReferencesEntry->FunctionReferences.Add(InReferenceNode);
-			MarkPackageDirty();
+			BuildData->ForEachFunctionReferenceSoftPtr(Function, PerReferenceFunction);
 		}
 	}
 }
@@ -205,23 +159,3 @@ URigVMLibraryNode* URigVMFunctionLibrary::FindPreviouslyLocalizedFunction(URigVM
 	return LocalizedFunction;
 }
 
-void URigVMFunctionLibrary::ClearInvalidReferences()
-{
-	// check each function's each reference
-	for (TTuple<TObjectPtr<URigVMLibraryNode>, FRigVMFunctionReferenceArray>& FunctionReferenceInfo : FunctionReferences)
-	{
-		FRigVMFunctionReferenceArray* ReferencesEntry = &FunctionReferenceInfo.Value;
-
-		Modify();
-		int32 NumRemoved = ReferencesEntry->FunctionReferences.RemoveAll([](TSoftObjectPtr<URigVMFunctionReferenceNode> Referencer)
-		{
-			return Referencer.ToString().StartsWith(TEXT("/Temp/"));
-		});
-
-		if (NumRemoved > 0)
-		{
-			MarkPackageDirty();
-		}
-	}
-}
-
diff --git a/Engine/Source/Developer/RigVMDeveloper/Public/RigVMModel/RigVMBuildData.h b/Engine/Source/Developer/RigVMDeveloper/Public/RigVMModel/RigVMBuildData.h
new file mode 100644
--- /dev/null
+++ b/Engine/Source/Developer/RigVMDeveloper/Public/RigVMModel/RigVMBuildData.h
@@ -0,0 +1,121 @@
+// Copyright Epic Games, Inc. All Rights Reserved.
+
+#pragma once
+
+#include "RigVMModel/Nodes/RigVMFunctionReferenceNode.h"
+#include "RigVMBuildData.generated.h"
+
+USTRUCT(BlueprintType)
+struct RIGVMDEVELOPER_API FRigVMFunctionReferenceArray
+{
+	GENERATED_BODY()
+
+	// Resets the data structure and maintains all storage.
+	void Reset() { FunctionReferences.Reset();  }
+
+	// Returns true if a given function reference index is valid.
+	bool IsValidIndex(int32 InIndex) const { return FunctionReferences.IsValidIndex(InIndex); }
+
+	// Returns the number of reference functions
+	FORCEINLINE int32 Num() const { return FunctionReferences.Num(); }
+
+	// const accessor for an function reference given its index
+	FORCEINLINE const TSoftObjectPtr<URigVMFunctionReferenceNode>& operator[](int32 InIndex) const { return FunctionReferences[InIndex]; }
+
+	UPROPERTY(VisibleAnywhere, Category = "BuildData")
+	TArray< TSoftObjectPtr<URigVMFunctionReferenceNode> > FunctionReferences;
+};
+
+USTRUCT()
+struct RIGVMDEVELOPER_API FRigVMReferenceNodeData
+{
+	GENERATED_BODY();
+
+public:
+	
+	FRigVMReferenceNodeData()
+		:ReferenceNodePath()
+		,ReferencedFunctionPath()
+	{}
+
+	FRigVMReferenceNodeData(URigVMFunctionReferenceNode* InReferenceNode);
+
+	UPROPERTY()
+	FString ReferenceNodePath;
+
+	UPROPERTY()
+	FString ReferencedFunctionPath;
+
+	TSoftObjectPtr<URigVMFunctionReferenceNode> GetReferenceNodeObjectPath();
+	TSoftObjectPtr<URigVMLibraryNode> GetReferencedFunctionObjectPath();
+	URigVMFunctionReferenceNode* GetReferenceNode();
+	URigVMLibraryNode* GetReferencedFunction();
+
+private:
+
+	TSoftObjectPtr<URigVMFunctionReferenceNode> ReferenceNodePtr;
+	TSoftObjectPtr<URigVMLibraryNode> LibraryNodePtr;
+};
+
+/**
+ * The Build Data is used to store transient / intermediate build information
+ * for the RigVM graph to improve the user experience.
+ * This object is never serialized.
+ */
+UCLASS()
+class RIGVMDEVELOPER_API URigVMBuildData : public UObject
+{
+	GENERATED_BODY()
+
+public:
+
+	// Default constructor
+	URigVMBuildData();
+
+	// Returns the list of references for a given function definition
+	const FRigVMFunctionReferenceArray* FindFunctionReferences(const URigVMLibraryNode* InFunction) const;
+
+	/**
+	 * Iterator function to invoke a lambda / TFunction for each reference of a function
+	 * @param InFunction The function to iterate all references for
+	 * @param PerReferenceFunction The function to invoke for each reference
+	 */
+	void ForEachFunctionReference(const URigVMLibraryNode* InFunction, TFunction<void(URigVMFunctionReferenceNode*)> PerReferenceFunction) const;
+
+	/**
+	* Iterator function to invoke a lambda / TFunction for each reference of a function
+	* @param InFunction The function to iterate all references for
+	* @param PerReferenceFunction The function to invoke for each reference
+	*/
+	void ForEachFunctionReferenceSoftPtr(const URigVMLibraryNode* InFunction, TFunction<void(TSoftObjectPtr<URigVMFunctionReferenceNode>)> PerReferenceFunction) const;
+
+	// Update the references list for a given reference node
+	void UpdateReferencesForFunctionReferenceNode(URigVMFunctionReferenceNode* InReferenceNode);
+
+	// registers a new reference node for a given function
+	void RegisterFunctionReference(URigVMLibraryNode* InFunction, URigVMFunctionReferenceNode* InReference);
+
+	// registers a new reference node for a given function
+	void RegisterFunctionReference(TSoftObjectPtr<URigVMLibraryNode> InFunction, TSoftObjectPtr<URigVMFunctionReferenceNode> InReference);
+
+	// registers a new reference node for a given function
+	void RegisterFunctionReference(FRigVMReferenceNodeData InReferenceNodeData);
+
+	// unregisters a new reference node for a given function
+	void UnregisterFunctionReference(URigVMLibraryNode* InFunction, URigVMFunctionReferenceNode* InReference);
+
+	// unregisters a new reference node for a given function
+	void UnregisterFunctionReference(TSoftObjectPtr<URigVMLibraryNode> InFunction, TSoftObjectPtr<URigVMFunctionReferenceNode> InReference);
+
+	// Clear references to temp assets
+	void ClearInvalidReferences();
+
+private:
+
+	UPROPERTY(VisibleAnywhere, Category = "BuildData")
+	TMap< TSoftObjectPtr<URigVMLibraryNode>, FRigVMFunctionReferenceArray > FunctionReferences;
+
+	friend class URigVMController;
+	friend class URigVMCompiler;
+};
+
diff --git a/Engine/Source/Developer/RigVMDeveloper/Public/RigVMModel/RigVMController.h b/Engine/Source/Developer/RigVMDeveloper/Public/RigVMModel/RigVMController.h
--- a/Engine/Source/Developer/RigVMDeveloper/Public/RigVMModel/RigVMController.h
+++ b/Engine/Source/Developer/RigVMDeveloper/Public/RigVMModel/RigVMController.h
@@ -715,6 +715,9 @@ public:
 	// A delegate to inform the host / client about the progress during a bulk edit
 	FRigVMController_OnBulkEditProgressDelegate OnBulkEditProgressDelegate;
 
+	// Returns the build data of the host
+	static URigVMBuildData* GetBuildData(bool bCreateIfNeeded = true);
+
 	int32 DetachLinksFromPinObjects(const TArray<URigVMLink*>* InLinks = nullptr, bool bNotify = false);
 	int32 ReattachLinksToPinObjects(bool bFollowCoreRedirectors = false, const TArray<URigVMLink*>* InLinks = nullptr, bool bNotify = false, bool bSetupOrphanedPins = false);
 	void AddPinRedirector(bool bInput, bool bOutput, const FString& OldPinPath, const FString& NewPinPath);
diff --git a/Engine/Source/Developer/RigVMDeveloper/Public/RigVMModel/RigVMFunctionLibrary.h b/Engine/Source/Developer/RigVMDeveloper/Public/RigVMModel/RigVMFunctionLibrary.h
--- a/Engine/Source/Developer/RigVMDeveloper/Public/RigVMModel/RigVMFunctionLibrary.h
+++ b/Engine/Source/Developer/RigVMDeveloper/Public/RigVMModel/RigVMFunctionLibrary.h
@@ -3,31 +3,11 @@
 #pragma once
 
 #include "RigVMGraph.h"
+#include "RigVMBuildData.h"
 #include "RigVMModel/Nodes/RigVMLibraryNode.h"
 #include "RigVMModel/Nodes/RigVMFunctionReferenceNode.h"
 #include "RigVMFunctionLibrary.generated.h"
 
-USTRUCT(BlueprintType)
-struct RIGVMDEVELOPER_API FRigVMFunctionReferenceArray
-{
-	GENERATED_BODY()
-
-	// Resets the data structure and maintains all storage.
-	void Reset() { FunctionReferences.Reset();  }
-
-	// Returns true if a given function reference index is valid.
-	bool IsValidIndex(int32 InIndex) const { return FunctionReferences.IsValidIndex(InIndex); }
-
-	// Returns the number of reference functions
-	FORCEINLINE int32 Num() const { return FunctionReferences.Num(); }
-
-	// const accessor for an function reference given its index
-	FORCEINLINE const TSoftObjectPtr<URigVMFunctionReferenceNode>& operator[](int32 InIndex) const { return FunctionReferences[InIndex]; }
-
-	UPROPERTY()
-	TArray< TSoftObjectPtr<URigVMFunctionReferenceNode> > FunctionReferences;
-};
-
 /**
  * The Function Library is a graph used only to store
  * the sub graphs used for functions.
@@ -59,6 +39,10 @@ public:
 	UFUNCTION(BlueprintCallable, Category = RigVMGraph)
     URigVMLibraryNode* FindFunctionForNode(URigVMNode* InNode) const;
 
+	// returns the build data for for this library
+	UFUNCTION()
+	URigVMBuildData* GetBuildData() const;
+
 	// Returns all references for a given function name
 	UFUNCTION(BlueprintCallable, Category = RigVMGraph)
 	TArray< TSoftObjectPtr<URigVMFunctionReferenceNode> > GetReferencesForFunction(const FName& InFunctionName);
@@ -67,9 +51,6 @@ public:
 	UFUNCTION(BlueprintCallable, Category = RigVMGraph)
 	TArray< FString > GetReferencePathsForFunction(const FName& InFunctionName);
 
-	// Update the references list for a given reference node
-	void UpdateReferencesForReferenceNode(URigVMFunctionReferenceNode* InReferenceNode);
-
 	/**
 	* Iterator function to invoke a lambda / TFunction for each reference of a function
 	* @param InFunctionName The function name to iterate all references for
@@ -88,14 +69,11 @@ public:
 	// We maintain meta data on what functions have been created locally based on which other ones,
 	// and use this method to avoid redundant localizations.
 	URigVMLibraryNode* FindPreviouslyLocalizedFunction(URigVMLibraryNode* InFunctionToLocalize);
-
-	// Clear references to temp assets
-	void ClearInvalidReferences();
 	
 private:
 
 	UPROPERTY()
-	TMap< TObjectPtr<URigVMLibraryNode>, FRigVMFunctionReferenceArray > FunctionReferences;
+	TMap< TObjectPtr<URigVMLibraryNode>, FRigVMFunctionReferenceArray > FunctionReferences_DEPRECATED;
 
 	// A map which stores a library node per original pathname.
 	// The source pathname is the full path of the source function that was localized
@@ -105,5 +83,6 @@ private:
 
 	friend class URigVMController;
 	friend class URigVMCompiler;
+	friend class UControlRigBlueprint;
 };
 
