Fix non-unity
#rnx

diff --git a/Engine/Source/Developer/RigVMDeveloper/Private/RigVMModel/RigVMController.cpp b/Engine/Source/Developer/RigVMDeveloper/Private/RigVMModel/RigVMController.cpp
--- a/Engine/Source/Developer/RigVMDeveloper/Private/RigVMModel/RigVMController.cpp
+++ b/Engine/Source/Developer/RigVMDeveloper/Private/RigVMModel/RigVMController.cpp
@@ -13,6 +13,7 @@
 #include "RigVMDeveloperModule.h"
 #include "UObject/PropertyPortFlags.h"
 #include "UObject/Package.h"
+#include "UObject/StrongObjectPtr.h"
 #include "Misc/CoreMisc.h"
 #include "Algo/Sort.h"
 #include "RigVMPythonUtils.h"
diff --git a/Engine/Source/Runtime/Engine/Private/Level.cpp b/Engine/Source/Runtime/Engine/Private/Level.cpp
--- a/Engine/Source/Runtime/Engine/Private/Level.cpp
+++ b/Engine/Source/Runtime/Engine/Private/Level.cpp
@@ -28,6 +28,7 @@ Level.cpp: Level-related functions
 #include "Engine/Brush.h"
 #include "Engine/Engine.h"
 #include "Containers/TransArray.h"
+#include "UObject/MetaData.h"
 #include "UObject/ObjectSaveContext.h"
 #include "UObject/UObjectHash.h"
 #include "UObject/UObjectIterator.h"
