Control Rig: Fixed variable node is ill-formed after change to real types

diff --git a/Engine/Plugins/Animation/ControlRig/Source/ControlRigDeveloper/Private/ControlRigBlueprint.cpp b/Engine/Plugins/Animation/ControlRig/Source/ControlRigDeveloper/Private/ControlRigBlueprint.cpp
--- a/Engine/Plugins/Animation/ControlRig/Source/ControlRigDeveloper/Private/ControlRigBlueprint.cpp
+++ b/Engine/Plugins/Animation/ControlRig/Source/ControlRigDeveloper/Private/ControlRigBlueprint.cpp
@@ -3687,25 +3687,30 @@ void UControlRigBlueprint::PatchVariableNodesWithIncorrectType()
 		{
 			if (URigVMVariableNode* VariableNode = Cast<URigVMVariableNode>(Node))
 			{
+				if (VariableNode->IsInputArgument())
+				{
+					continue;
+				}
+				
 				FRigVMGraphVariableDescription Description = VariableNode->GetVariableDescription();
 
-				// Check for inputs and local variables
-				TArray<FRigVMGraphVariableDescription> LocalVariables = Graph->GetLocalVariables(true);
-				bool bLocalVariableFound = false;
-				for (FRigVMGraphVariableDescription Variable : LocalVariables)
+				// Check for local variables
+				if (VariableNode->IsLocalVariable())
 				{
-					if (Variable.Name == Description.Name)
+					TArray<FRigVMGraphVariableDescription> LocalVariables = Graph->GetLocalVariables(false);
+					for (FRigVMGraphVariableDescription Variable : LocalVariables)
 					{
-						if (Local::RefreshIfNeeded(Controller, VariableNode, Variable.CPPType, Variable.CPPTypeObject))
+						if (Variable.Name == Description.Name)
 						{
-							bDirtyDuringLoad = true;
+							if (Local::RefreshIfNeeded(Controller, VariableNode, Variable.CPPType, Variable.CPPTypeObject))
+							{
+								bDirtyDuringLoad = true;
+							}
+							break;
 						}
-						bLocalVariableFound = true;
-						break;
 					}
 				}
-
-				if (!bLocalVariableFound)
+				else
 				{
 					for (struct FBPVariableDescription& Variable : NewVariables)
 					{
diff --git a/Engine/Source/Developer/RigVMDeveloper/Private/RigVMModel/RigVMController.cpp b/Engine/Source/Developer/RigVMDeveloper/Private/RigVMModel/RigVMController.cpp
--- a/Engine/Source/Developer/RigVMDeveloper/Private/RigVMModel/RigVMController.cpp
+++ b/Engine/Source/Developer/RigVMDeveloper/Private/RigVMModel/RigVMController.cpp
@@ -11917,31 +11917,47 @@ void URigVMController::RepopulatePinsOnNode(URigVMNode* InNode, bool bFollowCore
 			const FRigVMGraphVariableDescription VariableDescription = VariableNode->GetVariableDescription();
 			const FRigVMExternalVariable CurrentExternalVariable = VariableDescription.ToExternalVariable();
 
-			for(const FRigVMExternalVariable& ExternalVariable : ExternalVariables)
+			FRigVMExternalVariable Variable;
+			if (VariableNode->IsInputArgument())
 			{
-				if(ExternalVariable.Name == CurrentExternalVariable.Name)
+				if (URigVMFunctionEntryNode* GraphEntryNode = Graph->GetEntryNode())
 				{
-					if(ExternalVariable.TypeName != CurrentExternalVariable.TypeName ||
-						ExternalVariable.TypeObject != CurrentExternalVariable.TypeObject ||
-						ExternalVariable.bIsArray != CurrentExternalVariable.bIsArray)
+					if (URigVMPin* EntryPin = GraphEntryNode->FindPin(VariableDescription.Name.ToString()))
 					{
-						FString CPPType;
-						UObject* CPPTypeObject;
-						
-						if(RigVMTypeUtils::CPPTypeFromExternalVariable(ExternalVariable, CPPType, &CPPTypeObject))
-						{
-							RefreshVariableNode(VariableNode->GetFName(), ExternalVariable.Name, CPPType, ExternalVariable.TypeObject, false, bSetupOrphanedPins);
-						}
-						else
-						{
-							ReportErrorf(
-								TEXT("Control Rig '%s', Type of Variable '%s' cannot be resolved."),
-								*InNode->GetOutermost()->GetPathName(),
-								*ExternalVariable.Name.ToString()
-							);
-						}
+						Variable = RigVMTypeUtils::ExternalVariableFromCPPType(VariableDescription.Name, EntryPin->GetCPPType(), EntryPin->GetCPPTypeObject());
 					}
-					break;
+				}
+			}
+			else
+			{				
+				for(const FRigVMExternalVariable& ExternalVariable : ExternalVariables)
+				{
+					if(ExternalVariable.Name == CurrentExternalVariable.Name)
+					{
+						Variable = ExternalVariable;
+						break;
+					}
+				}
+			}
+			
+			if(Variable.TypeName != CurrentExternalVariable.TypeName ||
+				Variable.TypeObject != CurrentExternalVariable.TypeObject ||
+				Variable.bIsArray != CurrentExternalVariable.bIsArray)
+			{
+				FString CPPType;
+				UObject* CPPTypeObject;
+				
+				if(RigVMTypeUtils::CPPTypeFromExternalVariable(Variable, CPPType, &CPPTypeObject))
+				{
+					RefreshVariableNode(VariableNode->GetFName(), Variable.Name, CPPType, Variable.TypeObject, false, bSetupOrphanedPins);
+				}
+				else
+				{
+					ReportErrorf(
+						TEXT("Control Rig '%s', Type of Variable '%s' cannot be resolved."),
+						*InNode->GetOutermost()->GetPathName(),
+						*Variable.Name.ToString()
+					);
 				}
 			}
 		}
