Control Rig: FIx crash when using old SpringInterpSpringInterp node

diff --git a/Engine/Source/Developer/RigVMDeveloper/Private/RigVMCompiler/RigVMAST.cpp b/Engine/Source/Developer/RigVMDeveloper/Private/RigVMCompiler/RigVMAST.cpp
--- a/Engine/Source/Developer/RigVMDeveloper/Private/RigVMCompiler/RigVMAST.cpp
+++ b/Engine/Source/Developer/RigVMDeveloper/Private/RigVMCompiler/RigVMAST.cpp
@@ -1488,8 +1488,10 @@ void FRigVMParserAST::FoldAssignments()
 		}
 
 		FRigVMAssignExprAST* AssignExpr = Expression->To<FRigVMAssignExprAST>();
-		ensure(AssignExpr->Parents.Num() == 1);
-		ensure(AssignExpr->Children.Num() == 1);
+		if (AssignExpr->Parents.Num() != 1 || AssignExpr->Children.Num() != 1)
+		{
+			continue;
+		}
 
 		URigVMPin* SourcePin = AssignExpr->GetSourcePin();
 		URigVMPin* TargetPin = AssignExpr->GetTargetPin();
