Control Rig: Fix ControlRig.Debug unit test failing

diff --git a/Engine/Source/Developer/RigVMDeveloper/Private/RigVMModel/RigVMBuildData.cpp b/Engine/Source/Developer/RigVMDeveloper/Private/RigVMModel/RigVMBuildData.cpp
--- a/Engine/Source/Developer/RigVMDeveloper/Private/RigVMModel/RigVMBuildData.cpp
+++ b/Engine/Source/Developer/RigVMDeveloper/Private/RigVMModel/RigVMBuildData.cpp
@@ -65,6 +65,7 @@ URigVMLibraryNode* FRigVMReferenceNodeData::GetReferencedFunction()
 
 URigVMBuildData::URigVMBuildData()
 : UObject()
+, bIsRunningUnitTest(false)
 {
 }
 
@@ -234,6 +235,11 @@ void URigVMBuildData::UnregisterFunctionReference(TSoftObjectPtr<URigVMLibraryNo
 
 void URigVMBuildData::ClearInvalidReferences()
 {
+	if (bIsRunningUnitTest)
+	{
+		return;
+	}
+	
 	Modify();
 	
 	// check each function's each reference
diff --git a/Engine/Source/Developer/RigVMDeveloper/Private/RigVMModel/RigVMController.cpp b/Engine/Source/Developer/RigVMDeveloper/Private/RigVMModel/RigVMController.cpp
--- a/Engine/Source/Developer/RigVMDeveloper/Private/RigVMModel/RigVMController.cpp
+++ b/Engine/Source/Developer/RigVMDeveloper/Private/RigVMModel/RigVMController.cpp
@@ -233,6 +233,16 @@ void URigVMController::ResendAllNotifications()
 	}
 }
 
+void URigVMController::SetIsRunningUnitTest(bool bIsRunning)
+{
+	bIsRunningUnitTest = bIsRunning;
+
+	if(URigVMBuildData* BuildData = GetBuildData())
+	{
+		BuildData->SetIsRunningUnitTest(bIsRunning);
+	}	
+}
+
 void URigVMController::HandleModifiedEvent(ERigVMGraphNotifType InNotifType, URigVMGraph* InGraph, UObject* InSubject)
 {
 	switch (InNotifType)
diff --git a/Engine/Source/Developer/RigVMDeveloper/Public/RigVMModel/RigVMBuildData.h b/Engine/Source/Developer/RigVMDeveloper/Public/RigVMModel/RigVMBuildData.h
--- a/Engine/Source/Developer/RigVMDeveloper/Public/RigVMModel/RigVMBuildData.h
+++ b/Engine/Source/Developer/RigVMDeveloper/Public/RigVMModel/RigVMBuildData.h
@@ -110,11 +110,16 @@ public:
 	// Clear references to temp assets
 	void ClearInvalidReferences();
 
+	// Helper function to disable clearing transient package references
+	void SetIsRunningUnitTest(bool bIsRunning) { bIsRunningUnitTest = bIsRunning; }
+
 private:
 
 	UPROPERTY(VisibleAnywhere, Category = "BuildData")
 	TMap< TSoftObjectPtr<URigVMLibraryNode>, FRigVMFunctionReferenceArray > FunctionReferences;
 
+	bool bIsRunningUnitTest;
+
 	friend class URigVMController;
 	friend class URigVMCompiler;
 };
diff --git a/Engine/Source/Developer/RigVMDeveloper/Public/RigVMModel/RigVMController.h b/Engine/Source/Developer/RigVMDeveloper/Public/RigVMModel/RigVMController.h
--- a/Engine/Source/Developer/RigVMDeveloper/Public/RigVMModel/RigVMController.h
+++ b/Engine/Source/Developer/RigVMDeveloper/Public/RigVMModel/RigVMController.h
@@ -783,7 +783,7 @@ public:
 	 * Helper function to disable a series of checks that can be ignored during a unit test
 	 */
 	UFUNCTION(BlueprintCallable, Category = RigVMController)
-	void SetIsRunningUnitTest(bool bIsRunning) { bIsRunningUnitTest = bIsRunning; }
+	void SetIsRunningUnitTest(bool bIsRunning);
 
 private:
 
