Control Rig: Fix errors when undoing change exposed pin type of function

diff --git a/Engine/Source/Developer/RigVMDeveloper/Private/RigVMModel/RigVMController.cpp b/Engine/Source/Developer/RigVMDeveloper/Private/RigVMModel/RigVMController.cpp
--- a/Engine/Source/Developer/RigVMDeveloper/Private/RigVMModel/RigVMController.cpp
+++ b/Engine/Source/Developer/RigVMDeveloper/Private/RigVMModel/RigVMController.cpp
@@ -8523,7 +8523,7 @@ bool URigVMController::ChangeExposedPinType(const FName& InPinName, const FStrin
 	FRigVMBaseAction Action;
 	if (bSetupUndoRedo)
 	{
-		Action.Title = FString::Printf(TEXT("Change Pin Type"));
+		Action.Title = FString::Printf(TEXT("Change Exposed Pin Type"));
 		ActionStack->BeginAction(Action);
 	}
 
@@ -8564,7 +8564,7 @@ bool URigVMController::ChangeExposedPinType(const FName& InPinName, const FStrin
         {
 			if (URigVMPin* ReferencedNodePin = ReferenceNode->FindPin(Pin->GetName()))
 			{
-				FRigVMControllerGraphGuard GraphGuard(this, ReferenceNode->GetGraph(), false);
+				FRigVMControllerGraphGuard GraphGuard(this, ReferenceNode->GetGraph(), bSetupUndoRedo);
 				ChangePinType(ReferencedNodePin, InCPPType, InCPPTypeObjectPath, bSetupUndoRedo, bSetupOrphanPins);
 				RemoveUnusedOrphanedPins(ReferenceNode, true);
 			}
