ControlRig: Fix for static analysis

diff --git a/Engine/Source/Developer/RigVMDeveloper/Private/RigVMModel/RigVMController.cpp b/Engine/Source/Developer/RigVMDeveloper/Private/RigVMModel/RigVMController.cpp
--- a/Engine/Source/Developer/RigVMDeveloper/Private/RigVMModel/RigVMController.cpp
+++ b/Engine/Source/Developer/RigVMDeveloper/Private/RigVMModel/RigVMController.cpp
@@ -2471,6 +2471,8 @@ URigVMNode* URigVMController::EjectNodeFromPin(const FString& InPinPath, bool bS
 	{
 		return nullptr;
 	}
+	
+	check(EjectedNode);
 
 	return EjectNodeFromPin(Pin, bSetupUndoRedo, bPrintPythonCommand);
 }
