Control Rig: Fixed Cook Asserts in RigVMTypeUtils

diff --git a/Engine/Source/Developer/RigVMDeveloper/Private/RigVMModel/RigVMController.cpp b/Engine/Source/Developer/RigVMDeveloper/Private/RigVMModel/RigVMController.cpp
--- a/Engine/Source/Developer/RigVMDeveloper/Private/RigVMModel/RigVMController.cpp
+++ b/Engine/Source/Developer/RigVMDeveloper/Private/RigVMModel/RigVMController.cpp
@@ -11949,27 +11949,38 @@ void URigVMController::RepopulatePinsOnNode(URigVMNode* InNode, bool bFollowCore
 					}
 				}
 			}
-			
-			if(Variable.TypeName != CurrentExternalVariable.TypeName ||
-				Variable.TypeObject != CurrentExternalVariable.TypeObject ||
-				Variable.bIsArray != CurrentExternalVariable.bIsArray)
+
+			if (Variable.IsValid(true))
 			{
-				FString CPPType;
-				UObject* CPPTypeObject;
-				
-				if(RigVMTypeUtils::CPPTypeFromExternalVariable(Variable, CPPType, &CPPTypeObject))
+				if(Variable.TypeName != CurrentExternalVariable.TypeName ||
+				   Variable.TypeObject != CurrentExternalVariable.TypeObject ||
+				   Variable.bIsArray != CurrentExternalVariable.bIsArray)
 				{
-					RefreshVariableNode(VariableNode->GetFName(), Variable.Name, CPPType, Variable.TypeObject, false, bSetupOrphanedPins);
-				}
-				else
-				{
-					ReportErrorf(
-						TEXT("Control Rig '%s', Type of Variable '%s' cannot be resolved."),
-						*InNode->GetOutermost()->GetPathName(),
-						*Variable.Name.ToString()
-					);
+					FString CPPType;
+					UObject* CPPTypeObject;
+				
+					if(RigVMTypeUtils::CPPTypeFromExternalVariable(Variable, CPPType, &CPPTypeObject))
+					{
+						RefreshVariableNode(VariableNode->GetFName(), Variable.Name, CPPType, Variable.TypeObject, false, bSetupOrphanedPins);
+					}
+					else
+					{
+						ReportErrorf(
+							TEXT("Control Rig '%s', Type of Variable '%s' cannot be resolved."),
+							*InNode->GetOutermost()->GetPathName(),
+							*Variable.Name.ToString()
+						);
+					}
 				}
 			}
+			else
+			{
+				ReportErrorf(
+					TEXT("Control Rig '%s', Variable '%s' not found."),
+					*InNode->GetOutermost()->GetPathName(),
+					*CurrentExternalVariable.Name.ToString()
+				);
+			}
 		}
 		
 		RemovePinsDuringRepopulate(InNode, ValuePin->SubPins, bNotify, bSetupOrphanedPins);
