Downgrading error to warning
#rnx

diff --git a/Engine/Source/Developer/RigVMDeveloper/Private/RigVMModel/RigVMController.cpp b/Engine/Source/Developer/RigVMDeveloper/Private/RigVMModel/RigVMController.cpp
--- a/Engine/Source/Developer/RigVMDeveloper/Private/RigVMModel/RigVMController.cpp
+++ b/Engine/Source/Developer/RigVMDeveloper/Private/RigVMModel/RigVMController.cpp
@@ -11975,7 +11975,7 @@ void URigVMController::RepopulatePinsOnNode(URigVMNode* InNode, bool bFollowCore
 			}
 			else
 			{
-				ReportErrorf(
+				ReportWarningf(
 					TEXT("Control Rig '%s', Variable '%s' not found."),
 					*InNode->GetOutermost()->GetPathName(),
 					*CurrentExternalVariable.Name.ToString()
