Control Rig: Fix function changes when folding is disabled

diff --git a/Engine/Source/Developer/RigVMDeveloper/Private/RigVMCompiler/RigVMAST.cpp b/Engine/Source/Developer/RigVMDeveloper/Private/RigVMCompiler/RigVMAST.cpp
--- a/Engine/Source/Developer/RigVMDeveloper/Private/RigVMCompiler/RigVMAST.cpp
+++ b/Engine/Source/Developer/RigVMDeveloper/Private/RigVMCompiler/RigVMAST.cpp
@@ -2445,7 +2445,7 @@ FString FRigVMParserAST::DumpDot() const
 					}
 					else
 					{
-						Label = InExpr->To<FRigVMVarExprAST>()->GetPin()->GetName();
+						Label = InExpr->To<FRigVMVarExprAST>()->GetPin()->GetPinPath(true);
 					}
 						
 					if (InExpr->To<FRigVMVarExprAST>()->IsExecuteContext())
diff --git a/Engine/Source/Developer/RigVMDeveloper/Private/RigVMCompiler/RigVMCompiler.cpp b/Engine/Source/Developer/RigVMDeveloper/Private/RigVMCompiler/RigVMCompiler.cpp
--- a/Engine/Source/Developer/RigVMDeveloper/Private/RigVMCompiler/RigVMCompiler.cpp
+++ b/Engine/Source/Developer/RigVMDeveloper/Private/RigVMCompiler/RigVMCompiler.cpp
@@ -1926,7 +1926,7 @@ void URigVMCompiler::InitializeLocalVariables(const FRigVMExprAST* InExpr, FRigV
 	}
 }
 
-FString URigVMCompiler::GetPinHash(const URigVMPin* InPin, const FRigVMVarExprAST* InVarExpr, bool bIsDebugValue)
+FString URigVMCompiler::GetPinHash(const URigVMPin* InPin, const FRigVMVarExprAST* InVarExpr, bool bIsDebugValue, const FRigVMASTProxy& InPinProxy)
 {
 	FString Prefix = bIsDebugValue ? TEXT("DebugWatch:") : TEXT("");
 	FString Suffix;
@@ -2029,6 +2029,20 @@ FString URigVMCompiler::GetPinHash(const URigVMPin* InPin, const FRigVMVarExprAS
 					}
 				}
 			}
+			else if(VariableNode->IsInputArgument())
+			{
+				FString FullPath;
+				if (InPinProxy.IsValid())
+				{
+					FullPath = InPinProxy.GetCallstack().GetCallPath(true);
+				}
+				else
+				{						
+					const FRigVMASTProxy NodeProxy = InVarExpr->GetProxy().GetSibling(Node);
+					FullPath = InPinProxy.GetCallstack().GetCallPath(true);
+				}
+				return FString::Printf(TEXT("%s%s%s"), *Prefix, *FullPath, *Suffix);
+			}
 
 			if (!bIsLiteral)
 			{		
@@ -2070,10 +2084,23 @@ FString URigVMCompiler::GetPinHash(const URigVMPin* InPin, const FRigVMVarExprAS
 					bUseFullNodePath = false;
 				}
 			}
+			else if(Node->IsA<URigVMFunctionEntryNode>() || Node->IsA<URigVMFunctionReturnNode>())
+			{
+				const FString FullPath = InPinProxy.GetCallstack().GetCallPath(true);
+				return FString::Printf(TEXT("%s%s%s"), *Prefix, *FullPath, *Suffix);
+			}
 		}
 	}
 
-	return FString::Printf(TEXT("%s%s%s"), *Prefix, *InPin->GetPinPath(bUseFullNodePath), *Suffix);
+	if (InPinProxy.IsValid())
+	{
+		const FString FullPath = InPinProxy.GetCallstack().GetCallPath(true);
+		return FString::Printf(TEXT("%s%s%s"), *Prefix, *FullPath, *Suffix);
+	}
+
+	FString PinPath = InPin->GetPinPath(bUseFullNodePath);
+	ensureMsgf(!PinPath.StartsWith(TEXT("FunctionLibrary::")), TEXT("A library path should never be part of a pin hash %s."), *PinPath);
+	return FString::Printf(TEXT("%s%s%s"), *Prefix, *PinPath, *Suffix);
 }
 
 const FRigVMVarExprAST* URigVMCompiler::GetSourceVarExpr(const FRigVMExprAST* InExpr)
@@ -2516,7 +2543,7 @@ FRigVMOperand URigVMCompiler::FindOrAddRegister(const FRigVMVarExprAST* InVarExp
 			{
 				if (URigVMPin* VirtualPin = Cast<URigVMPin>(Proxy.GetSubject()))
 				{
-					FString VirtualPinHash = GetPinHash(VirtualPin, InVarExpr, bIsDebugValue);
+					FString VirtualPinHash = GetPinHash(VirtualPin, InVarExpr, bIsDebugValue, Proxy);
 					WorkData.PinPathToOperand->Add(VirtualPinHash, Operand);
 				}	
 			}
diff --git a/Engine/Source/Developer/RigVMDeveloper/Public/RigVMCompiler/RigVMCompiler.h b/Engine/Source/Developer/RigVMDeveloper/Public/RigVMCompiler/RigVMCompiler.h
--- a/Engine/Source/Developer/RigVMDeveloper/Public/RigVMCompiler/RigVMCompiler.h
+++ b/Engine/Source/Developer/RigVMDeveloper/Public/RigVMCompiler/RigVMCompiler.h
@@ -145,7 +145,7 @@ public:
 	bool Compile(URigVMGraph* InGraph, URigVMController* InController, URigVM* OutVM, const TArray<FRigVMExternalVariable>& InExternalVariables, const TArray<FRigVMUserDataArray>& InRigVMUserData, TMap<FString, FRigVMOperand>* OutOperands, TSharedPtr<FRigVMParserAST> InAST = TSharedPtr<FRigVMParserAST>());
 
 	static UScriptStruct* GetScriptStructForCPPType(const FString& InCPPType);
-	static FString GetPinHash(const URigVMPin* InPin, const FRigVMVarExprAST* InVarExpr, bool bIsDebugValue = false);
+	static FString GetPinHash(const URigVMPin* InPin, const FRigVMVarExprAST* InVarExpr, bool bIsDebugValue = false, const FRigVMASTProxy& InPinProxy = FRigVMASTProxy());
 
 	// follows assignment expressions to find the source ref counted containers
 	// since ref counted containers are not copied for assignments.
