Control Rig: Fix compiler issue RigVMCompiler.cpp Dereferencing NULL pointer 'InVarExpr'

diff --git a/Engine/Source/Developer/RigVMDeveloper/Private/RigVMCompiler/RigVMCompiler.cpp b/Engine/Source/Developer/RigVMDeveloper/Private/RigVMCompiler/RigVMCompiler.cpp
--- a/Engine/Source/Developer/RigVMDeveloper/Private/RigVMCompiler/RigVMCompiler.cpp
+++ b/Engine/Source/Developer/RigVMDeveloper/Private/RigVMCompiler/RigVMCompiler.cpp
@@ -1999,7 +1999,7 @@ FString URigVMCompiler::GetPinHash(const URigVMPin* InPin, const FRigVMVarExprAS
 				if (bIsLiteral)
 				{
 					// Literal values will be reused for all instance of local variables
-					if (InVarExpr->NumParents() == 0 && InVarExpr->NumChildren() == 0)
+					if (InVarExpr && InVarExpr->NumParents() == 0 && InVarExpr->NumChildren() == 0)
 					{
 						return FString::Printf(TEXT("%sLocalVariableDefault::%s|%s%s"), *Prefix, *Node->GetGraph()->GetGraphName(), *VariableName.ToString(), *Suffix);
 					}
@@ -2036,7 +2036,7 @@ FString URigVMCompiler::GetPinHash(const URigVMPin* InPin, const FRigVMVarExprAS
 				{
 					FullPath = InPinProxy.GetCallstack().GetCallPath(true);
 				}
-				else
+				else if(InVarExpr)
 				{						
 					const FRigVMASTProxy NodeProxy = InVarExpr->GetProxy().GetSibling(Node);
 					FullPath = InPinProxy.GetCallstack().GetCallPath(true);
