Control Rig: Sharing of operands only for matching types

diff --git a/Engine/Source/Developer/RigVMDeveloper/Private/RigVMCompiler/RigVMCompiler.cpp b/Engine/Source/Developer/RigVMDeveloper/Private/RigVMCompiler/RigVMCompiler.cpp
--- a/Engine/Source/Developer/RigVMDeveloper/Private/RigVMCompiler/RigVMCompiler.cpp
+++ b/Engine/Source/Developer/RigVMDeveloper/Private/RigVMCompiler/RigVMCompiler.cpp
@@ -2580,6 +2580,8 @@ const FRigVMCompilerWorkData::FRigVMASTProxyArray& URigVMCompiler::FindProxiesWi
 
 	PinProxiesToProcess.Add(InProxy);
 
+	const FString CPPType = InProxy.GetSubjectChecked<URigVMPin>()->GetCPPType();
+
 	for(int32 ProxyIndex = 0; ProxyIndex < PinProxiesToProcess.Num(); ProxyIndex++)
 	{
 		const FRigVMASTProxy& CurrentProxy = PinProxiesToProcess[ProxyIndex];
@@ -2595,6 +2597,12 @@ const FRigVMCompilerWorkData::FRigVMASTProxyArray& URigVMCompiler::FindProxiesWi
 						continue;
 					}
 				}
+				// due to LWC we may have two pins that don't
+				// actually share the same CPP type (float vs double)
+				if(Pin->GetCPPType() != CPPType)
+				{
+					continue;
+				}
 			}
 			PinProxies.Add(CurrentProxy);
 		}
