Control Rig: FIx unit test FRigVMCompilerVariableRemappingTest failing

diff --git a/Engine/Source/Developer/RigVMDeveloper/Private/RigVMCompiler/RigVMCompiler.cpp b/Engine/Source/Developer/RigVMDeveloper/Private/RigVMCompiler/RigVMCompiler.cpp
--- a/Engine/Source/Developer/RigVMDeveloper/Private/RigVMCompiler/RigVMCompiler.cpp
+++ b/Engine/Source/Developer/RigVMDeveloper/Private/RigVMCompiler/RigVMCompiler.cpp
@@ -2098,6 +2098,12 @@ FString URigVMCompiler::GetPinHash(const URigVMPin* InPin, const FRigVMVarExprAS
 		return FString::Printf(TEXT("%s%s%s"), *Prefix, *FullPath, *Suffix);
 	}
 
+	if (InVarExpr)
+	{
+		FString FullPath = InVarExpr->GetProxy().GetCallstack().GetCallPath(true);
+		return FString::Printf(TEXT("%s%s%s"), *Prefix, *FullPath, *Suffix);
+	}
+
 	FString PinPath = InPin->GetPinPath(bUseFullNodePath);
 	ensureMsgf(!PinPath.StartsWith(TEXT("FunctionLibrary::")), TEXT("A library path should never be part of a pin hash %s."), *PinPath);
 	return FString::Printf(TEXT("%s%s%s"), *Prefix, *PinPath, *Suffix);
