RigVM: split up RigVMTypeUtils into runtime and developer

diff --git a/Engine/Plugins/Animation/ControlRig/Source/ControlRig/Private/ControlRig.cpp b/Engine/Plugins/Animation/ControlRig/Source/ControlRig/Private/ControlRig.cpp
--- a/Engine/Plugins/Animation/ControlRig/Source/ControlRig/Private/ControlRig.cpp
+++ b/Engine/Plugins/Animation/ControlRig/Source/ControlRig/Private/ControlRig.cpp
@@ -24,6 +24,7 @@
 #include "Engine/Blueprint.h"
 #include "EdGraphSchema_K2.h"
 #include "Kismet2/BlueprintEditorUtils.h"
+#include "RigVMDeveloperTypeUtils.h"
 #endif// WITH_EDITOR
 
 #define LOCTEXT_NAMESPACE "ControlRig"
diff --git a/Engine/Source/Developer/RigVMDeveloper/Private/RigVMDeveloperTypeUtils.cpp b/Engine/Source/Developer/RigVMDeveloper/Private/RigVMDeveloperTypeUtils.cpp
new file mode 100644
--- /dev/null
+++ b/Engine/Source/Developer/RigVMDeveloper/Private/RigVMDeveloperTypeUtils.cpp
@@ -0,0 +1,689 @@
+﻿// Copyright Epic Games, Inc. All Rights Reserved.
+
+#include "RigVMDeveloperTypeUtils.h"
+#include "RigVMModel/RigVMVariableDescription.h"
+
+#if WITH_EDITOR
+
+#include "EdGraphSchema_K2.h"
+#include "Engine/Blueprint.h"
+
+FRigVMExternalVariable RigVMTypeUtils::ExternalVariableFromBPVariableDescription(
+	const FBPVariableDescription& InVariableDescription)
+{
+	const bool bIsPublic = !((InVariableDescription.PropertyFlags & CPF_DisableEditOnInstance) == CPF_DisableEditOnInstance);
+	const bool bIsReadOnly = ((InVariableDescription.PropertyFlags & CPF_BlueprintReadOnly) == CPF_BlueprintReadOnly);
+	return ExternalVariableFromPinType(InVariableDescription.VarName, InVariableDescription.VarType, bIsPublic, bIsReadOnly);
+}
+
+FRigVMExternalVariable RigVMTypeUtils::ExternalVariableFromRigVMVariableDescription(const FRigVMGraphVariableDescription& InVariableDescription)
+{	
+	FRigVMExternalVariable ExternalVariable;
+	ExternalVariable.Name = InVariableDescription.Name;
+
+	if (IsArrayType(InVariableDescription.CPPType))
+	{
+		ExternalVariable.bIsArray = true;
+		ExternalVariable.TypeName = *InVariableDescription.CPPType.Mid(7, InVariableDescription.CPPType.Len() - 8);
+		ExternalVariable.TypeObject = InVariableDescription.CPPTypeObject;
+	}
+	else
+	{
+		ExternalVariable.bIsArray = false;
+		ExternalVariable.TypeName = *InVariableDescription.CPPType;
+		ExternalVariable.TypeObject = InVariableDescription.CPPTypeObject;
+	}
+
+	ExternalVariable.bIsPublic = false;
+	ExternalVariable.bIsReadOnly = false;
+	ExternalVariable.Memory = nullptr;
+	return ExternalVariable;
+
+}
+
+FRigVMExternalVariable RigVMTypeUtils::ExternalVariableFromPinType(const FName& InName, const FEdGraphPinType& InPinType, bool bInPublic, bool bInReadonly)
+{
+	FRigVMExternalVariable ExternalVariable;
+	ExternalVariable.Name = InName;
+	ExternalVariable.bIsPublic = bInPublic;
+	ExternalVariable.bIsReadOnly = bInReadonly;
+
+	if (InPinType.ContainerType == EPinContainerType::None)
+	{
+		ExternalVariable.bIsArray = false;
+	}
+	else if (InPinType.ContainerType == EPinContainerType::Array)
+	{
+		ExternalVariable.bIsArray = true;
+	}
+	else
+	{
+		return FRigVMExternalVariable();
+	}
+
+	if (InPinType.PinCategory == UEdGraphSchema_K2::PC_Boolean)
+	{
+		ExternalVariable.TypeName = TEXT("bool");
+		ExternalVariable.Size = sizeof(bool);
+	}
+	else if (InPinType.PinCategory == UEdGraphSchema_K2::PC_Int)
+	{
+		ExternalVariable.TypeName = TEXT("int32");
+		ExternalVariable.Size = sizeof(int32);
+	}
+	else if (InPinType.PinCategory == UEdGraphSchema_K2::PC_Enum ||
+		InPinType.PinCategory == UEdGraphSchema_K2::PC_Byte)
+	{
+		if (UEnum* Enum = Cast<UEnum>(InPinType.PinSubCategoryObject))
+		{
+			ExternalVariable.TypeName = Enum->GetFName();
+			ExternalVariable.TypeObject = Enum;
+		}
+		else
+		{
+			ExternalVariable.TypeName = TEXT("uint8");
+		}
+		ExternalVariable.Size = sizeof(uint8);
+	}
+	else if (InPinType.PinCategory == UEdGraphSchema_K2::PC_Real)
+	{
+		if (InPinType.PinSubCategory == UEdGraphSchema_K2::PC_Float)
+		{
+			ExternalVariable.TypeName = TEXT("float");
+			ExternalVariable.Size = sizeof(float);
+		}
+		else if (InPinType.PinSubCategory == UEdGraphSchema_K2::PC_Double)
+		{
+			ExternalVariable.TypeName = TEXT("double");
+			ExternalVariable.Size = sizeof(double);
+		}
+		else
+		{
+			checkf(false, TEXT("Unexpected subcategory for PC_Real pin type."));
+		}
+	}
+	else if (InPinType.PinCategory == UEdGraphSchema_K2::PC_Name)
+	{
+		ExternalVariable.TypeName = TEXT("FName");
+		ExternalVariable.Size = sizeof(FName);
+	}
+	else if (InPinType.PinCategory == UEdGraphSchema_K2::PC_String)
+	{
+		ExternalVariable.TypeName = TEXT("FString");
+		ExternalVariable.Size = sizeof(FString);
+	}
+	else if (InPinType.PinCategory == UEdGraphSchema_K2::PC_Struct)
+	{
+		if (UScriptStruct* Struct = Cast<UScriptStruct>(InPinType.PinSubCategoryObject))
+		{
+			ExternalVariable.TypeName = *Struct->GetStructCPPName();
+			ExternalVariable.TypeObject = Struct;
+			ExternalVariable.Size = Struct->GetStructureSize();
+		}
+	}
+
+	return ExternalVariable;
+}
+
+FRigVMExternalVariable RigVMTypeUtils::ExternalVariableFromCPPTypePath(const FName& InName, const FString& InCPPTypePath, bool bInPublic, bool bInReadonly)
+{
+	FRigVMExternalVariable Variable;
+	if (InCPPTypePath.StartsWith(TEXT("TMap<")))
+	{
+		return Variable;
+	}
+	
+	Variable.Name = InName;
+	Variable.bIsPublic = bInPublic;
+	Variable.bIsReadOnly = bInReadonly;
+
+	FString CPPTypePath = InCPPTypePath;
+	Variable.bIsArray = RigVMTypeUtils::IsArrayType(CPPTypePath);
+	if (Variable.bIsArray)
+	{
+		CPPTypePath = BaseTypeFromArrayType(CPPTypePath);
+	}
+
+	if (CPPTypePath == TEXT("bool"))
+	{
+		Variable.TypeName = *CPPTypePath;
+		Variable.Size = sizeof(bool);
+	}
+	else if (CPPTypePath == TEXT("float"))
+	{
+		Variable.TypeName = *CPPTypePath;
+		Variable.Size = sizeof(float);
+	}
+	else if (CPPTypePath == TEXT("double"))
+	{
+		Variable.TypeName = *CPPTypePath;
+		Variable.Size = sizeof(double);
+	}
+	else if (CPPTypePath == TEXT("int32"))
+	{
+		Variable.TypeName = *CPPTypePath;
+		Variable.Size = sizeof(int32);
+	}
+	else if (CPPTypePath == TEXT("FString"))
+	{
+		Variable.TypeName = *CPPTypePath;
+		Variable.Size = sizeof(FString);
+	}
+	else if (CPPTypePath == TEXT("FName"))
+	{
+		Variable.TypeName = *CPPTypePath;
+		Variable.Size = sizeof(FName);
+	}
+	else if(UScriptStruct* ScriptStruct = URigVMPin::FindObjectFromCPPTypeObjectPath<UScriptStruct>(CPPTypePath))
+	{
+		Variable.TypeName = *ScriptStruct->GetStructCPPName();
+		Variable.TypeObject = ScriptStruct;
+		Variable.Size = ScriptStruct->GetStructureSize();
+	}
+	else if (UEnum* Enum= URigVMPin::FindObjectFromCPPTypeObjectPath<UEnum>(CPPTypePath))
+	{
+		Variable.TypeName = *Enum->CppType;
+		Variable.TypeObject = Enum;
+		Variable.Size = Enum->GetResourceSizeBytes(EResourceSizeMode::EstimatedTotal);
+	}
+	else
+	{
+		check(false);
+	}
+
+	return Variable;
+}
+
+FRigVMExternalVariable RigVMTypeUtils::ExternalVariableFromCPPType(const FName& InName, const FString& InCPPType, UObject* InCPPTypeObject, bool bInPublic, bool bInReadonly)
+{
+	FRigVMExternalVariable Variable;
+	if (InCPPType.StartsWith(TEXT("TMap<")))
+	{
+		return Variable;
+	}
+	
+	Variable.Name = InName;
+	Variable.bIsPublic = bInPublic;
+	Variable.bIsReadOnly = bInReadonly;
+
+	FString CPPType = InCPPType;
+	Variable.bIsArray = RigVMTypeUtils::IsArrayType(CPPType);
+	if (Variable.bIsArray)
+	{
+		CPPType = BaseTypeFromArrayType(CPPType);
+	}
+
+	if (CPPType == TEXT("bool"))
+	{
+		Variable.TypeName = *CPPType;
+		Variable.Size = sizeof(bool);
+	}
+	else if (CPPType == TEXT("float"))
+	{
+		Variable.TypeName = *CPPType;
+		Variable.Size = sizeof(float);
+	}
+	else if (CPPType == TEXT("double"))
+	{
+		Variable.TypeName = *CPPType;
+		Variable.Size = sizeof(double);
+	}
+	else if (CPPType == TEXT("int32"))
+	{
+		Variable.TypeName = *CPPType;
+		Variable.Size = sizeof(int32);
+	}
+	else if (CPPType == TEXT("FString"))
+	{
+		Variable.TypeName = *CPPType;
+		Variable.Size = sizeof(FString);
+	}
+	else if (CPPType == TEXT("FName"))
+	{
+		Variable.TypeName = *CPPType;
+		Variable.Size = sizeof(FName);
+	}
+	else if(UScriptStruct* ScriptStruct = URigVMPin::FindObjectFromCPPTypeObjectPath<UScriptStruct>(CPPType))
+	{
+		Variable.TypeName = *ScriptStruct->GetStructCPPName();
+		Variable.TypeObject = ScriptStruct;
+		Variable.Size = ScriptStruct->GetStructureSize();
+	}
+	else if (UEnum* Enum= URigVMPin::FindObjectFromCPPTypeObjectPath<UEnum>(CPPType))
+	{
+		Variable.TypeName = *Enum->CppType;
+		Variable.TypeObject = Enum;
+		Variable.Size = Enum->GetResourceSizeBytes(EResourceSizeMode::EstimatedTotal);
+	}
+	else
+	{
+		Variable.TypeName = *CPPType;
+		Variable.TypeObject = InCPPTypeObject;
+		Variable.Size = InCPPTypeObject->StaticClass()->GetStructureSize();
+	}
+
+	return Variable;
+}
+
+FEdGraphPinType RigVMTypeUtils::PinTypeFromCPPType(const FString& InCPPType, UObject* InCPPTypeObject)
+{
+	FEdGraphPinType PinType;
+	PinType.ResetToDefaults();
+	PinType.PinCategory = NAME_None;
+
+	FString BaseCPPType = InCPPType;
+	PinType.ContainerType = EPinContainerType::None;
+	if (RigVMTypeUtils::IsArrayType(InCPPType))
+	{
+		BaseCPPType = RigVMTypeUtils::BaseTypeFromArrayType(InCPPType);
+		PinType.ContainerType = EPinContainerType::Array;
+	}
+
+	if (BaseCPPType == TEXT("bool"))
+	{
+		PinType.PinCategory = UEdGraphSchema_K2::PC_Boolean;
+	}
+	else if (BaseCPPType == TEXT("int32"))
+	{
+		PinType.PinCategory = UEdGraphSchema_K2::PC_Int;
+	}
+	else if (BaseCPPType == TEXT("float"))
+	{
+		PinType.PinCategory = UEdGraphSchema_K2::PC_Real;
+		PinType.PinSubCategory = UEdGraphSchema_K2::PC_Float;
+	}
+	else if (BaseCPPType == TEXT("double"))
+	{
+		PinType.PinCategory = UEdGraphSchema_K2::PC_Real;
+		PinType.PinSubCategory = UEdGraphSchema_K2::PC_Double;
+	}
+	else if (BaseCPPType == TEXT("FName"))
+	{
+		PinType.PinCategory = UEdGraphSchema_K2::PC_Name;
+	}
+	else if (BaseCPPType == TEXT("FString"))
+	{
+		PinType.PinCategory = UEdGraphSchema_K2::PC_String;
+	}
+	else if (Cast<UScriptStruct>(InCPPTypeObject))
+	{
+		PinType.PinCategory = UEdGraphSchema_K2::PC_Struct;
+		PinType.PinSubCategoryObject = InCPPTypeObject;
+	}
+	else if (Cast<UEnum>(InCPPTypeObject))
+	{
+		PinType.PinCategory = UEdGraphSchema_K2::PC_Byte;
+		PinType.PinSubCategoryObject = InCPPTypeObject;
+	}
+
+	return PinType;
+}
+
+FEdGraphPinType RigVMTypeUtils::PinTypeFromExternalVariable(const FRigVMExternalVariable& InExternalVariable)
+{
+	FEdGraphPinType PinType;
+	PinType.ResetToDefaults();
+	PinType.PinCategory = NAME_None;
+
+	if (InExternalVariable.TypeName == TEXT("bool"))
+	{
+		PinType.PinCategory = UEdGraphSchema_K2::PC_Boolean;
+	}
+	else if (InExternalVariable.TypeName == TEXT("int32"))
+	{
+		PinType.PinCategory = UEdGraphSchema_K2::PC_Int;
+	}
+	else if (InExternalVariable.TypeName == TEXT("float"))
+	{
+		PinType.PinCategory = UEdGraphSchema_K2::PC_Real;
+		PinType.PinSubCategory = UEdGraphSchema_K2::PC_Float;
+	}
+	else if (InExternalVariable.TypeName == TEXT("double"))
+	{
+		PinType.PinCategory = UEdGraphSchema_K2::PC_Real;
+		PinType.PinSubCategory = UEdGraphSchema_K2::PC_Double;
+	}
+	else if (InExternalVariable.TypeName == TEXT("FName"))
+	{
+		PinType.PinCategory = UEdGraphSchema_K2::PC_Name;
+	}
+	else if (InExternalVariable.TypeName == TEXT("FString"))
+	{
+		PinType.PinCategory = UEdGraphSchema_K2::PC_String;
+	}
+	else if (Cast<UScriptStruct>(InExternalVariable.TypeObject))
+	{
+		PinType.PinCategory = UEdGraphSchema_K2::PC_Struct;
+		PinType.PinSubCategoryObject = InExternalVariable.TypeObject;
+	}
+	else if (Cast<UEnum>(InExternalVariable.TypeObject))
+	{
+		PinType.PinCategory = UEdGraphSchema_K2::PC_Byte;
+		PinType.PinSubCategoryObject = InExternalVariable.TypeObject;
+	}
+
+	if (InExternalVariable.bIsArray)
+	{
+		PinType.ContainerType = EPinContainerType::Array;
+	}
+	else
+	{
+		PinType.ContainerType = EPinContainerType::None;
+	}
+
+	return PinType;
+}
+
+FEdGraphPinType RigVMTypeUtils::PinTypeFromRigVMVariableDescription(
+	const FRigVMGraphVariableDescription& InVariableDescription)
+{
+	FEdGraphPinType PinType;
+	PinType.ResetToDefaults();
+	PinType.PinCategory = NAME_None;
+
+	FString CurrentCPPType = InVariableDescription.CPPType;
+	if (IsArrayType(CurrentCPPType))
+	{
+		PinType.ContainerType = EPinContainerType::Array;
+		CurrentCPPType = BaseTypeFromArrayType(CurrentCPPType);
+	}
+	else
+	{
+		PinType.ContainerType = EPinContainerType::None;
+	}
+
+	if (CurrentCPPType == TEXT("bool"))
+	{
+		PinType.PinCategory = UEdGraphSchema_K2::PC_Boolean;
+	}
+	else if (CurrentCPPType == TEXT("int32"))
+	{
+		PinType.PinCategory = UEdGraphSchema_K2::PC_Int;
+	}
+	else if (CurrentCPPType == TEXT("float"))
+	{
+		PinType.PinCategory = UEdGraphSchema_K2::PC_Real;
+		PinType.PinSubCategory = UEdGraphSchema_K2::PC_Float;
+	}
+	else if (CurrentCPPType == TEXT("double"))
+	{
+		PinType.PinCategory = UEdGraphSchema_K2::PC_Real;
+		PinType.PinSubCategory = UEdGraphSchema_K2::PC_Double;
+	}
+	else if (CurrentCPPType == TEXT("FName"))
+	{
+		PinType.PinCategory = UEdGraphSchema_K2::PC_Name;
+	}
+	else if (CurrentCPPType == TEXT("FString"))
+	{
+		PinType.PinCategory = UEdGraphSchema_K2::PC_String;
+	}
+	else if (Cast<UScriptStruct>(InVariableDescription.CPPTypeObject))
+	{
+		PinType.PinCategory = UEdGraphSchema_K2::PC_Struct;
+		PinType.PinSubCategoryObject = InVariableDescription.CPPTypeObject;
+	}
+	else if (Cast<UEnum>(InVariableDescription.CPPTypeObject))
+	{
+		PinType.PinCategory = UEdGraphSchema_K2::PC_Byte;
+		PinType.PinSubCategoryObject = InVariableDescription.CPPTypeObject;
+	}
+	else if (Cast<UClass>(InVariableDescription.CPPTypeObject))
+	{
+		PinType.PinCategory = UEdGraphSchema_K2::PC_Object;
+		PinType.PinSubCategoryObject = InVariableDescription.CPPTypeObject;
+	}
+
+	return PinType;
+}
+
+FEdGraphPinType RigVMTypeUtils::SubPinType(const FEdGraphPinType& InPinType, const FString& InSegmentPath)
+{
+	FEdGraphPinType Result;
+	if (InSegmentPath.IsEmpty())
+	{
+		return InPinType;
+	}
+
+	if (InPinType.PinSubCategoryObject.IsValid())
+	{
+		if (UStruct* Struct = Cast<UStruct>(InPinType.PinSubCategoryObject))
+		{
+			int32 PartIndex = 0; 
+			TArray<FString> Parts;
+			if (!URigVMPin::SplitPinPath(InSegmentPath, Parts))
+			{
+				Parts.Add(InSegmentPath);
+			}
+			
+			FProperty* Property = Struct->FindPropertyByName(*Parts[PartIndex++]);
+			while (PartIndex < Parts.Num() && Property != nullptr)
+			{
+				if (FArrayProperty* ArrayProperty = CastField<FArrayProperty>(Property))
+				{
+					Property = ArrayProperty->Inner;
+					PartIndex++;
+					continue;
+				}
+
+				if (FStructProperty* StructProperty = CastField<FStructProperty>(Property))
+				{
+					Struct = StructProperty->Struct;
+					Property = Struct->FindPropertyByName(*Parts[PartIndex++]);
+					continue;
+				}
+
+				break;
+			}
+
+			if (PartIndex == Parts.Num() && Property)
+			{
+				UEdGraphSchema_K2::StaticClass()->GetDefaultObject<UEdGraphSchema_K2>()->ConvertPropertyToPinType(Property, Result);
+			}
+		}
+	}
+
+	return Result;
+}
+
+bool RigVMTypeUtils::CPPTypeFromPinType(const FEdGraphPinType& InPinType, FString& OutCPPType, UObject** OutCPPTypeObject)
+{
+	FString Prefix = "";
+	FString Suffix = "";
+	if (InPinType.ContainerType == EPinContainerType::Array)
+	{
+		Prefix = TEXT("TArray<");
+		Suffix = TEXT(">");
+	}
+
+	OutCPPType = FString();
+	*OutCPPTypeObject = nullptr;
+	if (InPinType.PinCategory == UEdGraphSchema_K2::PC_Boolean)
+	{
+		OutCPPType = Prefix + TEXT("bool") + Suffix;
+	}
+	else if (InPinType.PinCategory == UEdGraphSchema_K2::PC_Int)
+	{
+		OutCPPType = Prefix + TEXT("int32") + Suffix;
+	}
+	else if (InPinType.PinCategory == UEdGraphSchema_K2::PC_Real)
+	{
+		if (InPinType.PinSubCategory == UEdGraphSchema_K2::PC_Float)
+		{
+			OutCPPType = Prefix + TEXT("float") + Suffix;
+		}
+		else if (InPinType.PinSubCategory == UEdGraphSchema_K2::PC_Double)
+		{
+			OutCPPType = Prefix + TEXT("double") + Suffix;
+		}
+		else
+		{
+			checkf(false, TEXT("Unexpected subcategory for PC_Real pin type."));
+		}
+	}
+	else if (InPinType.PinCategory == UEdGraphSchema_K2::PC_Float)
+	{
+		OutCPPType = Prefix + TEXT("float") + Suffix;
+	}
+	else if (InPinType.PinCategory == UEdGraphSchema_K2::PC_Double)
+	{
+		OutCPPType = Prefix + TEXT("double") + Suffix;
+	}
+	else if (InPinType.PinCategory == UEdGraphSchema_K2::PC_Name)
+	{
+		OutCPPType = Prefix + TEXT("FName") + Suffix;
+	}
+	else if (InPinType.PinCategory == UEdGraphSchema_K2::PC_String)
+	{
+		OutCPPType = Prefix + TEXT("FString") + Suffix;
+	}
+	else if (InPinType.PinCategory == UEdGraphSchema_K2::PC_Struct)
+	{
+		if (UScriptStruct* Struct = Cast<UScriptStruct>(InPinType.PinSubCategoryObject))
+		{
+			OutCPPType = Prefix + *Struct->GetStructCPPName() + Suffix;
+			*OutCPPTypeObject = Struct;
+		}
+	}
+	else if (InPinType.PinCategory == UEdGraphSchema_K2::PC_Object ||
+		InPinType.PinCategory == UEdGraphSchema_K2::PC_SoftObject ||
+		InPinType.PinCategory == UEdGraphSchema_K2::AllObjectTypes)
+	{
+		if (UClass* Class = Cast<UClass>(InPinType.PinSubCategoryObject))
+		{
+			OutCPPType = Prefix + FString::Printf(TEXT("TObjectPtr<%s%s>"), Class->GetPrefixCPP(), *Class->GetName()) + Suffix;
+			*OutCPPTypeObject = Class;
+		}
+	}
+	else if (InPinType.PinCategory == UEdGraphSchema_K2::PC_Byte ||
+		InPinType.PinCategory == UEdGraphSchema_K2::PC_Enum)
+	{
+		if (UEnum* Enum = Cast<UEnum>(InPinType.PinSubCategoryObject))
+		{
+			OutCPPType = Prefix + Enum->GetFName().ToString() + Suffix;
+			*OutCPPTypeObject = Enum;
+		}
+		else
+		{
+			OutCPPType = Prefix + TEXT("uint8") + Suffix;
+		}
+	}
+	else
+	{
+		return false;
+	}
+		
+	return true;
+}
+
+bool RigVMTypeUtils::CPPTypeFromPinType(const FEdGraphPinType& InPinType, FString& OutCPPType, FName& OutCPPTypeObjectPath)
+{
+	OutCPPType = FString();
+	OutCPPTypeObjectPath = NAME_None;
+	UObject* CPPTypeObject = nullptr;
+	if (RigVMTypeUtils::CPPTypeFromPinType(InPinType, OutCPPType, &CPPTypeObject))
+	{
+		if (CPPTypeObject)
+		{
+			OutCPPTypeObjectPath = *CPPTypeObject->GetPathName();
+			if (UScriptStruct* ScriptStruct = Cast<UScriptStruct>(CPPTypeObject))
+			{
+				FString StructName = ScriptStruct->GetStructCPPName();
+				while (IsArrayType(OutCPPType))
+				{
+					OutCPPType = BaseTypeFromArrayType(OutCPPType);
+					StructName = ArrayTypeFromBaseType(StructName);
+				}
+				OutCPPType = StructName;
+			}
+		}
+
+		return true;
+	}
+
+	return false;
+}
+
+bool RigVMTypeUtils::CPPTypeFromExternalVariable(const FRigVMExternalVariable& InExternalVariable, FString& OutCPPType, UObject** OutCPPTypeObject)
+{
+	FString Prefix = "";
+	FString Suffix = "";
+	if (InExternalVariable.bIsArray)
+	{
+		Prefix = TEXT("TArray<");
+		Suffix = TEXT(">");
+	}
+
+	*OutCPPTypeObject = nullptr;
+	if (InExternalVariable.TypeName == TEXT("bool"))
+	{
+		OutCPPType = Prefix + TEXT("bool") + Suffix;
+	}
+	else if (InExternalVariable.TypeName == TEXT("int32"))
+	{
+		OutCPPType = Prefix + TEXT("int32") + Suffix;
+	}
+	else if (InExternalVariable.TypeName == TEXT("float"))
+	{
+		OutCPPType = Prefix + TEXT("float") + Suffix;
+	}
+	else if (InExternalVariable.TypeName == TEXT("double"))
+	{
+		OutCPPType = Prefix + TEXT("double") + Suffix;
+	}
+	else if (InExternalVariable.TypeName == TEXT("FName"))
+	{
+		OutCPPType = Prefix + TEXT("FName") + Suffix;
+	}
+	else if (InExternalVariable.TypeName == TEXT("FString"))
+	{
+		OutCPPType = Prefix + TEXT("FString") + Suffix;
+	}
+	else if (UScriptStruct* Struct = Cast<UScriptStruct>(InExternalVariable.TypeObject))
+	{
+		OutCPPType = Prefix + *Struct->GetStructCPPName() + Suffix;
+		*OutCPPTypeObject = Struct;	
+	}
+	else if (UEnum* Enum = Cast<UEnum>(InExternalVariable.TypeObject))
+	{
+		OutCPPType = Prefix + Enum->GetFName().ToString() + Suffix;
+		*OutCPPTypeObject = Enum;
+	}	
+	else
+	{
+		check(false);
+		return false;
+	}
+		
+	return true;
+}
+
+bool RigVMTypeUtils::AreCompatible(const FString& InCPPTypeA, UObject* InCPPTypeObjectA, const FString& InCPPTypeB,	UObject* InCPPTypeObjectB)
+{
+	return AreCompatible(PinTypeFromCPPType(InCPPTypeA, InCPPTypeObjectA), PinTypeFromCPPType(InCPPTypeB, InCPPTypeObjectB));
+}
+
+bool RigVMTypeUtils::AreCompatible(const FRigVMExternalVariable& InTypeA, const FRigVMExternalVariable& InTypeB, const FString& InSegmentPathA, const FString& InSegmentPathB)
+{
+	return AreCompatible(PinTypeFromExternalVariable(InTypeA), PinTypeFromExternalVariable(InTypeB), InSegmentPathA, InSegmentPathB);
+}
+
+bool RigVMTypeUtils::AreCompatible(const FEdGraphPinType& InTypeA, const FEdGraphPinType& InTypeB, const FString& InSegmentPathA, const FString& InSegmentPathB)
+{
+	FEdGraphPinType SubPinTypeA = SubPinType(InTypeA, InSegmentPathA);
+	FEdGraphPinType SubPinTypeB = SubPinType(InTypeB, InSegmentPathB);
+
+	// We allow connectiongs between floats and doubles, while EdGraphSchema_K2 does not
+	// Every other case is evaluated by UEdGraphSchema_K2::ArePinTypesCompatible
+	if (SubPinTypeA.ContainerType == SubPinTypeB.ContainerType)
+	{
+		if ((SubPinTypeA.PinCategory == TEXT("float") && SubPinTypeB.PinCategory == TEXT("double")) ||
+			(SubPinTypeA.PinCategory == TEXT("double") && SubPinTypeB.PinCategory == TEXT("float")))
+		{
+			return true;
+		}
+	}
+	
+	return UEdGraphSchema_K2::StaticClass()->GetDefaultObject<UEdGraphSchema_K2>()->ArePinTypesCompatible(SubPinTypeA, SubPinTypeB);	
+}
+#endif
diff --git a/Engine/Source/Developer/RigVMDeveloper/Private/RigVMTypeUtils.cpp b/Engine/Source/Developer/RigVMDeveloper/Private/RigVMTypeUtils.cpp
deleted file mode 100644
diff --git a/Engine/Source/Developer/RigVMDeveloper/Public/RigVMDeveloperTypeUtils.h b/Engine/Source/Developer/RigVMDeveloper/Public/RigVMDeveloperTypeUtils.h
new file mode 100644
--- /dev/null
+++ b/Engine/Source/Developer/RigVMDeveloper/Public/RigVMDeveloperTypeUtils.h
@@ -0,0 +1,46 @@
+﻿// Copyright Epic Games, Inc. All Rights Reserved.
+
+#pragma once
+
+#include "RigVMDeveloperModule.h"
+#include "RigVMCore/RigVMExternalVariable.h"
+#include "EdGraphSchema_K2.h"
+#include "RigVMTypeUtils.h"
+
+struct FRigVMGraphVariableDescription;
+
+namespace RigVMTypeUtils
+{
+#if WITH_EDITOR
+	RIGVMDEVELOPER_API FRigVMExternalVariable ExternalVariableFromBPVariableDescription(const FBPVariableDescription& InVariableDescription);
+
+	RIGVMDEVELOPER_API FRigVMExternalVariable ExternalVariableFromRigVMVariableDescription(const FRigVMGraphVariableDescription& InVariableDescription);
+	
+	RIGVMDEVELOPER_API FRigVMExternalVariable ExternalVariableFromPinType(const FName& InName, const FEdGraphPinType& InPinType, bool bInPublic = false, bool bInReadonly = false);
+
+	RIGVMDEVELOPER_API FRigVMExternalVariable ExternalVariableFromCPPTypePath(const FName& InName, const FString& InCPPTypePath, bool bInPublic = false, bool bInReadonly = false);
+
+	RIGVMDEVELOPER_API FRigVMExternalVariable ExternalVariableFromCPPType(const FName& InName, const FString& InCPPType, UObject* InCPPTypeObject, bool bInPublic = false, bool bInReadonly = false);
+
+	RIGVMDEVELOPER_API FEdGraphPinType PinTypeFromCPPType(const FString& InCPPType, UObject* InCPPTypeObject);
+
+	RIGVMDEVELOPER_API FEdGraphPinType PinTypeFromExternalVariable(const FRigVMExternalVariable& InExternalVariable);
+
+	RIGVMDEVELOPER_API FEdGraphPinType PinTypeFromRigVMVariableDescription(const FRigVMGraphVariableDescription& InVariableDescription);
+
+	RIGVMDEVELOPER_API FEdGraphPinType SubPinType(const FEdGraphPinType& InPinType, const FString& SegmentPath);
+
+	RIGVMDEVELOPER_API bool CPPTypeFromPinType(const FEdGraphPinType& InPinType, FString& OutCPPType, UObject** OutCPPTypeObject);
+
+	RIGVMDEVELOPER_API bool CPPTypeFromPinType(const FEdGraphPinType& InPinType, FString& OutCPPType, FName& OutCPPTypeObjectPath);
+
+	RIGVMDEVELOPER_API bool CPPTypeFromExternalVariable(const FRigVMExternalVariable& InExternalVariable, FString& OutCPPType, UObject** OutCPPTypeObject);
+
+	RIGVMDEVELOPER_API bool AreCompatible(const FString& InCPPTypeA, UObject* InCPPTypeObjectA, const FString& InCPPTypeB, UObject* InCPPTypeObjectB);
+
+	RIGVMDEVELOPER_API bool AreCompatible(const FRigVMExternalVariable& InTypeA, const FRigVMExternalVariable& InTypeB, const FString& InSegmentPathA = FString(), const FString& InSegmentPathB = FString());
+	
+	RIGVMDEVELOPER_API bool AreCompatible(const FEdGraphPinType& InTypeA, const FEdGraphPinType& InTypeB, const FString& InSegmentPathA = FString(), const FString& InSegmentPathB = FString());
+	
+#endif
+}
diff --git a/Engine/Source/Developer/RigVMDeveloper/Public/RigVMModel/RigVMVariableDescription.h b/Engine/Source/Developer/RigVMDeveloper/Public/RigVMModel/RigVMVariableDescription.h
--- a/Engine/Source/Developer/RigVMDeveloper/Public/RigVMModel/RigVMVariableDescription.h
+++ b/Engine/Source/Developer/RigVMDeveloper/Public/RigVMModel/RigVMVariableDescription.h
@@ -5,7 +5,7 @@
 #include "RigVMCore/RigVMExternalVariable.h"
 #include "RigVMModel/RigVMNode.h"
 #include "EdGraphSchema_K2.h"
-#include "RigVMTypeUtils.h"
+#include "RigVMDeveloperTypeUtils.h"
 #include "RigVMVariableDescription.generated.h"
 
 /**
diff --git a/Engine/Source/Developer/RigVMDeveloper/Public/RigVMTypeUtils.h b/Engine/Source/Developer/RigVMDeveloper/Public/RigVMTypeUtils.h
deleted file mode 100644
diff --git a/Engine/Source/Runtime/RigVM/Public/RigVMTypeUtils.h b/Engine/Source/Runtime/RigVM/Public/RigVMTypeUtils.h
new file mode 100644
--- /dev/null
+++ b/Engine/Source/Runtime/RigVM/Public/RigVMTypeUtils.h
@@ -0,0 +1,32 @@
+#pragma once
+
+#include "CoreMinimal.h"
+
+namespace RigVMTypeUtils
+{
+	const TCHAR TArrayPrefix[] = TEXT("TArray<");
+	const TCHAR TObjectPtrPrefix[] = TEXT("TObjectPtr<");
+	const TCHAR TArrayTemplate[] = TEXT("TArray<%s>");
+	const TCHAR TObjectPtrTemplate[] = TEXT("TObjectPtr<%s%s>");
+	
+	// Returns true if the type specified is an array
+	FORCEINLINE bool IsArrayType(const FString& InCPPType)
+	{
+		return InCPPType.StartsWith(TArrayPrefix);
+	}
+
+	FORCEINLINE FString ArrayTypeFromBaseType(const FString& InCPPType)
+	{
+		return FString::Printf(TArrayTemplate, *InCPPType);
+	}
+
+	FORCEINLINE FString BaseTypeFromArrayType(const FString& InCPPType)
+	{
+		return InCPPType.RightChop(7).LeftChop(1);
+	}
+
+	FORCEINLINE bool IsUObjectType(const FString& InCPPType)
+	{
+		return InCPPType.StartsWith(TObjectPtrPrefix);
+	}
+}
\ No newline at end of file
