Control Rig: Fix pytlhon command for changing exposed pin type incorrect

diff --git a/Engine/Source/Developer/RigVMDeveloper/Private/RigVMModel/RigVMController.cpp b/Engine/Source/Developer/RigVMDeveloper/Private/RigVMModel/RigVMController.cpp
--- a/Engine/Source/Developer/RigVMDeveloper/Private/RigVMModel/RigVMController.cpp
+++ b/Engine/Source/Developer/RigVMDeveloper/Private/RigVMModel/RigVMController.cpp
@@ -8597,11 +8597,12 @@ bool URigVMController::ChangeExposedPinType(const FName& InPinName, const FStrin
 		const FString GraphName = GetSanitizedGraphName(GetGraph()->GetGraphName());
 		
 		RigVMPythonUtils::Print(GetGraphOuterName(), 
-			FString::Printf(TEXT("blueprint.get_controller_by_name('%s').change_exposed_pin_type('%s', '%s', '%s')"),
+			FString::Printf(TEXT("blueprint.get_controller_by_name('%s').change_exposed_pin_type('%s', '%s', '%s', %s)"),
 				*GraphName,
 				*GetSanitizedPinName(InPinName.ToString()),
 				*InCPPType,
-				*InCPPTypeObjectPath.ToString()));
+				*InCPPTypeObjectPath.ToString(),
+				(bSetupUndoRedo) ? TEXT("True") : TEXT("False")));
 	}
 
 	return true;
