Control Rig: Template Node Support.

The UE_RIGVM_ENABLE_TEMPLATE_NODES is disabled in this submit for now - which disabled the creation of template nodes as well, disabled auto-update functionality during load and spawning resolved templates nodes instead of unit nodes.

We'll enable the define in a follow up submit

diff --git a/Engine/Plugins/Animation/ControlRig/Source/ControlRigDeveloper/Private/ControlRigBlueprint.cpp b/Engine/Plugins/Animation/ControlRig/Source/ControlRigDeveloper/Private/ControlRigBlueprint.cpp
--- a/Engine/Plugins/Animation/ControlRig/Source/ControlRigDeveloper/Private/ControlRigBlueprint.cpp
+++ b/Engine/Plugins/Animation/ControlRig/Source/ControlRigDeveloper/Private/ControlRigBlueprint.cpp
@@ -461,6 +461,9 @@ void UControlRigBlueprint::PostLoad()
 		PatchVariableNodesWithIncorrectType();
 		PatchRigElementKeyCacheOnLoad();
 		PatchBoundVariables();
+#if UE_RIGVM_ENABLE_TEMPLATE_NODES
+		ConvertUnitNodesToTemplateNodes();
+#endif
 
 #if WITH_EDITOR
 
@@ -2303,7 +2306,7 @@ bool UControlRigBlueprint::ChangeMemberVariableType(const FName& InName, const F
 	}
 	else if (UEnum* Enum= URigVMPin::FindObjectFromCPPTypeObjectPath<UEnum>(CPPType))
 	{
-		Variable.TypeName = *Enum->CppType;
+		Variable.TypeName = *RigVMTypeUtils::CPPTypeFromEnum(Enum);
 		Variable.TypeObject = Enum;
 		Variable.Size = Enum->GetResourceSizeBytes(EResourceSizeMode::EstimatedTotal);
 	}
@@ -3604,9 +3607,10 @@ void UControlRigBlueprint::PatchRigElementKeyCacheOnLoad()
 				if (URigVMUnitNode* UnitNode = Cast<URigVMUnitNode>(Node))
 				{
 					UScriptStruct* ScriptStruct = UnitNode->GetScriptStruct();
-					FString FunctionName = FString::Printf(TEXT("F%s::%s"), *ScriptStruct->GetName(), *UnitNode->GetMethodName().ToString());
-					FRigVMFunction Function = FRigVMRegistry::Get().FindFunctionInfo(*FunctionName);
-					for (TFieldIterator<FProperty> It(Function.Struct); It; ++It)
+					FString FunctionName = FString::Printf(TEXT("%s::%s"), *ScriptStruct->GetStructCPPName(), *UnitNode->GetMethodName().ToString());
+					const FRigVMFunction* Function = FRigVMRegistry::Get().FindFunction(*FunctionName);
+					check(Function);
+					for (TFieldIterator<FProperty> It(Function->Struct); It; ++It)
 					{
 						if (It->GetCPPType() == TEXT("FCachedRigElement"))
 						{
@@ -3734,6 +3738,21 @@ void UControlRigBlueprint::PatchVariableNodesWithIncorrectType()
 	}
 }
 
+void UControlRigBlueprint::ConvertUnitNodesToTemplateNodes()
+{
+	TGuardValue<bool> GuardNotifsSelf(bSuspendModelNotificationsForSelf, true);
+
+	for (URigVMGraph* Graph : GetAllModels())
+	{
+		URigVMController* Controller = GetOrCreateController(Graph);
+		TArray<URigVMNode*> Nodes = Graph->GetNodes();
+		for (URigVMNode* Node : Nodes)
+		{
+			Controller->ReplaceUnitNodeWithTemplateNode(Node->GetFName(), false);
+		}
+	}
+}
+
 void UControlRigBlueprint::PropagatePoseFromInstanceToBP(UControlRig* InControlRig)
 {
 	check(InControlRig);
diff --git a/Engine/Plugins/Animation/ControlRig/Source/ControlRigDeveloper/Private/Graph/ControlRigGraph.cpp b/Engine/Plugins/Animation/ControlRig/Source/ControlRigDeveloper/Private/Graph/ControlRigGraph.cpp
--- a/Engine/Plugins/Animation/ControlRig/Source/ControlRigDeveloper/Private/Graph/ControlRigGraph.cpp
+++ b/Engine/Plugins/Animation/ControlRig/Source/ControlRigDeveloper/Private/Graph/ControlRigGraph.cpp
@@ -226,6 +226,30 @@ void UControlRigGraph::HandleModifiedEvent(ERigVMGraphNotifType InNotifType, URi
 		ControlRigSchema->HandleModifiedEvent(InNotifType, InGraph, InSubject);
 	}
 
+	if(TemplateController)
+	{
+		switch(InNotifType)
+		{
+			case ERigVMGraphNotifType::NodeRemoved:
+			case ERigVMGraphNotifType::PinTypeChanged:
+			case ERigVMGraphNotifType::PinDefaultValueChanged:
+			case ERigVMGraphNotifType::PinRemoved:
+			case ERigVMGraphNotifType::PinRenamed:
+			case ERigVMGraphNotifType::VariableRemoved:
+			case ERigVMGraphNotifType::VariableRenamed:
+			case ERigVMGraphNotifType::ParameterRemoved:
+			case ERigVMGraphNotifType::ParameterRenamed:
+			case ERigVMGraphNotifType::GraphChanged:
+			{
+				return;
+			}
+			default:
+			{
+				break;
+			}
+		}
+	}
+
 	// increment the node topology version for any interaction
 	// with a node.
 	{
@@ -396,30 +420,7 @@ void UControlRigGraph::HandleModifiedEvent(ERigVMGraphNotifType InNotifType, URi
 				UEdGraphNode* EdNode = FindNodeForModelNodeName(ModelNode->GetFName(), false);
 				if (EdNode)
 				{
-					static UObject* NewOuter = GetTransientPackage();
-					
-					// Rename the soon to be deleted object to a unique name, so that other objects can use
-					// the old name
-					{
-						FString DeletedName;
-						UObject* ExistingObject = nullptr;
-						int32 Index = FMath::Rand();
-						do
-						{
-							DeletedName = FString::Printf(TEXT("%s_Deleted_%d"), *EdNode->GetName(), Index++); 
-							ExistingObject = StaticFindObject(/*Class=*/ NULL, NewOuter, *DeletedName, true);						
-						}
-						while (ExistingObject);
-						EdNode->Rename(*DeletedName, NewOuter, REN_ForceNoResetLoaders | REN_DoNotDirty | REN_DontCreateRedirectors | REN_NonTransactional);
-					}
-
-					// Remove the node and notify of the removal
-					{
-						Nodes.Remove(EdNode);
-						EdNode->BreakAllNodeLinks();
-						
-						NotifyGraphChanged();
-					}
+					RemoveAndDeleteNode(EdNode);
 				}
 			}
 			break;
@@ -452,6 +453,19 @@ void UControlRigGraph::HandleModifiedEvent(ERigVMGraphNotifType InNotifType, URi
 			}
 			break;
 		}
+		case ERigVMGraphNotifType::NodeDescriptionChanged:
+		case ERigVMGraphNotifType::NodeCategoryChanged:
+		{
+			if (URigVMNode* ModelNode = Cast<URigVMNode>(InSubject))
+			{
+				if (UControlRigGraphNode* RigNode = Cast<UControlRigGraphNode>(FindNodeForModelNodeName(ModelNode->GetFName())))
+				{
+					RigNode->InvalidateNodeTitle();
+					RigNode->ReconstructNode_Internal(true);
+				}
+			}
+			break;
+		}
 		case ERigVMGraphNotifType::RerouteCompactnessChanged:
 		{
 			if (URigVMRerouteNode* ModelNode = Cast<URigVMRerouteNode>(InSubject))
@@ -478,7 +492,7 @@ void UControlRigGraph::HandleModifiedEvent(ERigVMGraphNotifType InNotifType, URi
 		{
 			if (URigVMNode* ModelNode = Cast<URigVMNode>(InSubject))
 			{
-				if (ModelNode->IsA<URigVMLibraryNode>())
+				if (ModelNode->IsA<URigVMLibraryNode>() || ModelNode->IsA<URigVMTemplateNode>())
 				{
 					if (UControlRigGraphNode* RigNode = Cast<UControlRigGraphNode>(FindNodeForModelNodeName(ModelNode->GetFName())))
 					{
@@ -910,6 +924,34 @@ URigVMController* UControlRigGraph::GetController() const
 	return nullptr;
 }
 
+void UControlRigGraph::RemoveAndDeleteNode(UEdGraphNode* InNode)
+{
+	static UObject* NewOuter = GetTransientPackage();
+					
+	// Rename the soon to be deleted object to a unique name, so that other objects can use
+	// the old name
+	{
+		FString DeletedName;
+		UObject* ExistingObject = nullptr;
+		static int32 DeletedIndex = FMath::Rand();
+		do
+		{
+			DeletedName = FString::Printf(TEXT("%s_Deleted_%d"), *InNode->GetName(), DeletedIndex++); 
+			ExistingObject = StaticFindObject(/*Class=*/ NULL, NewOuter, *DeletedName, true);						
+		}
+		while (ExistingObject);
+		InNode->Rename(*DeletedName, NewOuter, REN_ForceNoResetLoaders | REN_DoNotDirty | REN_DontCreateRedirectors | REN_NonTransactional);
+	}
+
+	// Remove the node and notify of the removal
+	{
+		Nodes.Remove(InNode);
+		InNode->BreakAllNodeLinks();
+						
+		NotifyGraphChanged();
+	}
+}
+
 URigVMController* UControlRigGraph::GetTemplateController()
 {
 	if (TemplateController == nullptr)
diff --git a/Engine/Plugins/Animation/ControlRig/Source/ControlRigDeveloper/Private/Graph/ControlRigGraphNode.cpp b/Engine/Plugins/Animation/ControlRig/Source/ControlRigDeveloper/Private/Graph/ControlRigGraphNode.cpp
--- a/Engine/Plugins/Animation/ControlRig/Source/ControlRigDeveloper/Private/Graph/ControlRigGraphNode.cpp
+++ b/Engine/Plugins/Animation/ControlRig/Source/ControlRigDeveloper/Private/Graph/ControlRigGraphNode.cpp
@@ -46,6 +46,7 @@ UControlRigGraphNode::UControlRigGraphNode()
 #if WITH_EDITOR
 , bEnableProfiling(false)
 #endif
+, CachedTemplate(nullptr)
 {
 	bHasCompilerMessage = false;
 	ErrorType = (int32)EMessageSeverity::Info + 1;
@@ -424,6 +425,22 @@ int32 UControlRigGraphNode::GetInstructionIndex(bool bAsInput) const
 	return INDEX_NONE;
 }
 
+const FRigVMTemplate* UControlRigGraphNode::GetTemplate() const
+{
+	if(CachedTemplate == nullptr)
+	{
+		if(URigVMTemplateNode* TemplateNode = Cast<URigVMTemplateNode>(GetModelNode()))
+		{
+			CachedTemplate = TemplateNode->GetTemplate();
+		}
+		else if(ModelNodePath.Contains(TEXT("::Execute(")))
+		{
+			CachedTemplate = FRigVMRegistry::Get().FindTemplate(*ModelNodePath); 
+		}
+	}
+	return CachedTemplate;
+}
+
 FLinearColor UControlRigGraphNode::GetNodeProfilingColor() const
 {
 #if WITH_EDITOR
@@ -768,6 +785,7 @@ bool UControlRigGraphNode::ShowPaletteIconOnNode() const
 			ModelNode->IsA<URigVMFunctionReturnNode>() ||
 			ModelNode->IsA<URigVMFunctionReferenceNode>() ||
 			ModelNode->IsA<URigVMCollapseNode>() ||
+			ModelNode->IsA<URigVMTemplateNode>() ||
 			ModelNode->IsLoopNode() ||
 			Cast<URigVMUnitNode>(ModelNode) != nullptr;
 	}
@@ -783,6 +801,7 @@ FSlateIcon UControlRigGraphNode::GetIconAndTint(FLinearColor& OutColor) const
 	static FSlateIcon EntryReturnIcon("EditorStyle", "GraphEditor.Default_16x");
 	static FSlateIcon CollapsedNodeIcon("EditorStyle", "GraphEditor.SubGraph_16x");
 	static FSlateIcon ArrayNodeIteratorIcon("EditorStyle", "GraphEditor.Macro.ForEach_16x");
+	static FSlateIcon TemplateNodeIcon("ControlRigEditorStyle", "ControlRig.Template");
 
 	if (URigVMNode* ModelNode = GetModelNode())
 	{
@@ -815,16 +834,14 @@ FSlateIcon UControlRigGraphNode::GetIconAndTint(FLinearColor& OutColor) const
 			}
 		}
 
+		if(URigVMTemplateNode* TemplateNode = Cast<URigVMTemplateNode>(ModelNode))
+		{
+			return TemplateNodeIcon;
+		}
+
 		if (URigVMUnitNode *UnitNode = Cast<URigVMUnitNode>(ModelNode))
 		{
-			ensure(UnitNode->GetScriptStruct());
-			
 			FString IconPath;
-
-			// icon path format: StyleSetName|StyleName|SmallStyleName|StatusOverlayStyleName
-			// the last two names are optional, see FSlateIcon() for reference
-			UnitNode->GetScriptStruct()->GetStringMetaDataHierarchical(FRigVMStruct::IconMetaName, &IconPath);
-
 			const int32 NumOfIconPathNames = 4;
 			
 			FName IconPathNames[NumOfIconPathNames] = {
@@ -834,26 +851,31 @@ FSlateIcon UControlRigGraphNode::GetIconAndTint(FLinearColor& OutColor) const
 				NAME_None  // StatusOverlayStyleName
 			};
 
-			int32 NameIndex = 0;
-
-			while (!IconPath.IsEmpty() && NameIndex < NumOfIconPathNames)
+			if(UnitNode->GetScriptStruct())
 			{
-				FString Left;
-				FString Right;
+				// icon path format: StyleSetName|StyleName|SmallStyleName|StatusOverlayStyleName
+				// the last two names are optional, see FSlateIcon() for reference
+				UnitNode->GetScriptStruct()->GetStringMetaDataHierarchical(FRigVMStruct::IconMetaName, &IconPath);
 
-				if (!IconPath.Split(TEXT("|"), &Left, &Right))
+				int32 NameIndex = 0;
+
+				while (!IconPath.IsEmpty() && NameIndex < NumOfIconPathNames)
 				{
-					Left = IconPath;
-				}
+					FString Left;
+					FString Right;
 
-				IconPathNames[NameIndex] = FName(*Left);
+					if (!IconPath.Split(TEXT("|"), &Left, &Right))
+					{
+						Left = IconPath;
+					}
 
-				NameIndex++;
-				IconPath = Right;
-			}
+					IconPathNames[NameIndex] = FName(*Left);
 
+					NameIndex++;
+					IconPath = Right;
+				}
+			}
 			return FSlateIcon(IconPathNames[0], IconPathNames[1], IconPathNames[2], IconPathNames[3]);
-			
 		}
 	}
 
@@ -1105,8 +1127,22 @@ void UControlRigGraphNode::AutowireNewNode(UEdGraphPin* FromPin)
 
 	const UControlRigGraphSchema* Schema = GetDefault<UControlRigGraphSchema>();
 
+	// copying high level information into a local array since the try create connection below
+	// may cause the pin array to be destroyed / changed
+	TArray<TPair<FName, EEdGraphPinDirection>> PinsToVisit;
 	for(UEdGraphPin* Pin : Pins)
 	{
+		PinsToVisit.Emplace(Pin->GetFName(), Pin->Direction);
+	}
+
+	for(const TPair<FName, EEdGraphPinDirection>& PinToVisit : PinsToVisit)
+	{
+		UEdGraphPin* Pin = FindPin(PinToVisit.Key, PinToVisit.Value);
+		if(Pin == nullptr)
+		{
+			continue;
+		}
+		
 		if (Pin->ParentPin != nullptr)
 		{
 			continue;
@@ -1118,9 +1154,9 @@ void UControlRigGraphNode::AutowireNewNode(UEdGraphPin* FromPin)
 			if (Schema->TryCreateConnection(FromPin, Pin))
 			{
 				break;
-					}
-				}
+			}
 		}
+	}
 }
 
 bool UControlRigGraphNode::IsSelectedInEditor() const
diff --git a/Engine/Plugins/Animation/ControlRig/Source/ControlRigDeveloper/Private/Graph/ControlRigGraphSchema.cpp b/Engine/Plugins/Animation/ControlRig/Source/ControlRigDeveloper/Private/Graph/ControlRigGraphSchema.cpp
--- a/Engine/Plugins/Animation/ControlRig/Source/ControlRigDeveloper/Private/Graph/ControlRigGraphSchema.cpp
+++ b/Engine/Plugins/Animation/ControlRig/Source/ControlRigDeveloper/Private/Graph/ControlRigGraphSchema.cpp
@@ -643,7 +643,7 @@ bool UControlRigGraphSchema::TryCreateConnection(UEdGraphPin* PinA, UEdGraphPin*
 				}
 			}
 #endif
-			
+
 			return Controller->AddLink(PinA->GetName(), PinB->GetName(), true, true);
 		}
 	}
@@ -731,7 +731,7 @@ FLinearColor UControlRigGraphSchema::GetPinTypeColor(const FEdGraphPinType& PinT
 				return FLinearColor::White;
 			}
 
-			if (Struct->IsChildOf(FRigVMUnknownType::StaticStruct()))
+			if (Struct->IsChildOf(RigVMTypeUtils::GetWildCardCPPTypeObject()))
 			{
 				return FLinearColor(FVector3f::OneVector * 0.25f);
 			}
@@ -771,7 +771,7 @@ void UControlRigGraphSchema::InsertAdditionalActions(TArray<UBlueprint*> InBluep
 		{
 			if(URigVMPin* ModelPin = RigNode->GetModelPinFromPinPath(EdGraphPins[0]->GetName()))
 			{
-				if(!ModelPin->IsExecuteContext() && !ModelPin->IsUnknownType())
+				if(!ModelPin->IsExecuteContext() && !ModelPin->IsWildCard())
 				{
 					if(!ModelPin->GetNode()->IsA<URigVMVariableNode>())
 					{
@@ -1524,38 +1524,51 @@ bool UControlRigGraphSchema::ArePinsCompatible(const UEdGraphPin* PinA, const UE
 		return false;
 	}
 
-	if (UControlRigGraphNode* GraphNode = Cast<UControlRigGraphNode>(PinB->GetOwningNode()))
-	{
-	}
-
-	// for reroute nodes - always allow it
-	if (PinA->PinType.PinCategory == TEXT("ANY_TYPE"))
-	{
-		UControlRigGraphSchema* MutableThis = (UControlRigGraphSchema*)this;
-		MutableThis->LastPinForCompatibleCheck = PinB;
-		MutableThis->bLastPinWasInput = PinB->Direction == EGPD_Input;
-		return true;
-	}
-	if (PinB->PinType.PinCategory == TEXT("ANY_TYPE"))
-	{
-		UControlRigGraphSchema* MutableThis = (UControlRigGraphSchema*)this;
-		MutableThis->LastPinForCompatibleCheck = PinA;
-		MutableThis->bLastPinWasInput = PinA->Direction == EGPD_Input;
-		return true;
-	}
-
 	// if we are looking at a polymorphic node
 	if((PinA->PinType.ContainerType == PinB->PinType.ContainerType) ||
 		(PinA->PinType.PinSubCategoryObject != PinB->PinType.PinSubCategoryObject))
 	{
-		if(PinA->PinType.PinCategory == UEdGraphSchema_K2::PC_Struct && PinA->PinType.PinSubCategoryObject == FRigVMUnknownType::StaticStruct())
+		auto IsPinCompatibleWithType = [](const UEdGraphPin* InPin, const FEdGraphPinType& InPinType) -> bool
 		{
-			bool bIsExecuteContext = false;
-			if(const UScriptStruct* ExecuteContextScriptStruct = Cast<UScriptStruct>(PinB->PinType.PinSubCategoryObject))
+			if(const UScriptStruct* ScriptStruct = Cast<UScriptStruct>(InPinType.PinSubCategoryObject))
 			{
-				bIsExecuteContext = ExecuteContextScriptStruct->IsChildOf(FRigVMExecuteContext::StaticStruct());
+				if(ScriptStruct->IsChildOf(FRigVMExecuteContext::StaticStruct()))
+				{
+					return false;
+				}
 			}
-			if(!bIsExecuteContext)
+
+			if(const UControlRigGraphNode* RigNode = Cast<UControlRigGraphNode>(InPin->GetOwningNode()))
+			{
+				if(const FRigVMTemplate* Template = RigNode->GetTemplate())
+				{
+					FString CPPType;
+					UObject* CPPTypeObject = nullptr;
+					if(RigVMTypeUtils::CPPTypeFromPinType(InPinType, CPPType, &CPPTypeObject))
+					{
+						FString PinPath, PinName;
+						URigVMPin::SplitPinPathAtEnd(InPin->GetName(), PinPath, PinName);
+						
+						if((InPin->ParentPin != nullptr) &&
+							(InPin->ParentPin->ParentPin == nullptr) &&
+							(InPin->ParentPin->PinType.ContainerType == EPinContainerType::Array))
+						{
+							URigVMPin::SplitPinPathAtEnd(InPin->ParentPin->GetName(), PinPath, PinName);
+							CPPType = RigVMTypeUtils::ArrayTypeFromBaseType(CPPType);
+						}
+						if(!Template->ArgumentSupportsType(*PinName, CPPType))
+						{
+							return false;
+						}
+					}
+				}
+			}
+			return true;
+		};
+		
+		if(PinA->PinType.PinSubCategoryObject == RigVMTypeUtils::GetWildCardCPPTypeObject())
+		{
+			if(IsPinCompatibleWithType(PinA, PinB->PinType))
 			{
 				UControlRigGraphSchema* MutableThis = (UControlRigGraphSchema*)this;
 				MutableThis->LastPinForCompatibleCheck = PinB;
@@ -1563,14 +1576,9 @@ bool UControlRigGraphSchema::ArePinsCompatible(const UEdGraphPin* PinA, const UE
 				return true;
 			}
 		}
- 		else if(PinB->PinType.PinCategory == UEdGraphSchema_K2::PC_Struct && PinB->PinType.PinSubCategoryObject == FRigVMUnknownType::StaticStruct())
+ 		else if(PinB->PinType.PinSubCategoryObject == RigVMTypeUtils::GetWildCardCPPTypeObject())
 		{
-			bool bIsExecuteContext = false;
- 			if(const UScriptStruct* ExecuteContextScriptStruct = Cast<UScriptStruct>(PinA->PinType.PinSubCategoryObject))
- 			{
- 				bIsExecuteContext = ExecuteContextScriptStruct->IsChildOf(FRigVMExecuteContext::StaticStruct());
- 			}
- 			if(!bIsExecuteContext)
+ 			if(IsPinCompatibleWithType(PinB, PinA->PinType))
  			{
  				UControlRigGraphSchema* MutableThis = (UControlRigGraphSchema*)this;
  				MutableThis->LastPinForCompatibleCheck = PinA;
@@ -1594,51 +1602,6 @@ bool UControlRigGraphSchema::ArePinsCompatible(const UEdGraphPin* PinA, const UE
 		}
 	}
 
-	struct Local
-	{
-		static FString GetCPPTypeFromPinType(const FEdGraphPinType& InPinType)
-		{
-			return FString();
-		}
-	};
-
-	if (PinA->PinType.PinCategory.IsNone() && PinB->PinType.PinCategory.IsNone())
-	{
-		return true;
-	}
-	else if (PinA->PinType.PinCategory.IsNone() && !PinB->PinType.PinCategory.IsNone())
-	{
-		if (UControlRigGraphNode* RigNode = Cast<UControlRigGraphNode>(PinA->GetOwningNode()))
-		{
-			if (URigVMPrototypeNode* PrototypeNode = Cast<URigVMPrototypeNode>(RigNode->GetModelNode()))
-			{
-				FString CPPType = Local::GetCPPTypeFromPinType(PinB->PinType);
-				FString Left, Right;
-				URigVMPin::SplitPinPathAtStart(PinA->GetName(), Left, Right);
-				if (URigVMPin* ModelPin = PrototypeNode->FindPin(Right))
-				{
-					return PrototypeNode->SupportsType(ModelPin, CPPType);
-				}
-			}
-		}
-	}
-	else if (!PinA->PinType.PinCategory.IsNone() && PinB->PinType.PinCategory.IsNone())
-	{
-		if (UControlRigGraphNode* RigNode = Cast<UControlRigGraphNode>(PinB->GetOwningNode()))
-		{
-			if (URigVMPrototypeNode* PrototypeNode = Cast<URigVMPrototypeNode>(RigNode->GetModelNode()))
-			{
-				FString CPPType = Local::GetCPPTypeFromPinType(PinA->PinType);
-				FString Left, Right;
-				URigVMPin::SplitPinPathAtStart(PinB->GetName(), Left, Right);
-				if (URigVMPin* ModelPin = PrototypeNode->FindPin(Right))
-				{
-					return PrototypeNode->SupportsType(ModelPin, CPPType);
-				}
-			}
-		}
-	}
-
 	return GetDefault<UEdGraphSchema_K2>()->ArePinsCompatible(PinA, PinB, CallingContext, bIgnoreArray);
 }
 
diff --git a/Engine/Plugins/Animation/ControlRig/Source/ControlRigDeveloper/Public/ControlRigBlueprint.h b/Engine/Plugins/Animation/ControlRig/Source/ControlRigDeveloper/Public/ControlRigBlueprint.h
--- a/Engine/Plugins/Animation/ControlRig/Source/ControlRigDeveloper/Public/ControlRigBlueprint.h
+++ b/Engine/Plugins/Animation/ControlRig/Source/ControlRigDeveloper/Public/ControlRigBlueprint.h
@@ -551,6 +551,7 @@ private:
 	void PatchRigElementKeyCacheOnLoad();
 	void PatchBoundVariables();
 	void PatchVariableNodesWithIncorrectType();
+	void ConvertUnitNodesToTemplateNodes();
 
 	TMap<FName, int32> AddedMemberVariableMap;
 	TArray<FBPVariableDescription> LastNewVariables;
diff --git a/Engine/Plugins/Animation/ControlRig/Source/ControlRigDeveloper/Public/Graph/ControlRigGraph.h b/Engine/Plugins/Animation/ControlRig/Source/ControlRigDeveloper/Public/Graph/ControlRigGraph.h
--- a/Engine/Plugins/Animation/ControlRig/Source/ControlRigDeveloper/Public/Graph/ControlRigGraph.h
+++ b/Engine/Plugins/Animation/ControlRig/Source/ControlRigDeveloper/Public/Graph/ControlRigGraph.h
@@ -135,6 +135,8 @@ private:
 
 	TMap<URigVMNode*, TPair<int32, int32>> CachedInstructionIndices;
 
+	void RemoveAndDeleteNode(UEdGraphNode* InNode);
+
 #endif
 #if WITH_EDITORONLY_DATA
 
@@ -157,7 +159,7 @@ private:
 	friend class UControlRigBranchNodeSpawner;
 	friend class UControlRigIfNodeSpawner;
 	friend class UControlRigSelectNodeSpawner;
-	friend class UControlRigPrototypeNodeSpawner;
+	friend class UControlRigTemplateNodeSpawner;
 	friend class UControlRigEnumNodeSpawner;
 	friend class UControlRigFunctionRefNodeSpawner;
 	friend class UControlRigArrayNodeSpawner;
diff --git a/Engine/Plugins/Animation/ControlRig/Source/ControlRigDeveloper/Public/Graph/ControlRigGraphNode.h b/Engine/Plugins/Animation/ControlRig/Source/ControlRigDeveloper/Public/Graph/ControlRigGraphNode.h
--- a/Engine/Plugins/Animation/ControlRig/Source/ControlRigDeveloper/Public/Graph/ControlRigGraphNode.h
+++ b/Engine/Plugins/Animation/ControlRig/Source/ControlRigDeveloper/Public/Graph/ControlRigGraphNode.h
@@ -151,6 +151,8 @@ public:
 
 	int32 GetInstructionIndex(bool bAsInput) const;
 
+	const FRigVMTemplate* GetTemplate() const;
+
 protected:
 
 	FLinearColor GetNodeProfilingColor() const;
@@ -219,8 +221,10 @@ private:
 
 	FNodeTitleDirtied NodeTitleDirtied;
 
+	mutable const FRigVMTemplate* CachedTemplate;
+
 	friend class SControlRigGraphNode;
 	friend class FControlRigArgumentLayout;
 	friend class FControlRigGraphDetails;
-	friend class UControlRigPrototypeNodeSpawner;
+	friend class UControlRigTemplateNodeSpawner;
 };
diff --git a/Engine/Plugins/Animation/ControlRig/Source/ControlRigEditor/Private/ControlRigEditorModule.cpp b/Engine/Plugins/Animation/ControlRig/Source/ControlRigEditor/Private/ControlRigEditorModule.cpp
--- a/Engine/Plugins/Animation/ControlRig/Source/ControlRigEditor/Private/ControlRigEditorModule.cpp
+++ b/Engine/Plugins/Animation/ControlRig/Source/ControlRigEditor/Private/ControlRigEditorModule.cpp
@@ -52,7 +52,7 @@
 #include "Graph/NodeSpawners/ControlRigBranchNodeSpawner.h"
 #include "Graph/NodeSpawners/ControlRigIfNodeSpawner.h"
 #include "Graph/NodeSpawners/ControlRigSelectNodeSpawner.h"
-#include "Graph/NodeSpawners/ControlRigPrototypeNodeSpawner.h"
+#include "Graph/NodeSpawners/ControlRigTemplateNodeSpawner.h"
 #include "Graph/NodeSpawners/ControlRigEnumNodeSpawner.h"
 #include "Graph/NodeSpawners/ControlRigFunctionRefNodeSpawner.h"
 #include "Graph/NodeSpawners/ControlRigArrayNodeSpawner.h"
@@ -129,7 +129,6 @@
 
 DEFINE_LOG_CATEGORY(LogControlRigEditor);
 
-
 void FControlRigEditorModule::StartupModule()
 {
 	FControlRigEditModeCommands::Register();
@@ -1017,24 +1016,26 @@ void FControlRigEditorModule::GetTypeActions(UControlRigBlueprint* CRB, FBluepri
 		return;
 	}
 
-	/*
-	for (const FRigVMPrototype& Prototype : FRigVMRegistry::Get().GetPrototypes())
+#if UE_RIGVM_ENABLE_TEMPLATE_NODES
+	
+	for (const FRigVMTemplate& Template : FRigVMRegistry::Get().GetTemplates())
 	{
-		// ignore prototype that have only one function
-		if (Prototype.NumFunctions() <= 1)
+		// ignore templates that have only one permutation
+		if (Template.NumPermutations() <= 1)
 		{
 			continue;
 		}
 
-		FText NodeCategory = FText::FromString(Prototype.GetCategory());
-		FText MenuDesc = FText::FromName(Prototype.GetName());
-		FText ToolTip;
+		FText NodeCategory = FText::FromString(Template.GetCategory());
+		FText MenuDesc = FText::FromName(Template.GetName());
+		FText ToolTip = Template.GetTooltipText();
 
-		UBlueprintNodeSpawner* NodeSpawner = UControlRigPrototypeNodeSpawner::CreateFromNotation(Prototype.GetNotation(), MenuDesc, NodeCategory, ToolTip);
+		UBlueprintNodeSpawner* NodeSpawner = UControlRigTemplateNodeSpawner::CreateFromNotation(Template.GetNotation(), MenuDesc, NodeCategory, ToolTip);
 		check(NodeSpawner != nullptr);
 		ActionRegistrar.AddBlueprintAction(ActionKey, NodeSpawner);
 	};
-	*/
+
+#endif
 
 	// Add all rig units
 	for(const FRigVMFunction& Function : FRigVMRegistry::Get().GetFunctions())
@@ -1045,16 +1046,19 @@ void FControlRigEditorModule::GetTypeActions(UControlRigBlueprint* CRB, FBluepri
 			continue;
 		}
 
-		// skip rig units which have a prototype
-		/*
-		if (Function.PrototypeIndex != INDEX_NONE)
+#if UE_RIGVM_ENABLE_TEMPLATE_NODES
+		// skip rig units which have a template
+		if (Function.GetTemplate())
 		{
-			if (FRigVMRegistry::Get().GetPrototypes()[Function.PrototypeIndex].NumFunctions() > 1)
-			{
-				continue;
-			}
+			continue;
+		}
+#endif
+
+		// skip deprecated units
+		if(Function.Struct->HasMetaData(FRigVMStruct::DeprecatedMetaName))
+		{
+			continue;
 		}
-		*/
 
 		FString CategoryMetadata, DisplayNameMetadata, MenuDescSuffixMetadata;
 		Struct->GetStringMetaDataHierarchical(FRigVMStruct::CategoryMetaName, &CategoryMetadata);
@@ -1248,22 +1252,34 @@ void FControlRigEditorModule::GetContextMenuActions(const UControlRigGraphSchema
 		if (UEdGraphPin* InGraphPin = (UEdGraphPin* )Context->Pin)
 		{
 			UEdGraph* Graph = InGraphPin->GetOwningNode()->GetGraph();
-			
+
+			bool bIsExecutePin = false;
+			bool bIsUnknownPin = false;
+			if(UScriptStruct* ScriptStruct = Cast<UScriptStruct>(InGraphPin->PinType.PinSubCategoryObject))
+			{
+				bIsExecutePin = ScriptStruct->IsChildOf(FRigVMExecuteContext::StaticStruct());
+				bIsUnknownPin = ScriptStruct->IsChildOf(RigVMTypeUtils::GetWildCardCPPTypeObject());
+			}
+			const bool bIsEditablePin = !bIsExecutePin && !bIsUnknownPin;
+
 			if(UControlRigGraphNode* RigNode = Cast<UControlRigGraphNode>(InGraphPin->GetOwningNode()))
 			{
-				// Add the watch pin / unwatch pin menu items
-				FToolMenuSection& Section = Menu->AddSection("EdGraphSchemaWatches", LOCTEXT("WatchesHeader", "Watches"));
-				UBlueprint* OwnerBlueprint = FBlueprintEditorUtils::FindBlueprintForGraphChecked(Context->Graph);
+				if(bIsEditablePin)
 				{
-					if (FKismetDebugUtilities::IsPinBeingWatched(OwnerBlueprint, InGraphPin))
-					{
-						Section.AddMenuEntry(FGraphEditorCommands::Get().StopWatchingPin);
-					}
-					else
+					// Add the watch pin / unwatch pin menu items
+					FToolMenuSection& Section = Menu->AddSection("EdGraphSchemaWatches", LOCTEXT("WatchesHeader", "Watches"));
+					UBlueprint* OwnerBlueprint = FBlueprintEditorUtils::FindBlueprintForGraphChecked(Context->Graph);
 					{
-						Section.AddMenuEntry(FGraphEditorCommands::Get().StartWatchingPin);
+						if (FKismetDebugUtilities::IsPinBeingWatched(OwnerBlueprint, InGraphPin))
+						{
+							Section.AddMenuEntry(FGraphEditorCommands::Get().StopWatchingPin);
+						}
+						else
+						{
+							Section.AddMenuEntry(FGraphEditorCommands::Get().StartWatchingPin);
+						}
 					}
-				}				
+				}
 			}
 
 			// Add alphainterp menu entries
@@ -1273,7 +1289,7 @@ void FControlRigEditorModule::GetContextMenuActions(const UControlRigGraphSchema
 				{
 					URigVMController* Controller = RigBlueprint->GetController(ModelPin->GetGraph());
 
-					if (ModelPin->IsArray())
+					if (ModelPin->IsArray() && !bIsExecutePin)
 					{
 						FToolMenuSection& Section = Menu->AddSection("EdGraphSchemaPinArrays", LOCTEXT("PinArrays", "Arrays"));
 						Section.AddMenuEntry(
@@ -1286,7 +1302,7 @@ void FControlRigEditorModule::GetContextMenuActions(const UControlRigGraphSchema
 							})
 						));
 		}
-					if(ModelPin->IsArrayElement())
+					if(ModelPin->IsArrayElement() && !bIsExecutePin)
 					{
 						FToolMenuSection& Section = Menu->AddSection("EdGraphSchemaPinArrays", LOCTEXT("PinArrays", "Arrays"));
 						Section.AddMenuEntry(
@@ -1310,7 +1326,7 @@ void FControlRigEditorModule::GetContextMenuActions(const UControlRigGraphSchema
 					}
 
 					if (ModelPin->GetDirection() == ERigVMPinDirection::Input &&
-							!ModelPin->IsExecuteContext())
+							bIsEditablePin)
 					{
 						if (ModelPin->IsBoundToVariable())
 						{
@@ -1353,7 +1369,7 @@ void FControlRigEditorModule::GetContextMenuActions(const UControlRigGraphSchema
 						Cast<URigVMArrayNode>(ModelPin->GetNode()) != nullptr)
 					{
 						if (ModelPin->GetDirection() == ERigVMPinDirection::Input && 
-							!ModelPin->IsExecuteContext())
+							bIsEditablePin)
 						{
 							if (!ModelPin->IsBoundToVariable())
 							{
@@ -1431,12 +1447,12 @@ void FControlRigEditorModule::GetContextMenuActions(const UControlRigGraphSchema
 							bool bBoundToVariable = false;
 							for (URigVMInjectionInfo* Injection : ModelPin->GetInjectedNodes())
 							{
-								FString PrototypeName;
+								FString TemplateName;
 								if (URigVMUnitNode* UnitNode = Cast<URigVMUnitNode>(Injection->Node))
 								{
-									if (UnitNode->GetScriptStruct()->GetStringMetaDataHierarchical(TEXT("PrototypeName"), &PrototypeName))
+									if (UnitNode->GetScriptStruct()->GetStringMetaDataHierarchical(FRigVMStruct::TemplateNameMetaName, &TemplateName))
 									{
-										if (PrototypeName == TEXT("AlphaInterp"))
+										if (TemplateName == TEXT("AlphaInterp"))
 										{
 											InterpNode = Injection->Node;
 											break;
@@ -1523,12 +1539,12 @@ void FControlRigEditorModule::GetContextMenuActions(const UControlRigGraphSchema
 							bool bBoundToVariable = false;
 							for (URigVMInjectionInfo* Injection : ModelPin->GetInjectedNodes())
 							{
-								FString PrototypeName;
+								FString TemplateName;
 								if (URigVMUnitNode* UnitNode = Cast<URigVMUnitNode>(Injection->Node))
 								{
-									if (UnitNode->GetScriptStruct()->GetStringMetaDataHierarchical(TEXT("PrototypeName"), &PrototypeName))
+									if (UnitNode->GetScriptStruct()->GetStringMetaDataHierarchical(FRigVMStruct::TemplateNameMetaName, &TemplateName))
 									{
-										if (PrototypeName == TEXT("VisualDebug"))
+										if (TemplateName == TEXT("VisualDebug"))
 										{
 											VisualDebugNode = Injection->Node;
 											break;
@@ -1679,16 +1695,19 @@ void FControlRigEditorModule::GetContextMenuActions(const UControlRigGraphSchema
 						if (URigVMUnitNode* UnitNode = Cast<URigVMUnitNode>(ModelNode))
 						{
 							ScriptStruct = UnitNode->GetScriptStruct();
-							StructOnScope = UnitNode->ConstructStructInstance(false /* default */);
-							StructMemory = (FRigUnit*)StructOnScope->GetStructMemory();
-
-							FRigNameCache NameCache;
-							FRigUnitContext RigUnitContext;
-							RigUnitContext.Hierarchy = TemporaryHierarchy;
-							RigUnitContext.State = EControlRigState::Update;
-							RigUnitContext.NameCache = &NameCache;
-							
-							StructMemory->Execute(RigUnitContext);
+							if(ScriptStruct)
+							{
+								StructOnScope = UnitNode->ConstructStructInstance(false /* default */);
+								StructMemory = (FRigUnit*)StructOnScope->GetStructMemory();
+
+								FRigNameCache NameCache;
+								FRigUnitContext RigUnitContext;
+								RigUnitContext.Hierarchy = TemporaryHierarchy;
+								RigUnitContext.State = EControlRigState::Update;
+								RigUnitContext.NameCache = &NameCache;
+								
+								StructMemory->Execute(RigUnitContext);
+							}
 						}
 
 						const TArray<URigVMPin*> AllPins = ModelNode->GetAllPinsRecursively();
diff --git a/Engine/Plugins/Animation/ControlRig/Source/ControlRigEditor/Private/ControlRigEditorStyle.h b/Engine/Plugins/Animation/ControlRig/Source/ControlRigEditor/Private/ControlRigEditorStyle.h
--- a/Engine/Plugins/Animation/ControlRig/Source/ControlRigEditor/Private/ControlRigEditorStyle.h
+++ b/Engine/Plugins/Animation/ControlRig/Source/ControlRigEditor/Private/ControlRigEditorStyle.h
@@ -97,6 +97,7 @@ public:
 			
 			// icons for control units
 			Set("ControlRig.RigUnit", new IMAGE_BRUSH("Slate/ControlRigUnit_16x", Icon16x16));
+			Set("ControlRig.Template", new IMAGE_BRUSH("Slate/Template_16x", Icon16x16));
 
 			Set("ControlRig.AutoCompileGraph", new IMAGE_BRUSH("Slate/AutoCompile", Icon20x20));
 			Set("ControlRig.AutoCompileGraph.Small", new IMAGE_BRUSH("Slate/AutoCompile", Icon20x20));
diff --git a/Engine/Plugins/Animation/ControlRig/Source/ControlRigEditor/Private/ControlRigGraphDetails.cpp b/Engine/Plugins/Animation/ControlRig/Source/ControlRigEditor/Private/ControlRigGraphDetails.cpp
--- a/Engine/Plugins/Animation/ControlRig/Source/ControlRigEditor/Private/ControlRigGraphDetails.cpp
+++ b/Engine/Plugins/Animation/ControlRig/Source/ControlRigEditor/Private/ControlRigGraphDetails.cpp
@@ -1000,12 +1000,8 @@ FReply FControlRigGraphDetails::OnAddNewInputClicked()
 						{
 							ArgumentName = *StrippedArgumentName;
 						}
-						
-						CPPType = LastPin->GetCPPType();
-						if (LastPin->GetCPPTypeObject())
-						{
-							CPPTypeObjectPath = *LastPin->GetCPPTypeObject()->GetPathName();
-						}
+
+						RigVMTypeUtils::CPPTypeFromPin(LastPin, CPPType, CPPTypeObjectPath);						
 						DefaultValue = LastPin->GetDefaultValue();
 					}
 				}
@@ -1831,6 +1827,7 @@ void FControlRigGraphMathTypeDetails::CustomizeHeader(TSharedRef<IPropertyHandle
 	do
 	{
 		PropertyChain.AddHead(ChainHandle->GetProperty());
+		PropertyArrayIndices.Insert(ChainHandle->GetIndexInArray(), 0);
 		ChainHandle = ChainHandle->GetParentHandle();
 
 		if(ChainHandle.IsValid())
diff --git a/Engine/Plugins/Animation/ControlRig/Source/ControlRigEditor/Private/ControlRigGraphDetails.h b/Engine/Plugins/Animation/ControlRig/Source/ControlRigEditor/Private/ControlRigGraphDetails.h
--- a/Engine/Plugins/Animation/ControlRig/Source/ControlRigEditor/Private/ControlRigGraphDetails.h
+++ b/Engine/Plugins/Animation/ControlRig/Source/ControlRigEditor/Private/ControlRigGraphDetails.h
@@ -263,14 +263,40 @@ protected:
 
 	// extracts the value for a nested property (for Example Settings.WorldTransform) from an outer owner
 	template<typename ValueType>
-	FORCEINLINE ValueType& ContainerUObjectToValueRef(UObject* InOwner) const
+	FORCEINLINE_DEBUGGABLE ValueType& ContainerUObjectToValueRef(UObject* InOwner, ValueType& InDefault) const
 	{
 		FEditPropertyChain::TDoubleLinkedListNode* PropertyNode = PropertyChain.GetHead();
 		uint8* MemoryPtr = (uint8*)InOwner;
+		int32 ChainIndex = 0;
 		do
 		{
-			MemoryPtr = PropertyNode->GetValue()->ContainerPtrToValuePtr<uint8>(MemoryPtr);
+			const FProperty* Property = PropertyNode->GetValue();
+			MemoryPtr = Property->ContainerPtrToValuePtr<uint8>(MemoryPtr);
+
 			PropertyNode = PropertyNode->GetNextNode();
+			ChainIndex++;
+			
+			if(PropertyArrayIndices.IsValidIndex(ChainIndex))
+			{
+				const int32 ArrayIndex = PropertyArrayIndices[ChainIndex];
+				if(ArrayIndex != INDEX_NONE)
+				{
+					const FArrayProperty* ArrayProperty = CastField<FArrayProperty>(Property->GetOwnerProperty());
+					check(ArrayProperty);
+					
+					FScriptArrayHelper ArrayHelper(ArrayProperty, MemoryPtr);
+					if(!ArrayHelper.IsValidIndex(ArrayIndex))
+					{
+						return InDefault;
+					}
+					MemoryPtr = ArrayHelper.GetRawPtr(ArrayIndex);
+
+					// skip to the next property node already
+					PropertyNode = PropertyNode->GetNextNode();
+					ChainIndex++;
+				}
+			}
+
 		}
 		while (PropertyNode);
 
@@ -290,7 +316,8 @@ protected:
 		{
 			if(InPropertyHandle->IsValidHandle())
 			{
-				const VectorType& Vector = ContainerUObjectToValueRef<VectorType>(Object);
+				static VectorType ZeroVector = VectorType();
+				const VectorType& Vector = ContainerUObjectToValueRef<VectorType>(Object, ZeroVector);
 				NumericType Component = Vector[InComponent];
 				if(Result.IsSet())
 				{
@@ -330,7 +357,8 @@ protected:
 			UObject* Object = ObjectsBeingCustomized[Index];
 			if(InPropertyHandle->IsValidHandle())
 			{
-				VectorType& Vector = ContainerUObjectToValueRef<VectorType>(Object);
+				static VectorType ZeroVector = VectorType();
+				VectorType& Vector = ContainerUObjectToValueRef<VectorType>(Object, ZeroVector);
 				VectorType PreviousVector = Vector;
 				Vector[InComponent] = InValue;
 					
@@ -415,7 +443,8 @@ protected:
 		{
 			if(InPropertyHandle->IsValidHandle())
 			{
-				const RotationType& Rotation = ContainerUObjectToValueRef<RotationType>(Object);
+				static RotationType ZeroRotation = RotationType();
+				const RotationType& Rotation = ContainerUObjectToValueRef<RotationType>(Object, ZeroRotation);
 				if(Result.IsSet())
 				{
 					if(!Rotation.Equals(Result.GetValue()))
@@ -454,7 +483,8 @@ protected:
 			UObject* Object = ObjectsBeingCustomized[Index];
 			if(InPropertyHandle->IsValidHandle())
 			{
-				RotationType& Rotation = ContainerUObjectToValueRef<RotationType>(Object);
+				static RotationType ZeroRotation = RotationType();
+				RotationType& Rotation = ContainerUObjectToValueRef<RotationType>(Object, ZeroRotation);
 				RotationType PreviousRotation = Rotation;
 				Rotation = InValue;
 					
@@ -516,7 +546,8 @@ protected:
 		WidgetArgs.AllowEditRotationRepresentation(true);
 		WidgetArgs.UseQuaternionForRotation(IsQuaternionBasedRotation<TransformType>());
 
-		TransformType DefaultValue = ContainerUObjectToValueRef<TransformType>(ObjectsBeingCustomized[0]->GetClass()->GetDefaultObject());
+		static TransformType Identity = TransformType::Identity;
+		TransformType DefaultValue = ContainerUObjectToValueRef<TransformType>(ObjectsBeingCustomized[0]->GetClass()->GetDefaultObject(), Identity);
 
 		WidgetArgs.DiffersFromDefault_Lambda([this, InPropertyHandle, DefaultValue](ESlateTransformComponent::Type InTransformComponent) -> bool
 		{
@@ -524,7 +555,7 @@ protected:
 			{
 				if(InPropertyHandle->IsValidHandle())
 				{
-					const TransformType& Transform = ContainerUObjectToValueRef<TransformType>(Object);
+					const TransformType& Transform = ContainerUObjectToValueRef<TransformType>(Object, Identity);
 
 					switch(InTransformComponent)
 					{
@@ -572,7 +603,7 @@ protected:
 			{
 				if(InPropertyHandle->IsValidHandle())
 				{
-					const TransformType& Transform = ContainerUObjectToValueRef<TransformType>(Object);
+					const TransformType& Transform = ContainerUObjectToValueRef<TransformType>(Object, Identity);
 					
 					TOptional<FReal> Value = SAdvancedTransformInputBox<TransformType>::GetNumericValueFromTransform(
 						Transform,
@@ -627,7 +658,7 @@ protected:
 				UObject* Object = ObjectsBeingCustomized[Index];
 				if(InPropertyHandle->IsValidHandle())
 				{
-					TransformType& Transform = ContainerUObjectToValueRef<TransformType>(Object);
+					TransformType& Transform = ContainerUObjectToValueRef<TransformType>(Object, Identity);
 					TransformType PreviousTransform = Transform;
 					
 					SAdvancedTransformInputBox<TransformType>::ApplyNumericValueChange(
@@ -690,7 +721,7 @@ protected:
 				UObject* Object = ObjectsBeingCustomized[Index];
 				if(InPropertyHandle->IsValidHandle())
 				{
-					TransformType& Transform = ContainerUObjectToValueRef<TransformType>(Object);
+					TransformType& Transform = ContainerUObjectToValueRef<TransformType>(Object, Identity);
 					TransformType PreviousTransform = Transform;
 
 					switch(InTransformComponent)
@@ -746,6 +777,7 @@ protected:
 	TArray<UObject*> ObjectsBeingCustomized;
 	TArrayView<const UObject* const> ObjectBeingCustomizedView;
 	FEditPropertyChain PropertyChain;
+	TArray<int32> PropertyArrayIndices;
 	bool bEnabled;
 };
 
diff --git a/Engine/Plugins/Animation/ControlRig/Source/ControlRigEditor/Private/Editor/ControlRigEditor.cpp b/Engine/Plugins/Animation/ControlRig/Source/ControlRigEditor/Private/Editor/ControlRigEditor.cpp
--- a/Engine/Plugins/Animation/ControlRig/Source/ControlRigEditor/Private/Editor/ControlRigEditor.cpp
+++ b/Engine/Plugins/Animation/ControlRig/Source/ControlRigEditor/Private/Editor/ControlRigEditor.cpp
@@ -2693,8 +2693,11 @@ void FControlRigEditor::HandleModifiedEvent(ERigVMGraphNotifType InNotifType, UR
 								if(WrapperObject->GetOuter() == Pin->GetNode())
 								{
 									const FProperty* TargetProperty = WrapperObject->GetClass()->FindPropertyByName(RootPin->GetFName());
-									uint8* PropertyStorage = TargetProperty->ContainerPtrToValuePtr<uint8>(WrapperObject);
-									TargetProperty->ImportText_Direct(*DefaultValue, PropertyStorage, nullptr, PPF_None);
+									if(TargetProperty)
+									{
+										uint8* PropertyStorage = TargetProperty->ContainerPtrToValuePtr<uint8>(WrapperObject);
+										TargetProperty->ImportText_Direct(*DefaultValue, PropertyStorage, nullptr, PPF_None, nullptr);
+									}
 								}
 							}
 						}
@@ -3756,7 +3759,7 @@ void FControlRigEditor::SetPinControlNameListText(const FText& NewTypeInValue, E
 						URigVMNode* ControlledNode = ControlledPin->GetPinForLink()->GetNode();
 						if (URigVMUnitNode* UnitNode = Cast<URigVMUnitNode>(ControlledNode))
 						{
-							if (const FString* Value = UnitNode->GetScriptStruct()->FindMetaData(TEXT("PrototypeName")))
+							if (const FString* Value = UnitNode->GetScriptStruct()->FindMetaData(FRigVMStruct::TemplateNameMetaName))
 							{
 								if (Value->Equals("VisualDebug"))
 								{
diff --git a/Engine/Plugins/Animation/ControlRig/Source/ControlRigEditor/Private/Editor/DetailsViewWrapperObject.cpp b/Engine/Plugins/Animation/ControlRig/Source/ControlRigEditor/Private/Editor/DetailsViewWrapperObject.cpp
--- a/Engine/Plugins/Animation/ControlRig/Source/ControlRigEditor/Private/Editor/DetailsViewWrapperObject.cpp
+++ b/Engine/Plugins/Animation/ControlRig/Source/ControlRigEditor/Private/Editor/DetailsViewWrapperObject.cpp
@@ -205,7 +205,10 @@ UClass* UDetailsViewWrapperObject::GetClassForNodes(TArray<URigVMNode*> InNodes,
 	{
 		if(URigVMUnitNode* UnitNode = Cast<URigVMUnitNode>(Node))
 		{
-			UnitStructs.AddUnique(UnitNode->GetScriptStruct());
+			if(UScriptStruct* ScriptStruct = UnitNode->GetScriptStruct())
+			{
+				UnitStructs.AddUnique(ScriptStruct);
+			}
 		}
 	}
 
diff --git a/Engine/Plugins/Animation/ControlRig/Source/ControlRigEditor/Private/Graph/NodeSpawners/ControlRigArrayNodeSpawner.cpp b/Engine/Plugins/Animation/ControlRig/Source/ControlRigEditor/Private/Graph/NodeSpawners/ControlRigArrayNodeSpawner.cpp
--- a/Engine/Plugins/Animation/ControlRig/Source/ControlRigEditor/Private/Graph/NodeSpawners/ControlRigArrayNodeSpawner.cpp
+++ b/Engine/Plugins/Animation/ControlRig/Source/ControlRigEditor/Private/Graph/NodeSpawners/ControlRigArrayNodeSpawner.cpp
@@ -127,31 +127,7 @@ UEdGraphNode* UControlRigArrayNodeSpawner::Invoke(UEdGraph* ParentGraph, FBindin
 	UControlRigGraphNode* NewNode = nullptr;
 
 	bool const bIsTemplateNode = FBlueprintNodeTemplateCache::IsTemplateOuter(ParentGraph);
-
-	if (bIsTemplateNode)
-	{
-		NewNode = NewObject<UControlRigGraphNode>(ParentGraph, TEXT("ArrayNode"));
-		ParentGraph->AddNode(NewNode, false);
-
-		NewNode->CreateNewGuid();
-		NewNode->PostPlacedNewNode();
-
-		UEdGraphPin* InputValuePin = UEdGraphPin::CreatePin(NewNode);
-		UEdGraphPin* OutputValuePin = UEdGraphPin::CreatePin(NewNode);
-		NewNode->Pins.Add(InputValuePin);
-		NewNode->Pins.Add(OutputValuePin);
-
-		InputValuePin->PinType.PinCategory = TEXT("ANY_TYPE");
-		OutputValuePin->PinType.PinCategory = TEXT("ANY_TYPE");
-		InputValuePin->Direction = EGPD_Input;
-		OutputValuePin->Direction = EGPD_Output;
-		NewNode->SetFlags(RF_Transactional);
-
-		return NewNode;
-	}
-
-
-	bool const bUndo = !bIsTemplateNode;
+	bool const bIsUserFacingNode = !bIsTemplateNode;
 
 	// First create a backing member for our node
 	UControlRigGraph* RigGraph = Cast<UControlRigGraph>(ParentGraph);
@@ -159,27 +135,18 @@ UEdGraphNode* UControlRigArrayNodeSpawner::Invoke(UEdGraph* ParentGraph, FBindin
 	UControlRigBlueprint* RigBlueprint = Cast<UControlRigBlueprint>(FBlueprintEditorUtils::FindBlueprintForGraph(ParentGraph));
 	check(RigBlueprint);
 
-	FString CPPType = FRigVMUnknownType::StaticStruct()->GetStructCPPName();
-	FString CPPTypeObjectPath = FRigVMUnknownType::StaticStruct()->GetPathName();
+	FString CPPType = RigVMTypeUtils::GetWildCardCPPType();
+	FString CPPTypeObjectPath = RigVMTypeUtils::GetWildCardCPPTypeObject()->GetPathName();
 
-	if (UControlRigGraphSchema* RigSchema = Cast<UControlRigGraphSchema>((UEdGraphSchema*)ParentGraph->GetSchema()))
+	if(bIsUserFacingNode)
 	{
-		if (const UEdGraphPin* LastPin = RigSchema->LastPinForCompatibleCheck)
+		if (UControlRigGraphSchema* RigSchema = Cast<UControlRigGraphSchema>((UEdGraphSchema*)ParentGraph->GetSchema()))
 		{
-			if (URigVMPin* ModelPin = RigBlueprint->GetModel(ParentGraph)->FindPin(LastPin->GetName()))
+			if (const UEdGraphPin* LastPin = RigSchema->LastPinForCompatibleCheck)
 			{
-				CPPType = ModelPin->GetCPPType();
-				if(ModelPin->IsArray())
-				{
-					CPPType = ModelPin->GetArrayElementCppType();
-				}
-				if (ModelPin->GetCPPTypeObject())
+				if (URigVMPin* ModelPin = RigBlueprint->GetModel(ParentGraph)->FindPin(LastPin->GetName()))
 				{
-					CPPTypeObjectPath = *ModelPin->GetCPPTypeObject()->GetPathName();
-				}
-				else
-				{
-					CPPTypeObjectPath = FString();
+					RigVMTypeUtils::CPPTypeFromPin(ModelPin, CPPType, CPPTypeObjectPath, true);
 				}
 			}
 		}
@@ -204,24 +171,28 @@ UEdGraphNode* UControlRigArrayNodeSpawner::Invoke(UEdGraph* ParentGraph, FBindin
 		Controller->OpenUndoBracket(FString::Printf(TEXT("Add '%s' Node"), *Name.ToString()));
 	}
 
-	if (URigVMArrayNode* ModelNode = Controller->AddArrayNodeFromObjectPath(OpCode, CPPType, CPPTypeObjectPath, Location, Name.ToString(), bUndo, !bIsTemplateNode))
+	if (URigVMArrayNode* ModelNode = Controller->AddArrayNodeFromObjectPath(OpCode, CPPType, CPPTypeObjectPath, Location, Name.ToString(), bIsUserFacingNode, !bIsTemplateNode))
 	{
 		NewNode = Cast<UControlRigGraphNode>(RigGraph->FindNodeForModelNodeName(ModelNode->GetFName()));
 
-		if (NewNode && bUndo)
+		if (NewNode && bIsUserFacingNode)
 		{
 			Controller->ClearNodeSelection(true);
 			Controller->SelectNode(ModelNode, true, true);
 		}
 
-		if (bUndo)
+		if (bIsUserFacingNode)
 		{
 			Controller->CloseUndoBracket();
 		}
+		else
+		{
+			Controller->RemoveNode(ModelNode, false);
+		}
 	}
 	else
 	{
-		if (bUndo)
+		if (bIsUserFacingNode)
 		{
 			Controller->CancelUndoBracket();
 		}
diff --git a/Engine/Plugins/Animation/ControlRig/Source/ControlRigEditor/Private/Graph/NodeSpawners/ControlRigIfNodeSpawner.cpp b/Engine/Plugins/Animation/ControlRig/Source/ControlRigEditor/Private/Graph/NodeSpawners/ControlRigIfNodeSpawner.cpp
--- a/Engine/Plugins/Animation/ControlRig/Source/ControlRigEditor/Private/Graph/NodeSpawners/ControlRigIfNodeSpawner.cpp
+++ b/Engine/Plugins/Animation/ControlRig/Source/ControlRigEditor/Private/Graph/NodeSpawners/ControlRigIfNodeSpawner.cpp
@@ -54,31 +54,7 @@ UEdGraphNode* UControlRigIfNodeSpawner::Invoke(UEdGraph* ParentGraph, FBindingSe
 	UControlRigGraphNode* NewNode = nullptr;
 
 	bool const bIsTemplateNode = FBlueprintNodeTemplateCache::IsTemplateOuter(ParentGraph);
-
-	if (bIsTemplateNode)
-	{
-		NewNode = NewObject<UControlRigGraphNode>(ParentGraph, TEXT("IfNode"));
-		ParentGraph->AddNode(NewNode, false);
-
-		NewNode->CreateNewGuid();
-		NewNode->PostPlacedNewNode();
-
-		UEdGraphPin* InputValuePin = UEdGraphPin::CreatePin(NewNode);
-		UEdGraphPin* OutputValuePin = UEdGraphPin::CreatePin(NewNode);
-		NewNode->Pins.Add(InputValuePin);
-		NewNode->Pins.Add(OutputValuePin);
-
-		InputValuePin->PinType.PinCategory = TEXT("ANY_TYPE");
-		OutputValuePin->PinType.PinCategory = TEXT("ANY_TYPE");
-		InputValuePin->Direction = EGPD_Input;
-		OutputValuePin->Direction = EGPD_Output;
-		NewNode->SetFlags(RF_Transactional);
-
-		return NewNode;
-	}
-
-
-	bool const bUndo = !bIsTemplateNode;
+	bool const bIsUserFacingNode = !bIsTemplateNode;
 
 	// First create a backing member for our node
 	UControlRigGraph* RigGraph = Cast<UControlRigGraph>(ParentGraph);
@@ -86,23 +62,18 @@ UEdGraphNode* UControlRigIfNodeSpawner::Invoke(UEdGraph* ParentGraph, FBindingSe
 	UControlRigBlueprint* RigBlueprint = Cast<UControlRigBlueprint>(FBlueprintEditorUtils::FindBlueprintForGraph(ParentGraph));
 	check(RigBlueprint);
 
-	FString CPPType = FRigVMUnknownType::StaticStruct()->GetStructCPPName();
-	FName CPPTypeObjectPath = *FRigVMUnknownType::StaticStruct()->GetPathName();
+	FString CPPType = RigVMTypeUtils::GetWildCardCPPType();
+	FName CPPTypeObjectPath = *RigVMTypeUtils::GetWildCardCPPTypeObject()->GetPathName();
 
-	if (UControlRigGraphSchema* RigSchema = Cast<UControlRigGraphSchema>((UEdGraphSchema*)ParentGraph->GetSchema()))
+	if(bIsUserFacingNode)
 	{
-		if (const UEdGraphPin* LastPin = RigSchema->LastPinForCompatibleCheck)
+		if (UControlRigGraphSchema* RigSchema = Cast<UControlRigGraphSchema>((UEdGraphSchema*)ParentGraph->GetSchema()))
 		{
-			if (URigVMPin* ModelPin = RigBlueprint->GetModel(ParentGraph)->FindPin(LastPin->GetName()))
+			if (const UEdGraphPin* LastPin = RigSchema->LastPinForCompatibleCheck)
 			{
-				CPPType = ModelPin->GetCPPType();
-				if (ModelPin->GetCPPTypeObject())
-				{
-					CPPTypeObjectPath = *ModelPin->GetCPPTypeObject()->GetPathName();
-				}
-				else
+				if (URigVMPin* ModelPin = RigBlueprint->GetModel(ParentGraph)->FindPin(LastPin->GetName()))
 				{
-					CPPTypeObjectPath = NAME_None;
+					RigVMTypeUtils::CPPTypeFromPin(ModelPin, CPPType, CPPTypeObjectPath, false);
 				}
 			}
 		}
@@ -126,24 +97,28 @@ UEdGraphNode* UControlRigIfNodeSpawner::Invoke(UEdGraph* ParentGraph, FBindingSe
 		Controller->OpenUndoBracket(FString::Printf(TEXT("Add '%s' Node"), *Name.ToString()));
 	}
 
-	if (URigVMIfNode* ModelNode = Controller->AddIfNode(CPPType, CPPTypeObjectPath, Location, Name.ToString(), bUndo, !bIsTemplateNode))
+	if (URigVMIfNode* ModelNode = Controller->AddIfNode(CPPType, CPPTypeObjectPath, Location, Name.ToString(), bIsUserFacingNode, !bIsTemplateNode))
 	{
 		NewNode = Cast<UControlRigGraphNode>(RigGraph->FindNodeForModelNodeName(ModelNode->GetFName()));
 
-		if (NewNode && bUndo)
+		if (NewNode && bIsUserFacingNode)
 		{
 			Controller->ClearNodeSelection(true);
 			Controller->SelectNode(ModelNode, true, true);
 		}
 
-		if (bUndo)
+		if (bIsUserFacingNode)
 		{
 			Controller->CloseUndoBracket();
 		}
+		else
+		{
+			Controller->RemoveNode(ModelNode, false);
+		}
 	}
 	else
 	{
-		if (bUndo)
+		if (bIsUserFacingNode)
 		{
 			Controller->CancelUndoBracket();
 		}
diff --git a/Engine/Plugins/Animation/ControlRig/Source/ControlRigEditor/Private/Graph/NodeSpawners/ControlRigPrototypeNodeSpawner.cpp b/Engine/Plugins/Animation/ControlRig/Source/ControlRigEditor/Private/Graph/NodeSpawners/ControlRigPrototypeNodeSpawner.cpp
deleted file mode 100644
diff --git a/Engine/Plugins/Animation/ControlRig/Source/ControlRigEditor/Private/Graph/NodeSpawners/ControlRigPrototypeNodeSpawner.h b/Engine/Plugins/Animation/ControlRig/Source/ControlRigEditor/Private/Graph/NodeSpawners/ControlRigPrototypeNodeSpawner.h
deleted file mode 100644
diff --git a/Engine/Plugins/Animation/ControlRig/Source/ControlRigEditor/Private/Graph/NodeSpawners/ControlRigRerouteNodeSpawner.cpp b/Engine/Plugins/Animation/ControlRig/Source/ControlRigEditor/Private/Graph/NodeSpawners/ControlRigRerouteNodeSpawner.cpp
--- a/Engine/Plugins/Animation/ControlRig/Source/ControlRigEditor/Private/Graph/NodeSpawners/ControlRigRerouteNodeSpawner.cpp
+++ b/Engine/Plugins/Animation/ControlRig/Source/ControlRigEditor/Private/Graph/NodeSpawners/ControlRigRerouteNodeSpawner.cpp
@@ -34,6 +34,7 @@ UControlRigRerouteNodeSpawner* UControlRigRerouteNodeSpawner::CreateGeneric(cons
 	MenuSignature.Tooltip  = InTooltip;
 	MenuSignature.Category = InCategory;
 	MenuSignature.Keywords = FText::FromString(TEXT("Reroute,Elbow,Wire,Literal,Make Literal,Constant"));
+	MenuSignature.Icon = FSlateIcon(TEXT("ControlRigEditorStyle"), TEXT("ControlRig.RigUnit"));
 
 	return NodeSpawner;
 }
@@ -63,6 +64,7 @@ UEdGraphNode* UControlRigRerouteNodeSpawner::Invoke(UEdGraph* ParentGraph, FBind
 	UControlRigGraphNode* NewNode = nullptr;
 
 	bool const bIsTemplateNode = FBlueprintNodeTemplateCache::IsTemplateOuter(ParentGraph);
+	bool const bIsUserFacingNode = !bIsTemplateNode;
 
 	// First create a backing member for our node
 	UControlRigGraph* RigGraph = Cast<UControlRigGraph>(ParentGraph);
@@ -73,41 +75,24 @@ UEdGraphNode* UControlRigRerouteNodeSpawner::Invoke(UEdGraph* ParentGraph, FBind
 	FName MemberName = NAME_None;
 
 #if WITH_EDITOR
-	if (GEditor && !bIsTemplateNode)
+	if (GEditor && bIsUserFacingNode)
 	{
 		GEditor->CancelTransaction(0);
 	}
 #endif
 
-	if (bIsTemplateNode)
-	{
-		NewNode = NewObject<UControlRigGraphNode>(ParentGraph, TEXT("RerouteNode"));
-		ParentGraph->AddNode(NewNode, false);
-		
-		NewNode->CreateNewGuid();
-		NewNode->PostPlacedNewNode();
-
-		UEdGraphPin* InputValuePin = UEdGraphPin::CreatePin(NewNode);
-		UEdGraphPin* OutputValuePin = UEdGraphPin::CreatePin(NewNode);
-		NewNode->Pins.Add(InputValuePin);
-		NewNode->Pins.Add(OutputValuePin);
-
-		InputValuePin->PinType.PinCategory = TEXT("ANY_TYPE");
-		OutputValuePin->PinType.PinCategory = TEXT("ANY_TYPE");
-		InputValuePin->Direction = EGPD_Input;
-		OutputValuePin->Direction = EGPD_Output;
-		NewNode->SetFlags(RF_Transactional);
-
-		return NewNode;
-	}
-	else
+	URigVMController* Controller = bIsTemplateNode ? RigGraph->GetTemplateController() : RigBlueprint->GetController(ParentGraph);
+
+	if(bIsUserFacingNode)
 	{
-		URigVMController* Controller = RigBlueprint->GetController(ParentGraph);
 		Controller->OpenUndoBracket(TEXT("Added Reroute Node."));
+	}
 
-		FString PinPath;
-		bool bIsInput = false;
+	FString PinPath;
+	bool bIsInput = false;
 
+	if(bIsUserFacingNode)
+	{
 		if (UControlRigGraphSchema* RigSchema = Cast<UControlRigGraphSchema>((UEdGraphSchema*)ParentGraph->GetSchema()))
 		{
 			bIsInput = RigSchema->bLastPinWasInput;
@@ -120,59 +105,48 @@ UEdGraphNode* UControlRigRerouteNodeSpawner::Invoke(UEdGraph* ParentGraph, FBind
 				}
 			}
 		}
+	}
+
+	URigVMNode* ModelNode = nullptr;
+	if(!PinPath.IsEmpty())
+	{
+		ModelNode = Controller->AddRerouteNodeOnPin(PinPath, bIsInput, true, Location, FString(), bIsUserFacingNode, bIsUserFacingNode);
+	}
+	else
+	{
+		static const FString CPPType = RigVMTypeUtils::GetWildCardCPPType();
+		static const FName CPPTypeObjectPath = *RigVMTypeUtils::GetWildCardCPPTypeObject()->GetPathName();
+		ModelNode = Controller->AddFreeRerouteNode(true, CPPType, CPPTypeObjectPath, false, NAME_None, FString(), Location, FString(), bIsUserFacingNode);
+	}
+
+	if (ModelNode)
+	{
+		NewNode = Cast<UControlRigGraphNode>(RigGraph->FindNodeForModelNodeName(ModelNode->GetFName()));
 
-		if (URigVMNode* ModelNode = Controller->AddRerouteNodeOnPin(PinPath, bIsInput, true, Location, FString(), true, true))
+		if (NewNode)
 		{
-			for (UEdGraphNode* Node : ParentGraph->Nodes)
-			{
-				if (UControlRigGraphNode* RigNode = Cast<UControlRigGraphNode>(Node))
-				{
-					if (RigNode->GetModelNodeName() == ModelNode->GetFName())
-					{
-						NewNode = RigNode;
-						break;
-					}
-				}
-			}
+			Controller->ClearNodeSelection(bIsUserFacingNode);
+			Controller->SelectNode(ModelNode, true, bIsUserFacingNode);
+		}
 
-			if (NewNode)
-			{
-				Controller->ClearNodeSelection(true);
-				Controller->SelectNode(ModelNode, true, true);
-			}
+		if(bIsUserFacingNode)
+		{
 			Controller->CloseUndoBracket();
 		}
 		else
 		{
-			Controller->CancelUndoBracket();
+			Controller->RemoveNode(ModelNode, false);
 		}
 	}
-
-	return NewNode;
-}
-
-bool UControlRigRerouteNodeSpawner::IsTemplateNodeFilteredOut(FBlueprintActionFilter const& Filter) const
-{
-	for (UBlueprint* Blueprint : Filter.Context.Blueprints)
+	else
 	{
-		if (UControlRigBlueprint* RigBlueprint = Cast<UControlRigBlueprint>(Blueprint))
+		if(bIsUserFacingNode)
 		{
-			if (Filter.Context.Pins.Num() == 0)
-			{
-				return true;
-			}
-			for (UEdGraphPin* Pin : Filter.Context.Pins)
-			{
-				if (Pin->PinType.ContainerType != EPinContainerType::None && Pin->PinType.ContainerType != EPinContainerType::Array)
-				{
-					return true;
-				}
-			}
-
-			return false;
+			Controller->CancelUndoBracket();
 		}
 	}
-	return true;
+
+	return NewNode;
 }
 
 #undef LOCTEXT_NAMESPACE
diff --git a/Engine/Plugins/Animation/ControlRig/Source/ControlRigEditor/Private/Graph/NodeSpawners/ControlRigRerouteNodeSpawner.h b/Engine/Plugins/Animation/ControlRig/Source/ControlRigEditor/Private/Graph/NodeSpawners/ControlRigRerouteNodeSpawner.h
--- a/Engine/Plugins/Animation/ControlRig/Source/ControlRigEditor/Private/Graph/NodeSpawners/ControlRigRerouteNodeSpawner.h
+++ b/Engine/Plugins/Animation/ControlRig/Source/ControlRigEditor/Private/Graph/NodeSpawners/ControlRigRerouteNodeSpawner.h
@@ -35,8 +35,6 @@ public:
 	virtual FBlueprintNodeSignature GetSpawnerSignature() const override;
 	virtual FBlueprintActionUiSpec GetUiSpec(FBlueprintActionContext const& Context, FBindingSet const& Bindings) const override;
 	virtual UEdGraphNode* Invoke(UEdGraph* ParentGraph, FBindingSet const& Bindings, FVector2D const Location) const override;
-	virtual bool IsTemplateNodeFilteredOut(FBlueprintActionFilter const& Filter) const override;
-
 	// End UBlueprintNodeSpawner interface
 
 private:
diff --git a/Engine/Plugins/Animation/ControlRig/Source/ControlRigEditor/Private/Graph/NodeSpawners/ControlRigSelectNodeSpawner.cpp b/Engine/Plugins/Animation/ControlRig/Source/ControlRigEditor/Private/Graph/NodeSpawners/ControlRigSelectNodeSpawner.cpp
--- a/Engine/Plugins/Animation/ControlRig/Source/ControlRigEditor/Private/Graph/NodeSpawners/ControlRigSelectNodeSpawner.cpp
+++ b/Engine/Plugins/Animation/ControlRig/Source/ControlRigEditor/Private/Graph/NodeSpawners/ControlRigSelectNodeSpawner.cpp
@@ -55,31 +55,7 @@ UEdGraphNode* UControlRigSelectNodeSpawner::Invoke(UEdGraph* ParentGraph, FBindi
 	UControlRigGraphNode* NewNode = nullptr;
 
 	bool const bIsTemplateNode = FBlueprintNodeTemplateCache::IsTemplateOuter(ParentGraph);
-
-	if (bIsTemplateNode)
-	{
-		NewNode = NewObject<UControlRigGraphNode>(ParentGraph, TEXT("SelectNode"));
-		ParentGraph->AddNode(NewNode, false);
-
-		NewNode->CreateNewGuid();
-		NewNode->PostPlacedNewNode();
-
-		UEdGraphPin* InputValuePin = UEdGraphPin::CreatePin(NewNode);
-		UEdGraphPin* OutputValuePin = UEdGraphPin::CreatePin(NewNode);
-		NewNode->Pins.Add(InputValuePin);
-		NewNode->Pins.Add(OutputValuePin);
-
-		InputValuePin->PinType.PinCategory = TEXT("ANY_TYPE");
-		OutputValuePin->PinType.PinCategory = TEXT("ANY_TYPE");
-		InputValuePin->Direction = EGPD_Input;
-		OutputValuePin->Direction = EGPD_Output;
-		NewNode->SetFlags(RF_Transactional);
-
-		return NewNode;
-	}
-
-
-	bool const bUndo = !bIsTemplateNode;
+	bool const bIsUserFacingNode = !bIsTemplateNode;
 
 	// First create a backing member for our node
 	UControlRigGraph* RigGraph = Cast<UControlRigGraph>(ParentGraph);
@@ -87,29 +63,24 @@ UEdGraphNode* UControlRigSelectNodeSpawner::Invoke(UEdGraph* ParentGraph, FBindi
 	UControlRigBlueprint* RigBlueprint = Cast<UControlRigBlueprint>(FBlueprintEditorUtils::FindBlueprintForGraph(ParentGraph));
 	check(RigBlueprint);
 
-	FString CPPType = FRigVMUnknownType::StaticStruct()->GetStructCPPName();
-	FName CPPTypeObjectPath = *FRigVMUnknownType::StaticStruct()->GetPathName();
+	FString CPPType = RigVMTypeUtils::GetWildCardCPPType();
+	FName CPPTypeObjectPath = *RigVMTypeUtils::GetWildCardCPPTypeObject()->GetPathName();
 	FString LastOutputPinPath;
 
-	if (UControlRigGraphSchema* RigSchema = Cast<UControlRigGraphSchema>((UEdGraphSchema*)ParentGraph->GetSchema()))
+	if(bIsUserFacingNode)
 	{
-		if (const UEdGraphPin* LastPin = RigSchema->LastPinForCompatibleCheck)
+		if (UControlRigGraphSchema* RigSchema = Cast<UControlRigGraphSchema>((UEdGraphSchema*)ParentGraph->GetSchema()))
 		{
-			if (URigVMPin* ModelPin = RigBlueprint->GetModel(ParentGraph)->FindPin(LastPin->GetName()))
+			if (const UEdGraphPin* LastPin = RigSchema->LastPinForCompatibleCheck)
 			{
-				if(LastPin->Direction == EGPD_Output)
-				{
-					LastOutputPinPath = ModelPin->GetPinPath();
-				}
-				
-				CPPType = ModelPin->GetCPPType();
-				if (ModelPin->GetCPPTypeObject())
+				if (URigVMPin* ModelPin = RigBlueprint->GetModel(ParentGraph)->FindPin(LastPin->GetName()))
 				{
-					CPPTypeObjectPath = *ModelPin->GetCPPTypeObject()->GetPathName();
-				}
-				else
-				{
-					CPPTypeObjectPath = NAME_None;
+					if(LastPin->Direction == EGPD_Output)
+					{
+						LastOutputPinPath = ModelPin->GetPinPath();
+					}
+					
+					RigVMTypeUtils::CPPTypeFromPin(ModelPin, CPPType, CPPTypeObjectPath, false);
 				}
 			}
 		}
@@ -118,7 +89,7 @@ UEdGraphNode* UControlRigSelectNodeSpawner::Invoke(UEdGraph* ParentGraph, FBindi
 	FName MemberName = NAME_None;
 
 #if WITH_EDITOR
-	if (GEditor && !bIsTemplateNode)
+	if (GEditor && bIsUserFacingNode)
 	{
 		GEditor->CancelTransaction(0);
 	}
@@ -128,16 +99,16 @@ UEdGraphNode* UControlRigSelectNodeSpawner::Invoke(UEdGraph* ParentGraph, FBindi
 
 	FName Name = *URigVMSelectNode::SelectName;
 
-	if (!bIsTemplateNode)
+	if (bIsUserFacingNode)
 	{
 		Controller->OpenUndoBracket(FString::Printf(TEXT("Add '%s' Node"), *Name.ToString()));
 	}
 
-	if (URigVMSelectNode* ModelNode = Controller->AddSelectNode(CPPType, CPPTypeObjectPath, Location, Name.ToString(), bUndo, !bIsTemplateNode))
+	if (URigVMSelectNode* ModelNode = Controller->AddSelectNode(CPPType, CPPTypeObjectPath, Location, Name.ToString(), bIsUserFacingNode, !bIsTemplateNode))
 	{
 		NewNode = Cast<UControlRigGraphNode>(RigGraph->FindNodeForModelNodeName(ModelNode->GetFName()));
 
-		if (NewNode && bUndo)
+		if (NewNode && bIsUserFacingNode)
 		{
 			if(!LastOutputPinPath.IsEmpty())
 			{
@@ -150,14 +121,18 @@ UEdGraphNode* UControlRigSelectNodeSpawner::Invoke(UEdGraph* ParentGraph, FBindi
 			Controller->SelectNode(ModelNode, true, true);
 		}
 
-		if (bUndo)
+		if (bIsUserFacingNode)
 		{
 			Controller->CloseUndoBracket();
 		}
+		else
+		{
+			Controller->RemoveNode(ModelNode, false);
+		}
 	}
 	else
 	{
-		if (bUndo)
+		if (bIsUserFacingNode)
 		{
 			Controller->CancelUndoBracket();
 		}
diff --git a/Engine/Plugins/Animation/ControlRig/Source/ControlRigEditor/Private/Graph/NodeSpawners/ControlRigTemplateNodeSpawner.cpp b/Engine/Plugins/Animation/ControlRig/Source/ControlRigEditor/Private/Graph/NodeSpawners/ControlRigTemplateNodeSpawner.cpp
new file mode 100644
--- /dev/null
+++ b/Engine/Plugins/Animation/ControlRig/Source/ControlRigEditor/Private/Graph/NodeSpawners/ControlRigTemplateNodeSpawner.cpp
@@ -0,0 +1,164 @@
+// Copyright Epic Games, Inc. All Rights Reserved.
+
+#include "ControlRigTemplateNodeSpawner.h"
+#include "ControlRigUnitNodeSpawner.h"
+#include "Graph/ControlRigGraph.h"
+#include "Graph/ControlRigGraphNode.h"
+#include "Graph/ControlRigGraphSchema.h"
+#include "EdGraphSchema_K2.h"
+#include "Kismet2/BlueprintEditorUtils.h"
+#include "Classes/EditorStyleSettings.h"
+#include "Editor/EditorEngine.h"
+#include "ObjectEditorUtils.h"
+#include "EditorCategoryUtils.h"
+#include "K2Node_Variable.h"
+#include "BlueprintNodeTemplateCache.h"
+#include "ControlRigBlueprintUtils.h"
+#include "ScopedTransaction.h"
+#include "ControlRig.h"
+
+#if WITH_EDITOR
+#include "Editor.h"
+#endif
+
+#define LOCTEXT_NAMESPACE "ControlRigTemplateNodeSpawner"
+
+UControlRigTemplateNodeSpawner* UControlRigTemplateNodeSpawner::CreateFromNotation(const FName& InNotation, const FText& InMenuDesc, const FText& InCategory, const FText& InTooltip)
+{
+	UControlRigTemplateNodeSpawner* NodeSpawner = NewObject<UControlRigTemplateNodeSpawner>(GetTransientPackage());
+	NodeSpawner->TemplateNotation = InNotation;
+	NodeSpawner->NodeClass = UControlRigGraphNode::StaticClass();
+
+	FBlueprintActionUiSpec& MenuSignature = NodeSpawner->DefaultMenuSignature;
+	
+	MenuSignature.MenuName = InMenuDesc;
+	MenuSignature.Tooltip  = InTooltip;
+	MenuSignature.Category = InCategory;
+
+	if (const FRigVMTemplate* Template = FRigVMRegistry::Get().FindTemplate(InNotation))
+	{
+#if WITH_EDITOR
+		FString KeywordsMetadata = Template->GetKeywords();
+		MenuSignature.Keywords = FText::FromString(KeywordsMetadata);
+#endif
+	}
+
+	// add at least one character, so that PrimeDefaultUiSpec() doesn't 
+	// attempt to query the template node
+	//
+	// @TODO: maybe UPROPERTY() fields should have keyword metadata like functions
+	if (MenuSignature.Keywords.IsEmpty())
+	{
+		// want to set it to something so we won't end up back in this condition
+		MenuSignature.Keywords = FText::FromString(TEXT(" "));
+	}
+
+	// @TODO: should use details customization-like extensibility system to provide editor only data like this
+	MenuSignature.Icon = FSlateIcon(TEXT("ControlRigEditorStyle"), TEXT("ControlRig.RigUnit"));
+
+	return NodeSpawner;
+}
+
+void UControlRigTemplateNodeSpawner::Prime()
+{
+	// we expect that you don't need a node template to construct menu entries
+	// from this, so we choose not to pre-cache one here
+}
+
+FBlueprintNodeSignature UControlRigTemplateNodeSpawner::GetSpawnerSignature() const
+{
+	const int32 NotationHash = (int32)GetTypeHash(TemplateNotation);
+	return FBlueprintNodeSignature(FString("RigVMTemplate_") + FString::FromInt(NotationHash));
+}
+
+FBlueprintActionUiSpec UControlRigTemplateNodeSpawner::GetUiSpec(FBlueprintActionContext const& Context, FBindingSet const& Bindings) const
+{
+	UEdGraph* TargetGraph = (Context.Graphs.Num() > 0) ? Context.Graphs[0] : nullptr;
+	FBlueprintActionUiSpec MenuSignature = PrimeDefaultUiSpec(TargetGraph);
+
+	DynamicUiSignatureGetter.ExecuteIfBound(Context, Bindings, &MenuSignature);
+	return MenuSignature;
+}
+
+UEdGraphNode* UControlRigTemplateNodeSpawner::Invoke(UEdGraph* ParentGraph, FBindingSet const& Bindings, FVector2D const Location) const
+{
+	UControlRigGraphNode* NewNode = nullptr;
+
+	if(!TemplateNotation.IsNone())
+	{
+#if WITH_EDITOR
+		if (GEditor)
+		{
+			GEditor->CancelTransaction(0);
+		}
+#endif
+
+		UBlueprint* Blueprint = FBlueprintEditorUtils::FindBlueprintForGraph(ParentGraph);
+		NewNode = SpawnNode(ParentGraph, Blueprint, TemplateNotation, Location);
+	}
+
+	return NewNode;
+}
+
+UControlRigGraphNode* UControlRigTemplateNodeSpawner::SpawnNode(UEdGraph* ParentGraph, UBlueprint* Blueprint, const FName& InNotation, FVector2D const Location)
+{
+	UControlRigGraphNode* NewNode = nullptr;
+	UControlRigBlueprint* RigBlueprint = Cast<UControlRigBlueprint>(Blueprint);
+	UControlRigGraph* RigGraph = Cast<UControlRigGraph>(ParentGraph);
+
+	if (RigBlueprint != nullptr && RigGraph != nullptr)
+	{
+		bool const bIsTemplateNode = FBlueprintNodeTemplateCache::IsTemplateOuter(ParentGraph);
+		bool const bIsUserFacingNode = !bIsTemplateNode;
+
+		const FRigVMTemplate* Template = FRigVMRegistry::Get().FindTemplate(InNotation);
+		if (Template == nullptr)
+		{
+			return nullptr;
+		}
+
+		const int32 NotationHash = (int32)GetTypeHash(InNotation);
+		const FString TemplateName = TEXT("RigVMTemplate_") + FString::FromInt(NotationHash);
+
+		FName Name = bIsTemplateNode ? *TemplateName : FControlRigBlueprintUtils::ValidateName(RigBlueprint, Template->GetName().ToString());
+		URigVMController* Controller = bIsTemplateNode ? RigGraph->GetTemplateController() : RigBlueprint->GetController(ParentGraph);
+
+		if (!bIsTemplateNode)
+		{
+			Controller->OpenUndoBracket(FString::Printf(TEXT("Add '%s' Node"), *Name.ToString()));
+		}
+
+		if (URigVMTemplateNode* ModelNode = Controller->AddTemplateNode(InNotation, Location, Name.ToString(), bIsUserFacingNode, !bIsTemplateNode))
+		{
+			NewNode = Cast<UControlRigGraphNode>(RigGraph->FindNodeForModelNodeName(ModelNode->GetFName()));
+
+			if (NewNode && bIsUserFacingNode)
+			{
+				Controller->ClearNodeSelection(true);
+				Controller->SelectNode(ModelNode, true, true);
+
+				UControlRigUnitNodeSpawner::HookupMutableNode(ModelNode, RigBlueprint);
+			}
+
+			if (bIsUserFacingNode)
+			{
+				Controller->CloseUndoBracket();
+			}
+			else
+			{
+				NewNode->ModelNodePath = Template->GetNotation().ToString();
+				Controller->RemoveNode(ModelNode, false);
+			}
+		}
+		else
+		{
+			if (bIsUserFacingNode)
+			{
+				Controller->CancelUndoBracket();
+			}
+		}
+	}
+	return NewNode;
+}
+
+#undef LOCTEXT_NAMESPACE
diff --git a/Engine/Plugins/Animation/ControlRig/Source/ControlRigEditor/Private/Graph/NodeSpawners/ControlRigTemplateNodeSpawner.h b/Engine/Plugins/Animation/ControlRig/Source/ControlRigEditor/Private/Graph/NodeSpawners/ControlRigTemplateNodeSpawner.h
new file mode 100644
--- /dev/null
+++ b/Engine/Plugins/Animation/ControlRig/Source/ControlRigEditor/Private/Graph/NodeSpawners/ControlRigTemplateNodeSpawner.h
@@ -0,0 +1,53 @@
+// Copyright Epic Games, Inc. All Rights Reserved.
+
+#pragma once
+
+#include "CoreMinimal.h"
+#include "UObject/ObjectMacros.h"
+#include "UObject/Object.h"
+#include "Templates/SubclassOf.h"
+#include "UObject/UnrealType.h"
+#include "EdGraph/EdGraphPin.h"
+#include "Engine/Blueprint.h"
+#include "BlueprintActionFilter.h"
+#include "BlueprintNodeSignature.h"
+#include "BlueprintFieldNodeSpawner.h"
+#include "ControlRigBlueprint.h"
+#include "Graph/ControlRigGraph.h"
+#include "Graph/ControlRigGraphNode.h"
+#include "RigVMModel/Nodes/RigVMTemplateNode.h"
+#include "ControlRigTemplateNodeSpawner.generated.h"
+
+class UControlRigGraphNode;
+
+UCLASS(Transient)
+class CONTROLRIGEDITOR_API UControlRigTemplateNodeSpawner : public UBlueprintNodeSpawner
+{
+	GENERATED_BODY()
+
+public:
+	/**
+	 * Creates a new UControlRigVariableNodeSpawner, charged with spawning 
+	 * a new member-variable node
+	 * 
+	 * @return A newly allocated instance of this class.
+	 */
+	static UControlRigTemplateNodeSpawner* CreateFromNotation(const FName& InNotation, const FText& InMenuDesc, const FText& InCategory, const FText& InTooltip);
+
+	// UBlueprintNodeSpawner interface
+	virtual void Prime() override;
+	virtual FBlueprintNodeSignature GetSpawnerSignature() const override;
+	virtual FBlueprintActionUiSpec GetUiSpec(FBlueprintActionContext const& Context, FBindingSet const& Bindings) const override;
+	virtual UEdGraphNode* Invoke(UEdGraph* ParentGraph, FBindingSet const& Bindings, FVector2D const Location) const override;
+	// End UBlueprintNodeSpawner interface
+
+private:
+	/** The unit type we will spawn */
+	UPROPERTY(Transient)
+	FName TemplateNotation;
+
+	static UControlRigGraphNode* SpawnNode(UEdGraph* ParentGraph, UBlueprint* Blueprint, const FName& InNotation, FVector2D const Location);
+
+	friend class UEngineTestControlRig;
+	friend class FControlRigEditor;
+};
diff --git a/Engine/Plugins/Animation/ControlRig/Source/ControlRigEditor/Private/Graph/NodeSpawners/ControlRigUnitNodeSpawner.cpp b/Engine/Plugins/Animation/ControlRig/Source/ControlRigEditor/Private/Graph/NodeSpawners/ControlRigUnitNodeSpawner.cpp
--- a/Engine/Plugins/Animation/ControlRig/Source/ControlRigEditor/Private/Graph/NodeSpawners/ControlRigUnitNodeSpawner.cpp
+++ b/Engine/Plugins/Animation/ControlRig/Source/ControlRigEditor/Private/Graph/NodeSpawners/ControlRigUnitNodeSpawner.cpp
@@ -37,17 +37,17 @@ UControlRigUnitNodeSpawner* UControlRigUnitNodeSpawner::CreateFromStruct(UScript
 	MenuSignature.Tooltip  = InTooltip;
 	MenuSignature.Category = InCategory;
 
-	FString KeywordsMetadata, PrototypeNameMetadata;
+	FString KeywordsMetadata, TemplateNameMetadata;
 	InStruct->GetStringMetaDataHierarchical(FRigVMStruct::KeywordsMetaName, &KeywordsMetadata);
-	if(!PrototypeNameMetadata.IsEmpty())
+	if(!TemplateNameMetadata.IsEmpty())
 	{
 		if(KeywordsMetadata.IsEmpty())
 		{
-			KeywordsMetadata = PrototypeNameMetadata;
+			KeywordsMetadata = TemplateNameMetadata;
 		}
 		else
 		{
-			KeywordsMetadata = KeywordsMetadata + TEXT(",") + PrototypeNameMetadata;
+			KeywordsMetadata = KeywordsMetadata + TEXT(",") + TemplateNameMetadata;
 		}
 	}
 	MenuSignature.Keywords = FText::FromString(KeywordsMetadata);
@@ -108,20 +108,6 @@ UEdGraphNode* UControlRigUnitNodeSpawner::Invoke(UEdGraph* ParentGraph, FBinding
 	return NewNode;
 }
 
-bool UControlRigUnitNodeSpawner::IsTemplateNodeFilteredOut(FBlueprintActionFilter const& Filter) const
-{
-	if (StructTemplate)
-	{
-		FString DeprecatedMetadata;
-		StructTemplate->GetStringMetaDataHierarchical(FRigVMStruct::DeprecatedMetaName, &DeprecatedMetadata);
-		if (!DeprecatedMetadata.IsEmpty())
-		{
-			return true;
-		}
-	}
-	return Super::IsTemplateNodeFilteredOut(Filter);
-}
-
 UControlRigGraphNode* UControlRigUnitNodeSpawner::SpawnNode(UEdGraph* ParentGraph, UBlueprint* Blueprint, UScriptStruct* StructTemplate, FVector2D const Location)
 {
 	UControlRigGraphNode* NewNode = nullptr;
@@ -133,7 +119,8 @@ UControlRigGraphNode* UControlRigUnitNodeSpawner::SpawnNode(UEdGraph* ParentGrap
 		bool const bIsTemplateNode = FBlueprintNodeTemplateCache::IsTemplateOuter(ParentGraph);
 		bool const bIsUserFacingNode = !bIsTemplateNode;
 
-		FName Name = bIsTemplateNode ? *StructTemplate->GetDisplayNameText().ToString() : FControlRigBlueprintUtils::ValidateName(RigBlueprint, StructTemplate->GetFName().ToString());
+		FName Name = bIsTemplateNode ? *StructTemplate->GetStructCPPName() : FControlRigBlueprintUtils::ValidateName(RigBlueprint, StructTemplate->GetFName().ToString());
+		
 		URigVMController* Controller = bIsTemplateNode ? RigGraph->GetTemplateController() : RigBlueprint->GetController(ParentGraph);
 
 		if (!bIsTemplateNode)
diff --git a/Engine/Plugins/Animation/ControlRig/Source/ControlRigEditor/Private/Graph/NodeSpawners/ControlRigUnitNodeSpawner.h b/Engine/Plugins/Animation/ControlRig/Source/ControlRigEditor/Private/Graph/NodeSpawners/ControlRigUnitNodeSpawner.h
--- a/Engine/Plugins/Animation/ControlRig/Source/ControlRigEditor/Private/Graph/NodeSpawners/ControlRigUnitNodeSpawner.h
+++ b/Engine/Plugins/Animation/ControlRig/Source/ControlRigEditor/Private/Graph/NodeSpawners/ControlRigUnitNodeSpawner.h
@@ -39,7 +39,6 @@ public:
 	virtual FBlueprintNodeSignature GetSpawnerSignature() const override;
 	virtual FBlueprintActionUiSpec GetUiSpec(FBlueprintActionContext const& Context, FBindingSet const& Bindings) const override;
 	virtual UEdGraphNode* Invoke(UEdGraph* ParentGraph, FBindingSet const& Bindings, FVector2D const Location) const override;
-	virtual bool IsTemplateNodeFilteredOut(FBlueprintActionFilter const& Filter) const override;
 	static void HookupMutableNode(URigVMNode* InModelNode, UControlRigBlueprint* InRigBlueprint);
 	// End UBlueprintNodeSpawner interface
 
diff --git a/Engine/Plugins/Animation/ControlRig/Source/ControlRigEditor/Private/Graph/SControlRigGraphNode.cpp b/Engine/Plugins/Animation/ControlRig/Source/ControlRigEditor/Private/Graph/SControlRigGraphNode.cpp
--- a/Engine/Plugins/Animation/ControlRig/Source/ControlRigEditor/Private/Graph/SControlRigGraphNode.cpp
+++ b/Engine/Plugins/Animation/ControlRig/Source/ControlRigEditor/Private/Graph/SControlRigGraphNode.cpp
@@ -1073,10 +1073,10 @@ TArray<FOverlayWidgetInfo> SControlRigGraphNode::GetOverlayWidgets(bool bSelecte
 				{
 					if (URigVMUnitNode* VisualDebugNode = Cast<URigVMUnitNode>(Injection->Node))
 					{
-						FString PrototypeName;
-					   if (VisualDebugNode->GetScriptStruct()->GetStringMetaDataHierarchical(TEXT("PrototypeName"), &PrototypeName))
+						FString TemplateName;
+					   if (VisualDebugNode->GetScriptStruct()->GetStringMetaDataHierarchical(FRigVMStruct::TemplateNameMetaName, &TemplateName))
 					   {
-						   if (PrototypeName == TEXT("VisualDebug"))
+						   if (TemplateName == TEXT("VisualDebug"))
 						   {
 							   if (!bSetColor)
 							   {
diff --git a/Engine/Source/Developer/RigVMDeveloper/Private/RigVMCompiler/RigVMAST.cpp b/Engine/Source/Developer/RigVMDeveloper/Private/RigVMCompiler/RigVMAST.cpp
--- a/Engine/Source/Developer/RigVMDeveloper/Private/RigVMCompiler/RigVMAST.cpp
+++ b/Engine/Source/Developer/RigVMDeveloper/Private/RigVMCompiler/RigVMAST.cpp
@@ -1793,7 +1793,12 @@ bool FRigVMParserAST::FoldConstantValuesToLiterals(URigVMGraph* InGraph, URigVMC
 		}
 
 		FString PinHash = URigVMCompiler::GetPinHash(RootPin, RootVarExpr, false);
-		const FRigVMOperand& Operand = Operands.FindChecked(PinHash);
+		const FRigVMOperand* OperandPtr = Operands.Find(PinHash);
+		if(OperandPtr == nullptr)
+		{
+			continue;
+		}
+		const FRigVMOperand& Operand = *OperandPtr;
 
 		FString DefaultValue;
 		if(Operand.GetMemoryType() == ERigVMMemoryType::Work)
diff --git a/Engine/Source/Developer/RigVMDeveloper/Private/RigVMCompiler/RigVMCompiler.cpp b/Engine/Source/Developer/RigVMDeveloper/Private/RigVMCompiler/RigVMCompiler.cpp
--- a/Engine/Source/Developer/RigVMDeveloper/Private/RigVMCompiler/RigVMCompiler.cpp
+++ b/Engine/Source/Developer/RigVMDeveloper/Private/RigVMCompiler/RigVMCompiler.cpp
@@ -773,6 +773,11 @@ int32 URigVMCompiler::TraverseCallExtern(const FRigVMCallExternExprAST* InExpr,
 		return INDEX_NONE;
 	}
 
+	if(UnitNode->GetScriptStruct() == nullptr)
+	{
+		return INDEX_NONE;
+	}
+
 	int32 InstructionIndex = INDEX_NONE;
 
 	if (WorkData.bSetupMemory)
@@ -2661,9 +2666,9 @@ bool URigVMCompiler::ValidateNode(URigVMNode* InNode)
 {
 	check(InNode);
 
-	if(InNode->HasUnknownTypePin())
+	if(InNode->HasWildCardPin())
 	{
-		static const FString UnknownTypeMessage = TEXT("Node @@ has unresolved pins of unknown type.");
+		static const FString UnknownTypeMessage = TEXT("Node @@ has unresolved pins of wildcard type.");
 		Settings.Report(EMessageSeverity::Error, InNode, UnknownTypeMessage);
 		return false;
 	}
diff --git a/Engine/Source/Developer/RigVMDeveloper/Private/RigVMDeveloperTypeUtils.cpp b/Engine/Source/Developer/RigVMDeveloper/Private/RigVMDeveloperTypeUtils.cpp
--- a/Engine/Source/Developer/RigVMDeveloper/Private/RigVMDeveloperTypeUtils.cpp
+++ b/Engine/Source/Developer/RigVMDeveloper/Private/RigVMDeveloperTypeUtils.cpp
@@ -182,7 +182,7 @@ FRigVMExternalVariable RigVMTypeUtils::ExternalVariableFromCPPTypePath(const FNa
 	}
 	else if (UEnum* Enum= URigVMPin::FindObjectFromCPPTypeObjectPath<UEnum>(CPPTypePath))
 	{
-		Variable.TypeName = *Enum->CppType;
+		Variable.TypeName = *CPPTypeFromEnum(Enum);
 		Variable.TypeObject = Enum;
 		Variable.Size = Enum->GetResourceSizeBytes(EResourceSizeMode::EstimatedTotal);
 	}
@@ -251,7 +251,7 @@ FRigVMExternalVariable RigVMTypeUtils::ExternalVariableFromCPPType(const FName&
 	}
 	else if (UEnum* Enum= URigVMPin::FindObjectFromCPPTypeObjectPath<UEnum>(CPPType))
 	{
-		Variable.TypeName = *Enum->CppType;
+		Variable.TypeName = *CPPTypeFromEnum(Enum);
 		Variable.TypeObject = Enum;
 		Variable.Size = Enum->GetResourceSizeBytes(EResourceSizeMode::EstimatedTotal);
 	}
@@ -486,6 +486,53 @@ FEdGraphPinType RigVMTypeUtils::SubPinType(const FEdGraphPinType& InPinType, con
 	return Result;
 }
 
+bool RigVMTypeUtils::CPPTypeFromPin(URigVMPin* InPin, FString& OutCPPType, UObject** OutCPPTypeObject, bool bGetBaseCPPType)
+{
+	check(InPin);
+	OutCPPType = InPin->GetCPPType();
+	if(bGetBaseCPPType && IsArrayType(OutCPPType))
+	{
+		OutCPPType = BaseTypeFromArrayType(OutCPPType);
+	}
+	if(OutCPPTypeObject)
+	{
+		*OutCPPTypeObject = InPin->GetCPPTypeObject();
+	}
+	return true;
+}
+
+bool RigVMTypeUtils::CPPTypeFromPin(URigVMPin* InPin, FString& OutCPPType, FName& OutCPPTypeObjectPath, bool bGetBaseCPPType)
+{
+	OutCPPType = FString();
+	OutCPPTypeObjectPath = NAME_None;
+	UObject* CPPTypeObject = nullptr;
+	if (RigVMTypeUtils::CPPTypeFromPin(InPin, OutCPPType, &CPPTypeObject, bGetBaseCPPType))
+	{
+		if (CPPTypeObject)
+		{
+			OutCPPTypeObjectPath = *CPPTypeObject->GetPathName();
+		}
+		return true;
+	}
+	return false;
+}
+
+bool RigVMTypeUtils::CPPTypeFromPin(URigVMPin* InPin, FString& OutCPPType, FString& OutCPPTypeObjectPath, bool bGetBaseCPPType)
+{
+	OutCPPType = FString();
+	OutCPPTypeObjectPath.Reset();
+	FName CPPTypeObjectPath(NAME_None);
+	if (RigVMTypeUtils::CPPTypeFromPin(InPin, OutCPPType, CPPTypeObjectPath, bGetBaseCPPType))
+	{
+		if (!CPPTypeObjectPath.IsNone())
+		{
+			OutCPPTypeObjectPath = CPPTypeObjectPath.ToString();
+		}
+		return true;
+	}
+	return false;
+}
+
 bool RigVMTypeUtils::CPPTypeFromPinType(const FEdGraphPinType& InPinType, FString& OutCPPType, UObject** OutCPPTypeObject)
 {
 	FString Prefix = "";
@@ -604,6 +651,23 @@ bool RigVMTypeUtils::CPPTypeFromPinType(const FEdGraphPinType& InPinType, FStrin
 	return false;
 }
 
+bool RigVMTypeUtils::CPPTypeFromPinType(const FEdGraphPinType& InPinType, FString& OutCPPType, FString& OutCPPTypeObjectPath)
+{
+	OutCPPType = FString();
+	OutCPPTypeObjectPath.Reset();
+	FName CPPTypeObjectPath(NAME_None);
+	if (RigVMTypeUtils::CPPTypeFromPinType(InPinType, OutCPPType, CPPTypeObjectPath))
+	{
+		if(!CPPTypeObjectPath.IsNone())
+		{
+			OutCPPTypeObjectPath = CPPTypeObjectPath.ToString();
+		}
+		return true;
+	}
+
+	return false;
+}
+
 bool RigVMTypeUtils::CPPTypeFromExternalVariable(const FRigVMExternalVariable& InExternalVariable, FString& OutCPPType, UObject** OutCPPTypeObject)
 {
 	FString Prefix = "";
diff --git a/Engine/Source/Developer/RigVMDeveloper/Private/RigVMModel/Nodes/RigVMPrototypeNode.cpp b/Engine/Source/Developer/RigVMDeveloper/Private/RigVMModel/Nodes/RigVMPrototypeNode.cpp
deleted file mode 100644
diff --git a/Engine/Source/Developer/RigVMDeveloper/Private/RigVMModel/Nodes/RigVMTemplateNode.cpp b/Engine/Source/Developer/RigVMDeveloper/Private/RigVMModel/Nodes/RigVMTemplateNode.cpp
new file mode 100644
--- /dev/null
+++ b/Engine/Source/Developer/RigVMDeveloper/Private/RigVMModel/Nodes/RigVMTemplateNode.cpp
@@ -0,0 +1,261 @@
+// Copyright Epic Games, Inc. All Rights Reserved.
+
+#include "RigVMModel/Nodes/RigVMTemplateNode.h"
+
+URigVMTemplateNode::URigVMTemplateNode()
+	: Super()
+	, TemplateNotation(NAME_None)
+	, ResolvedFunctionName()
+	, CachedTemplate(nullptr)
+	, CachedFunction(nullptr)
+{
+}
+
+UScriptStruct* URigVMTemplateNode::GetScriptStruct() const
+{
+	if(const FRigVMFunction* Function = GetResolvedFunction())
+	{
+		return Function->Struct;
+	}
+	return nullptr;
+}
+
+FString URigVMTemplateNode::GetNodeTitle() const
+{
+	if(!IsResolved())
+	{
+		if(const FRigVMTemplate* Template = GetTemplate())
+		{
+			return Template->GetName().ToString();
+		}
+	}
+	
+	FString ResolvedNodeTitle = Super::GetNodeTitle();
+
+	const int32 BracePos = ResolvedNodeTitle.Find(TEXT(" ("));
+	if(BracePos != INDEX_NONE)
+	{
+		ResolvedNodeTitle = ResolvedNodeTitle.Left(BracePos);
+	}
+
+	return ResolvedNodeTitle;
+}
+
+FName URigVMTemplateNode::GetMethodName() const
+{
+	if(const FRigVMFunction* Function = GetResolvedFunction())
+	{
+		return Function->GetMethodName();
+	}
+	return NAME_None;
+}
+
+FText URigVMTemplateNode::GetToolTipText() const
+{
+	if(const FRigVMTemplate* Template = GetTemplate())
+	{
+		const TArray<int32> PermutationIndices = GetResolvedPermutationIndices();
+		return Template->GetTooltipText(PermutationIndices);
+	}
+	return Super::GetToolTipText();
+}
+
+FText URigVMTemplateNode::GetToolTipTextForPin(const URigVMPin* InPin) const
+{
+	const FText SuperToolTip = Super::GetToolTipTextForPin(InPin);
+	if (const FRigVMTemplate* Template = GetTemplate())
+	{
+		const URigVMPin* RootPin = InPin->GetRootPin();
+		if (RootPin->IsWildCard())
+		{
+			if (const FRigVMTemplateArgument* Arg = Template->FindArgument(RootPin->GetFName()))
+			{
+				const TArray<int32> PermutationIndices = GetResolvedPermutationIndices();
+				const FString SupportedTypesJoined = FString::Join(Arg->GetSupportedTypeStrings(PermutationIndices), TEXT("\n"));
+				FString Tooltip = TEXT("Supported Types:\n\n") + SupportedTypesJoined;
+
+				if(!SuperToolTip.IsEmpty())
+				{
+					Tooltip += TEXT("\n\n") + SuperToolTip.ToString();
+				}
+				
+				return FText::FromString(Tooltip);
+			}
+		}
+	}
+	return SuperToolTip;
+}
+
+FName URigVMTemplateNode::GetNotation() const
+{
+	return TemplateNotation;
+}
+
+bool URigVMTemplateNode::SupportsType(const URigVMPin* InPin, const FString& InCPPType, FString* OutCPPType)
+{
+	static const FString WildCardCPPType = RigVMTypeUtils::GetWildCardCPPType();
+	static const FString WildCardArrayCPPType = RigVMTypeUtils::ArrayTypeFromBaseType(WildCardCPPType);
+
+	// we always support the unknown type
+	if(((InCPPType == WildCardCPPType) && !InPin->IsArray()) ||
+		((InCPPType == WildCardArrayCPPType) && InPin->IsArray()))
+	{
+		if(OutCPPType)
+		{
+			*OutCPPType = InCPPType;
+		}
+		return true;
+	}
+	
+	FString CPPType = InCPPType;
+	const URigVMPin* RootPin = InPin->GetRootPin();
+
+	if(InPin->GetParentPin() == RootPin && RootPin->IsArray())
+	{
+		CPPType = RigVMTypeUtils::ArrayTypeFromBaseType(CPPType);
+	}
+	
+	if (RootPin->IsWildCard())
+	{
+		const FString CacheKey = RootPin->GetName() + TEXT("|") + CPPType;
+		if (const TPair<bool, FRigVMTemplateArgument::FType>* CachedResult = SupportedTypesCache.Find(CacheKey))
+		{
+			if(OutCPPType)
+			{
+				*OutCPPType = CachedResult->Value.CPPType;
+			}
+			return CachedResult->Key;
+		}
+
+		if (const FRigVMTemplate* Template = GetTemplate())
+		{
+			FRigVMTemplate::FTypeMap Types;
+
+			for (URigVMPin* Pin : GetPins())
+			{
+				if (Pin == RootPin)
+				{
+					continue;
+				}
+
+				if (Pin->IsWildCard())
+				{
+					continue;
+				}
+
+				FRigVMTemplateArgument::FType Type(Pin->GetCPPType(), Pin->GetCPPTypeObject());
+				Types.Add(Pin->GetFName(), Type);
+			}
+
+			FRigVMTemplateArgument::FType ResolvedType;
+			bool Result = Template->ArgumentSupportsType(RootPin->GetFName(), CPPType, Types, &ResolvedType);
+			SupportedTypesCache.Add(CacheKey, TPair<bool, FRigVMTemplateArgument::FType>(Result, ResolvedType));
+			if(OutCPPType)
+			{
+				*OutCPPType = ResolvedType.CPPType;
+			}
+			return Result;
+		}
+	}
+
+	if(RootPin->GetCPPType() == CPPType)
+	{
+		if(OutCPPType)
+		{
+			*OutCPPType = CPPType;
+		}
+		return true;
+	}
+	return false;
+}
+
+TArray<int32> URigVMTemplateNode::GetResolvedPermutationIndices(FRigVMTemplate::FTypeMap* OutTypes) const
+{
+	if (const FRigVMTemplate* Template = GetTemplate())
+	{
+		FRigVMTemplate::FTypeMap Types;
+
+		for (URigVMPin* Pin : GetPins())
+		{
+			if (Pin->IsWildCard())
+			{
+				continue;
+			}
+
+			FRigVMTemplateArgument::FType Type(Pin->GetCPPType(), Pin->GetCPPTypeObject());
+			Types.Add(Pin->GetFName(), Type);
+		}
+
+		TArray<int32> PermutationIndices;
+		Template->Resolve(Types, PermutationIndices, false);
+
+		if (OutTypes)
+		{
+			*OutTypes = Types;
+		}
+
+		return PermutationIndices;
+	}
+	return TArray<int32>();
+}
+
+TArray<const FRigVMFunction*> URigVMTemplateNode::GetResolvedPermutations(FRigVMTemplate::FTypeMap* OutTypes) const
+{
+	TArray<int32> Indices = GetResolvedPermutationIndices(OutTypes);
+	TArray<const FRigVMFunction*> Functions;
+	for(const int32 Index : Indices)
+	{
+		Functions.Add(GetTemplate()->GetPermutation(Index));
+	}
+	return Functions;
+}
+
+const FRigVMTemplate* URigVMTemplateNode::GetTemplate() const
+{
+	if(CachedTemplate == nullptr)
+	{
+		CachedTemplate = FRigVMRegistry::Get().FindTemplate(TemplateNotation);
+	}
+	return CachedTemplate;
+}
+
+const FRigVMFunction* URigVMTemplateNode::GetResolvedFunction() const
+{
+	if(CachedFunction == nullptr)
+	{
+		if(!ResolvedFunctionName.IsEmpty())
+		{
+			CachedFunction = FRigVMRegistry::Get().FindFunction(*ResolvedFunctionName);
+		}
+
+		if(CachedFunction == nullptr)
+		{
+			TArray<int32> PermutationIndices = GetResolvedPermutationIndices();
+			if(PermutationIndices.Num() == 1)
+			{
+				CachedFunction = GetTemplate()->GetPermutation(PermutationIndices[0]);
+			}
+		}
+	}
+	return CachedFunction;
+}
+
+bool URigVMTemplateNode::IsResolved() const
+{
+	return GetScriptStruct() != nullptr;
+}
+
+void URigVMTemplateNode::InvalidateCache()
+{
+	SupportedTypesCache.Reset();
+	CachedFunction = nullptr;
+
+	for(URigVMPin* Pin : GetPins())
+	{
+		if(Pin->IsWildCard())
+		{
+			ResolvedFunctionName.Reset();
+			break;
+		}
+	}
+}
diff --git a/Engine/Source/Developer/RigVMDeveloper/Private/RigVMModel/Nodes/RigVMUnitNode.cpp b/Engine/Source/Developer/RigVMDeveloper/Private/RigVMModel/Nodes/RigVMUnitNode.cpp
--- a/Engine/Source/Developer/RigVMDeveloper/Private/RigVMModel/Nodes/RigVMUnitNode.cpp
+++ b/Engine/Source/Developer/RigVMDeveloper/Private/RigVMModel/Nodes/RigVMUnitNode.cpp
@@ -158,18 +158,18 @@ FString URigVMUnitNode::GetStructDefaultValue() const
 
 TSharedPtr<FStructOnScope> URigVMUnitNode::ConstructStructInstance(bool bUseDefault) const
 {
-	if (ScriptStruct)
+	if (UScriptStruct* Struct = GetScriptStruct())
 	{
-		TSharedPtr<FStructOnScope> StructOnScope = MakeShareable(new FStructOnScope(ScriptStruct));
+		TSharedPtr<FStructOnScope> StructOnScope = MakeShareable(new FStructOnScope(Struct));
 		FRigVMStruct* StructMemory = (FRigVMStruct*)StructOnScope->GetStructMemory();
 		if (bUseDefault)
 		{
-			ScriptStruct->InitializeDefaultValue((uint8*)StructMemory);
+			Struct->InitializeDefaultValue((uint8*)StructMemory);
 		}
 		else
 		{
 			FString StructDefaultValue = GetStructDefaultValue();
-			ScriptStruct->ImportText(*StructDefaultValue, StructMemory, nullptr, PPF_None, nullptr, ScriptStruct->GetName());
+			Struct->ImportText(*StructDefaultValue, StructMemory, nullptr, PPF_None, nullptr, Struct->GetName());
 		}
 		return StructOnScope;
 	}
diff --git a/Engine/Source/Developer/RigVMDeveloper/Private/RigVMModel/RigVMController.cpp b/Engine/Source/Developer/RigVMDeveloper/Private/RigVMModel/RigVMController.cpp
--- a/Engine/Source/Developer/RigVMDeveloper/Private/RigVMModel/RigVMController.cpp
+++ b/Engine/Source/Developer/RigVMDeveloper/Private/RigVMModel/RigVMController.cpp
@@ -16,6 +16,7 @@
 #include "UObject/StrongObjectPtr.h"
 #include "Misc/CoreMisc.h"
 #include "Algo/Sort.h"
+#include "Algo/Count.h"
 #include "RigVMPythonUtils.h"
 #include "RigVMTypeUtils.h"
 #include "Engine/UserDefinedStruct.h"
@@ -63,6 +64,7 @@ URigVMController::URigVMController()
 	, bReportWarningsAndErrors(true)
 	, bIgnoreRerouteCompactnessChanges(false)
 	, bIsRunningUnitTest(false)
+	, bIsFullyResolvingTemplateNode(false)
 {
 	SetExecuteContextStruct(FRigVMExecuteContext::StaticStruct());
 }
@@ -666,13 +668,13 @@ TArray<FString> URigVMController::GetAddNodePythonCommands(URigVMNode* Node) con
 						*RigVMPythonUtils::Vector2DToPythonString(SelectNode->GetPosition()),
 						*NodeName));
 	}
-	else if (const URigVMPrototypeNode* PrototypeNode = Cast<URigVMPrototypeNode>(Node))
+	else if (const URigVMTemplateNode* TemplateNode = Cast<URigVMTemplateNode>(Node))
 	{
-		// add_prototype_node(notation, position=[0.0, 0.0], node_name='', undo=True)
-		Commands.Add(FString::Printf(TEXT("blueprint.get_controller_by_name('%s').add_prototype_node('%s', %s, '%s')"),
+		// add_template_node(notation, position=[0.0, 0.0], node_name='', undo=True)
+		Commands.Add(FString::Printf(TEXT("blueprint.get_controller_by_name('%s').add_template_node('%s', %s, '%s')"),
 						*GraphName,
-						*PrototypeNode->GetNotation().ToString(),
-						*RigVMPythonUtils::Vector2DToPythonString(PrototypeNode->GetPosition()),
+						*TemplateNode->GetNotation().ToString(),
+						*RigVMPythonUtils::Vector2DToPythonString(TemplateNode->GetPosition()),
 						*NodeName));
 	}
 	else if (const URigVMEnumNode* EnumNode = Cast<URigVMEnumNode>(Node))
@@ -823,11 +825,10 @@ URigVMUnitNode* URigVMController::AddUnitNode(UScriptStruct* InScriptStruct, con
 	URigVMGraph* Graph = GetGraph();
 	check(Graph);
 
-	FString FunctionName = FString::Printf(TEXT("F%s::%s"), *InScriptStruct->GetName(), *InMethodName.ToString());
-	FRigVMFunctionPtr Function = FRigVMRegistry::Get().FindFunction(*FunctionName);
+	const FRigVMFunction* Function = FRigVMRegistry::Get().FindFunction(InScriptStruct, *InMethodName.ToString());
 	if (Function == nullptr)
 	{
-		ReportErrorf(TEXT("RIGVM_METHOD '%s' cannot be found."), *FunctionName);
+		ReportErrorf(TEXT("RIGVM_METHOD '%s::%s' cannot be found."), *InScriptStruct->GetStructCPPName(), *InMethodName.ToString());
 		return nullptr;
 	}
 
@@ -838,6 +839,35 @@ URigVMUnitNode* URigVMController::AddUnitNode(UScriptStruct* InScriptStruct, con
 		return nullptr;
 	}
 
+#if UE_RIGVM_ENABLE_TEMPLATE_NODES
+
+	if(const FRigVMTemplate* Template = Function->GetTemplate())
+	{
+		if(bSetupUndoRedo)
+		{
+			OpenUndoBracket(FString::Printf(TEXT("Add %s Node"), *Template->GetName().ToString()));
+		}
+
+		const FString Name = GetValidNodeName(InNodeName.IsEmpty() ? InScriptStruct->GetName() : InNodeName);
+		URigVMTemplateNode* TemplateNode = AddTemplateNode(Template->GetNotation(), InPosition, Name, bSetupUndoRedo, bPrintPythonCommand);
+		if(TemplateNode == nullptr)
+		{
+			CancelUndoBracket();
+			return nullptr;
+		}
+
+		FullyResolveTemplateNode(TemplateNode, Function, bSetupUndoRedo);
+		
+		if(bSetupUndoRedo)
+		{
+			CloseUndoBracket();
+		}
+
+		return TemplateNode;
+	}
+
+#endif
+
 	FStructOnScope StructOnScope(InScriptStruct);
 	FRigVMStruct* StructMemory = (FRigVMStruct*)StructOnScope.GetStructMemory();
 	InScriptStruct->InitializeDefaultValue((uint8*)StructMemory);
@@ -1026,7 +1056,7 @@ URigVMVariableNode* URigVMController::AddVariableNode(const FName& InVariableNam
 	if (!bIsGetter)
 	{
 		URigVMPin* ExecutePin = NewObject<URigVMPin>(Node, FRigVMStruct::ExecuteContextName);
-		ExecutePin->CPPType = FString::Printf(TEXT("F%s"), *ExecuteContextStruct->GetName());
+		ExecutePin->CPPType = ExecuteContextStruct->GetStructCPPName();
 		ExecutePin->CPPTypeObject = ExecuteContextStruct;
 		ExecutePin->CPPTypeObjectPath = *ExecutePin->CPPTypeObject->GetPathName();
 		ExecutePin->Direction = ERigVMPinDirection::IO;
@@ -1034,7 +1064,7 @@ URigVMVariableNode* URigVMController::AddVariableNode(const FName& InVariableNam
 	}
 
 	URigVMPin* VariablePin = NewObject<URigVMPin>(Node, *URigVMVariableNode::VariableName);
-	VariablePin->CPPType = TEXT("FName");
+	VariablePin->CPPType = RigVMTypeUtils::FNameType;
 	VariablePin->Direction = ERigVMPinDirection::Hidden;
 	VariablePin->DefaultValue = InVariableName.ToString();
 	VariablePin->CustomWidgetName = TEXT("VariableName");
@@ -1561,6 +1591,77 @@ URigVMVariableNode* URigVMController::ReplaceParameterNodeWithVariable(const FNa
 	return nullptr;
 }
 
+URigVMTemplateNode* URigVMController::ReplaceUnitNodeWithTemplateNode(const FName& InNodeName, bool bSetupUndoRedo)
+{
+	if(!IsValidGraph())
+	{
+		return nullptr;
+	}
+
+	URigVMUnitNode* UnitNode = Cast<URigVMUnitNode>(GetGraph()->FindNodeByName(InNodeName));
+	if(UnitNode == nullptr)
+	{
+		return nullptr;
+	}
+	if(UnitNode->IsA<URigVMTemplateNode>())
+	{
+		return Cast<URigVMTemplateNode>(UnitNode);
+	}
+
+	UScriptStruct* ScriptStruct = UnitNode->GetScriptStruct();
+	if(ScriptStruct == nullptr)
+	{
+		return nullptr;
+	}
+	
+	const FRigVMFunction* Function = FRigVMRegistry::Get().FindFunction(ScriptStruct, *UnitNode->GetMethodName().ToString());
+	if(Function == nullptr)
+	{
+		return nullptr;
+	}
+			
+	const FRigVMTemplate* Template = Function->GetTemplate();
+	if(Template == nullptr)
+	{
+		return nullptr;
+	}
+
+	FRigVMBaseAction Action;
+	if (bSetupUndoRedo)
+	{
+		Action.Title = FString::Printf(TEXT("Replace Unit with Template"));
+		ActionStack->BeginAction(Action);
+	}
+
+	const TMap<FString, FPinState> PinStates = GetPinStates(UnitNode);
+	TArray<URigVMLink*> Links = UnitNode->GetLinks();
+	DetachLinksFromPinObjects(&Links);
+
+	const FVector2D NodePosition = UnitNode->GetPosition();
+	RemoveNode(UnitNode, bSetupUndoRedo);
+	URigVMTemplateNode* TemplateNode = AddTemplateNode(Template->GetNotation(), NodePosition, InNodeName.ToString(), bSetupUndoRedo);
+	if(TemplateNode == nullptr)
+	{
+		if(bSetupUndoRedo)
+		{
+			ActionStack->CancelAction(Action);
+		}
+		return nullptr;
+	}
+
+	FullyResolveTemplateNode(TemplateNode, Function, bSetupUndoRedo);
+
+	ApplyPinStates(TemplateNode, PinStates);
+	ReattachLinksToPinObjects(false, &Links);
+
+	if(bSetupUndoRedo)
+	{
+		ActionStack->EndAction(Action);
+	}
+	
+	return TemplateNode;
+}
+
 URigVMParameterNode* URigVMController::AddParameterNode(const FName& InParameterName, const FString& InCPPType, UObject* InCPPTypeObject, bool bIsInput, const FString& InDefaultValue, const FVector2D& InPosition, const FString& InNodeName, bool bSetupUndoRedo, bool bPrintPythonCommand)
 {
 	if (!IsValidGraph())
@@ -1616,7 +1717,7 @@ URigVMParameterNode* URigVMController::AddParameterNode(const FName& InParameter
 	}
 
 	URigVMPin* ParameterPin = NewObject<URigVMPin>(Node, *URigVMParameterNode::ParameterName);
-	ParameterPin->CPPType = TEXT("FName");
+	ParameterPin->CPPType = RigVMTypeUtils::FNameType;
 	ParameterPin->Direction = ERigVMPinDirection::Visible;
 	ParameterPin->DefaultValue = InParameterName.ToString();
 	ParameterPin->CustomWidgetName = TEXT("ParameterName");
@@ -1655,7 +1756,7 @@ URigVMParameterNode* URigVMController::AddParameterNode(const FName& InParameter
 
 	if (bIsInput)
 	{
-		if (ValuePin->CPPType == TEXT("FName"))
+		if (ValuePin->CPPType == RigVMTypeUtils::FNameType)
 		{
 			ValuePin->bIsConstant = true;
 		}
@@ -6601,12 +6702,12 @@ FString URigVMController::GetPinInitialDefaultValue(const URigVMPin* InPin)
 	static const FString TObjectInitValue( TEXT("()") );
 	static const TMap<FString, FString> InitValues =
 	{
-		{ TEXT("bool"),	TEXT("False") },
-		{ TEXT("int32"),	TEXT("0") },
-		{ TEXT("float"),	TEXT("0.000000") },
-		{ TEXT("double"),	TEXT("0.000000") },
-		{ TEXT("FName"),	FName(NAME_None).ToString() },
-		{ TEXT("FString"),	TEXT("") }
+		{ RigVMTypeUtils::BoolType,	TEXT("False") },
+		{ RigVMTypeUtils::Int32Type,	TEXT("0") },
+		{ RigVMTypeUtils::FloatType,	TEXT("0.000000") },
+		{ RigVMTypeUtils::DoubleType,	TEXT("0.000000") },
+		{ RigVMTypeUtils::FNameType,	FName(NAME_None).ToString() },
+		{ RigVMTypeUtils::FStringType,	TEXT("") }
 	};
 
 	if (InPin->IsStruct())
@@ -7657,16 +7758,24 @@ bool URigVMController::AddLink(URigVMPin* OutputPin, URigVMPin* InputPin, bool b
 	// resolve types on the pins if needed
 	if(InputPin->GetCPPTypeObject() != OutputPin->GetCPPTypeObject())
 	{
-		if(InputPin->GetCPPTypeObject() == FRigVMUnknownType::StaticStruct())
+		if(InputPin->IsWildCard())
 		{
 			Notify(ERigVMGraphNotifType::InteractionBracketOpened, nullptr);
-			ResolveUnknownTypePin(InputPin, OutputPin, bSetupUndoRedo);
+			ResolveWildCardPin(InputPin, OutputPin, bSetupUndoRedo);
 			Notify(ERigVMGraphNotifType::InteractionBracketClosed, nullptr);
 		}
-		else if(OutputPin->GetCPPTypeObject() == FRigVMUnknownType::StaticStruct())
+		else if(OutputPin->IsWildCard())
 		{
 			Notify(ERigVMGraphNotifType::InteractionBracketOpened, nullptr);
-			ResolveUnknownTypePin(OutputPin, InputPin, bSetupUndoRedo);
+			ResolveWildCardPin(OutputPin, InputPin, bSetupUndoRedo);
+			if(OutputPin->GetNode()->IsA<URigVMRerouteNode>())
+			{
+				SetPinDefaultValue(OutputPin, InputPin->GetDefaultValue(), true, bSetupUndoRedo, false, true);
+			}
+			if(InputPin->GetNode()->IsA<URigVMRerouteNode>())
+			{
+				SetPinDefaultValue(OutputPin, OutputPin->GetDefaultValue(), true, bSetupUndoRedo, false, true);
+			}
 			Notify(ERigVMGraphNotifType::InteractionBracketClosed, nullptr);
 		}
 	}
@@ -7715,10 +7824,6 @@ bool URigVMController::AddLink(URigVMPin* OutputPin, URigVMPin* InputPin, bool b
 		UpdateRerouteNodeAfterChangingLinks(InputPin, bSetupUndoRedo);
 	}
 
-	//TArray<URigVMNode*> NodesVisited;
-	//PotentiallyResolvePrototypeNode(Cast<URigVMPrototypeNode>(InputPin->GetNode()), bSetupUndoRedo, NodesVisited);
-	//PotentiallyResolvePrototypeNode(Cast<URigVMPrototypeNode>(OutputPin->GetNode()), bSetupUndoRedo, NodesVisited);
-
 	if (bSetupUndoRedo)
 	{
 		ActionStack->EndAction(Action);
@@ -8034,6 +8139,12 @@ FName URigVMController::AddExposedPin(const FName& InPinName, ERigVMPinDirection
 
 	if (CPPTypeObject)
 	{
+		if(CPPTypeObject == RigVMTypeUtils::GetWildCardCPPTypeObject())
+		{
+			ReportError(TEXT("Cannot expose pins of wildcard type."));
+			return NAME_None;
+		}
+		
 		if (UScriptStruct* ScriptStruct = Cast<UScriptStruct>(CPPTypeObject))
 		{
 			if (ScriptStruct->IsChildOf(FRigVMExecuteContext::StaticStruct()))
@@ -9917,7 +10028,7 @@ URigVMBranchNode* URigVMController::AddBranchNode(const FVector2D& InPosition, c
 	AddNodePin(Node, ExecutePin);
 
 	URigVMPin* ConditionPin = NewObject<URigVMPin>(Node, *URigVMBranchNode::ConditionName);
-	ConditionPin->CPPType = TEXT("bool");
+	ConditionPin->CPPType = RigVMTypeUtils::BoolType;
 	ConditionPin->Direction = ERigVMPinDirection::Input;
 	AddNodePin(Node, ConditionPin);
 
@@ -9999,7 +10110,7 @@ URigVMIfNode* URigVMController::AddIfNode(const FString& InCPPType, const FName&
 	Node->Position = InPosition;
 
 	URigVMPin* ConditionPin = NewObject<URigVMPin>(Node, *URigVMIfNode::ConditionName);
-	ConditionPin->CPPType = TEXT("bool");
+	ConditionPin->CPPType = RigVMTypeUtils::BoolType;
 	ConditionPin->Direction = ERigVMPinDirection::Input;
 	AddNodePin(Node, ConditionPin);
 
@@ -10115,7 +10226,7 @@ URigVMSelectNode* URigVMController::AddSelectNode(const FString& InCPPType, cons
 	Node->Position = InPosition;
 
 	URigVMPin* IndexPin = NewObject<URigVMPin>(Node, *URigVMSelectNode::IndexName);
-	IndexPin->CPPType = TEXT("int32");
+	IndexPin->CPPType = RigVMTypeUtils::Int32Type;
 	IndexPin->Direction = ERigVMPinDirection::Input;
 	AddNodePin(Node, IndexPin);
 
@@ -10175,7 +10286,7 @@ URigVMSelectNode* URigVMController::AddSelectNodeFromStruct(UScriptStruct* InScr
 	return AddSelectNode(InScriptStruct->GetStructCPPName(), FName(InScriptStruct->GetPathName()), InPosition, InNodeName, bSetupUndoRedo);
 }
 
-URigVMPrototypeNode* URigVMController::AddPrototypeNode(const FName& InNotation, const FVector2D& InPosition, const FString& InNodeName, bool bSetupUndoRedo, bool bPrintPythonCommand)
+URigVMTemplateNode* URigVMController::AddTemplateNode(const FName& InNotation, const FVector2D& InPosition, const FString& InNodeName, bool bSetupUndoRedo, bool bPrintPythonCommand)
 {
 	if (!IsValidGraph())
 	{
@@ -10187,28 +10298,28 @@ URigVMPrototypeNode* URigVMController::AddPrototypeNode(const FName& InNotation,
 
 	ensure(!InNotation.IsNone());
 
-	const FRigVMPrototype* Prototype = FRigVMRegistry::Get().FindPrototype(InNotation);
-	if (Prototype == nullptr)
+	const FRigVMTemplate* Template = FRigVMRegistry::Get().FindTemplate(InNotation);
+	if (Template == nullptr)
 	{
-		ReportErrorf(TEXT("Prototype '%s' cannot be found."), *InNotation.ToString());
+		ReportErrorf(TEXT("Template '%s' cannot be found."), *InNotation.ToString());
 		return nullptr;
 	}
 
-	FString Name = GetValidNodeName(InNodeName.IsEmpty() ? Prototype->GetName().ToString() : InNodeName);
-	URigVMPrototypeNode* Node = NewObject<URigVMPrototypeNode>(Graph, *Name);
-	Node->PrototypeNotation = Prototype->GetNotation();
+	FString Name = GetValidNodeName(InNodeName.IsEmpty() ? Template->GetName().ToString() : InNodeName);
+	URigVMTemplateNode* Node = NewObject<URigVMTemplateNode>(Graph, *Name);
+	Node->TemplateNotation = Template->GetNotation();
 	Node->Position = InPosition;
 
-	int32 FunctionIndex = INDEX_NONE;
-	FRigVMPrototype::FTypeMap Types;
-	Prototype->Resolve(Types, FunctionIndex);
+	int32 PermutationIndex = INDEX_NONE;
+	FRigVMTemplate::FTypeMap Types;
+	Template->FullyResolve(Types, PermutationIndex);
 
-	for (int32 ArgIndex = 0; ArgIndex < Prototype->NumArgs(); ArgIndex++)
+	for (int32 ArgIndex = 0; ArgIndex < Template->NumArguments(); ArgIndex++)
 	{
-		const FRigVMPrototypeArg* Arg = Prototype->GetArg(ArgIndex);
+		const FRigVMTemplateArgument* Arg = Template->GetArgument(ArgIndex);
 
 		URigVMPin* Pin = NewObject<URigVMPin>(Node, Arg->GetName());
-		const FRigVMPrototypeArg::FType& Type = Types.FindChecked(Arg->GetName());
+		const FRigVMTemplateArgument::FType& Type = Types.FindChecked(Arg->GetName());
 		Pin->CPPType = Type.CPPType;
 		Pin->CPPTypeObject = Type.CPPTypeObject;
 		if (Pin->CPPTypeObject)
@@ -10226,7 +10337,7 @@ URigVMPrototypeNode* URigVMController::AddPrototypeNode(const FName& InNotation,
 	FRigVMControllerCompileBracketScope CompileScope(this);
 	if (bSetupUndoRedo)
 	{
-		ActionStack->AddAction(FRigVMAddPrototypeNodeAction(Node));
+		ActionStack->AddAction(FRigVMAddTemplateNodeAction(Node));
 	}
 
 	if (bPrintPythonCommand)
@@ -10239,6 +10350,8 @@ URigVMPrototypeNode* URigVMController::AddPrototypeNode(const FName& InNotation,
 		}
 	}
 
+	ResolveTemplateNodeMetaData(Node);
+
 	return Node;
 }
 
@@ -10280,7 +10393,7 @@ URigVMEnumNode* URigVMController::AddEnumNode(const FName& InCPPTypeObjectPath,
 	AddNodePin(Node, EnumValuePin);
 
 	URigVMPin* EnumIndexPin = NewObject<URigVMPin>(Node, *URigVMEnumNode::EnumIndexName);
-	EnumIndexPin->CPPType = TEXT("int32");
+	EnumIndexPin->CPPType = RigVMTypeUtils::Int32Type;
 	EnumIndexPin->Direction = ERigVMPinDirection::Output;
 	EnumIndexPin->DisplayName = TEXT("Result");
 	AddNodePin(Node, EnumIndexPin);
@@ -10434,32 +10547,32 @@ URigVMArrayNode* URigVMController::AddArrayNode(ERigVMOpCode InOpCode, const FSt
 
 		static URigVMPin* AddIndexPin(URigVMController* InController, URigVMNode* InNode, ERigVMPinDirection InDirection)
 		{
-			return AddPin(InController, InNode, *URigVMArrayNode::IndexName, InDirection, false, TEXT("int32"), nullptr);  
+			return AddPin(InController, InNode, *URigVMArrayNode::IndexName, InDirection, false, RigVMTypeUtils::Int32Type, nullptr);  
 		}
 
 		static URigVMPin* AddNumPin(URigVMController* InController, URigVMNode* InNode, ERigVMPinDirection InDirection)
 		{
-			return AddPin(InController, InNode, *URigVMArrayNode::NumName, InDirection, false, TEXT("int32"), nullptr);  
+			return AddPin(InController, InNode, *URigVMArrayNode::NumName, InDirection, false, RigVMTypeUtils::Int32Type, nullptr);  
 		}
 
 		static URigVMPin* AddCountPin(URigVMController* InController, URigVMNode* InNode, ERigVMPinDirection InDirection)
 		{
-			return AddPin(InController, InNode, *URigVMArrayNode::CountName, InDirection, false, TEXT("int32"), nullptr);  
+			return AddPin(InController, InNode, *URigVMArrayNode::CountName, InDirection, false, RigVMTypeUtils::Int32Type, nullptr);  
 		}
 
 		static URigVMPin* AddRatioPin(URigVMController* InController, URigVMNode* InNode)
 		{
-			return AddPin(InController, InNode, *URigVMArrayNode::RatioName, ERigVMPinDirection::Output, false, TEXT("float"), nullptr);  
+			return AddPin(InController, InNode, *URigVMArrayNode::RatioName, ERigVMPinDirection::Output, false, RigVMTypeUtils::FloatType, nullptr);  
 		}
 
 		static URigVMPin* AddContinuePin(URigVMController* InController, URigVMNode* InNode)
 		{
-			return AddPin(InController, InNode, *URigVMArrayNode::ContinueName, ERigVMPinDirection::Hidden, false, TEXT("bool"), nullptr);
+			return AddPin(InController, InNode, *URigVMArrayNode::ContinueName, ERigVMPinDirection::Hidden, false, RigVMTypeUtils::BoolType, nullptr);
 		}
 
 		static URigVMPin* AddSuccessPin(URigVMController* InController, URigVMNode* InNode)
 		{
-			return AddPin(InController, InNode, *URigVMArrayNode::SuccessName, ERigVMPinDirection::Output, false, TEXT("bool"), nullptr);  
+			return AddPin(InController, InNode, *URigVMArrayNode::SuccessName, ERigVMPinDirection::Output, false, RigVMTypeUtils::BoolType, nullptr);  
 		}
 	};
 
@@ -10640,7 +10753,7 @@ void URigVMController::ForEveryPinRecursively(URigVMNode* InNode, TFunction<void
 	}
 }
 
-void URigVMController::SetExecuteContextStruct(UStruct* InExecuteContextStruct)
+void URigVMController::SetExecuteContextStruct(UScriptStruct* InExecuteContextStruct)
 {
 	check(InExecuteContextStruct);
 	ensure(InExecuteContextStruct->IsChildOf(FRigVMExecuteContext::StaticStruct()));
@@ -11290,7 +11403,7 @@ bool URigVMController::ShouldStructBeUnfolded(const UStruct* Struct)
 	{
 		return false;
 	}
-	if(Struct->IsChildOf(FRigVMUnknownType::StaticStruct()))
+	if(Struct->IsChildOf(RigVMTypeUtils::GetWildCardCPPTypeObject()))
 	{
 		return false;
 	}
@@ -11679,6 +11792,11 @@ void URigVMController::AddPinRedirector(bool bInput, bool bOutput, const FString
 
 bool URigVMController::ShouldRedirectPin(UScriptStruct* InOwningStruct, const FString& InOldRelativePinPath, FString& InOutNewRelativePinPath) const
 {
+	if(InOwningStruct == nullptr) // potentially a template node
+	{
+		return false;
+	}
+	
 	FControlRigStructPinRedirectorKey RedirectorKey(InOwningStruct, InOldRelativePinPath);
 	if (const FString* RedirectedPinPath = PinPathCoreRedirectors.Find(RedirectorKey))
 	{
@@ -11876,9 +11994,16 @@ void URigVMController::RepopulatePinsOnNode(URigVMNode* InNode, bool bFollowCore
 	// step 2/3: clear pins on the node and repopulate the node with new pins
 	if (UnitNode != nullptr)
 	{
+		UScriptStruct* ScriptStruct = UnitNode->GetScriptStruct();
+		if (ScriptStruct == nullptr)
+		{
+			// this may be an unresolved template node
+			// in that case there's nothing we can do here
+			return;
+		}
+
 		RemovePinsDuringRepopulate(UnitNode, UnitNode->Pins, bNotify, bSetupOrphanedPins);
 
-		UScriptStruct* ScriptStruct = UnitNode->GetScriptStruct();
 		if (ScriptStruct == nullptr)
 		{
 			ReportWarningf(
@@ -12649,7 +12774,7 @@ FString URigVMController::PostProcessCPPType(const FString& InCPPType, UObject*
 	}
 	else if (UEnum* Enum = Cast<UEnum>(InCPPTypeObject))
 	{
-		CPPType = Enum->CppType;
+		CPPType = RigVMTypeUtils::CPPTypeFromEnum(Enum);
 	}
 
 	if(CPPType != InCPPType)
@@ -12665,139 +12790,401 @@ FString URigVMController::PostProcessCPPType(const FString& InCPPType, UObject*
 	return CPPType;
 }
 
-/*
-void URigVMController::PotentiallyResolvePrototypeNode(URigVMPrototypeNode* InNode, bool bSetupUndoRedo)
-{
-	TArray<URigVMNode*> NodesVisited;
-	PotentiallyResolvePrototypeNode(InNode, bSetupUndoRedo, NodesVisited);
-}
-
-void URigVMController::PotentiallyResolvePrototypeNode(URigVMPrototypeNode* InNode, bool bSetupUndoRedo, TArray<URigVMNode*>& NodesVisited)
+void URigVMController::ResolveTemplateNodePins(URigVMTemplateNode* InNode, bool bSetupUndoRedo)
 {
 	if (InNode == nullptr)
 	{
 		return;
 	}
 
-	if (NodesVisited.Contains(InNode))
-	{
-		return;
-	}
-	NodesVisited.Add(InNode);
+	InNode->ResolvedFunctionName.Reset();
 
-	// propagate types first
-	for (URigVMPin* Pin : InNode->GetPins())
+	// check if the node is resolved
+	TArray<int32> ResolvedPermutationIndices = InNode->GetResolvedPermutationIndices();
+	if(ResolvedPermutationIndices.Num() > 0)
 	{
-		if (Pin->CPPType.IsEmpty())
+		TArray<URigVMPin*> ResolvedPins;
+		
+		for (URigVMPin* Pin : InNode->GetPins())
 		{
-			TArray<URigVMPin*> LinkedPins = Pin->GetLinkedSourcePins();
-			LinkedPins.Append(Pin->GetLinkedTargetPins());
+			if(!Pin->IsWildCard())
+			{
+				continue;
+			}
+
+			const FRigVMTemplateArgument* Arg = InNode->GetTemplate()->FindArgument(Pin->GetFName());
+			if(Arg == nullptr)
+			{
+				continue;
+			}
 
-			for (URigVMPin* LinkedPin : LinkedPins)
+			if(Arg->IsSingleton(ResolvedPermutationIndices))
 			{
-				if (!LinkedPin->CPPType.IsEmpty())
+				const FRigVMTemplateArgument::FType ResolvedType = Arg->GetTypes()[ResolvedPermutationIndices[0]];
+				if(ChangePinType(Pin, ResolvedType.CPPType, ResolvedType.CPPTypeObject, bSetupUndoRedo, false, false, false))
 				{
-					ChangePinType(Pin, LinkedPin->CPPType, LinkedPin->CPPTypeObjectPath, bSetupUndoRedo);
-					break;
+					ResolvedPins.Add(Pin);
+					
+					if(Pin->IsArray())
+					{
+						const FString BaseCPPType = RigVMTypeUtils::BaseTypeFromArrayType(ResolvedType.CPPType);
+						for(URigVMPin* SubPin : Pin->GetSubPins())
+						{
+							if(SubPin->IsWildCard())
+							{
+								ResolveWildCardPinImpl(SubPin, BaseCPPType, ResolvedType.CPPTypeObject, bSetupUndoRedo, false, false, true);
+							}
+						}
+					}
 				}
 			}
 		}
-	}
 
-	// check if the node is resolved
-	FRigVMPrototype::FTypeMap ResolvedTypes;
-	int32 FunctionIndex = InNode->GetResolvedFunctionIndex(&ResolvedTypes);
-	if (FunctionIndex != INDEX_NONE)
-	{
-		// we have a valid node - let's replace this node... first let's find all 
-		// links and all default values
-		TMap<FString, FString> DefaultValues;
-		TArray<TPair<FString, FString>> LinkPaths;
+		if(ResolvedPermutationIndices.Num() == 1)
+		{
+			const FRigVMFunction* ResolvedFunction = InNode->GetTemplate()->GetPermutation(ResolvedPermutationIndices[0]);
+			FullyResolveTemplateNode(InNode, ResolvedFunction, bSetupUndoRedo);
+		}
 
-		for (URigVMPin* Pin : InNode->GetPins())
+		// resolve all wires on resolved pins
+		for(URigVMPin* ResolvedPin : ResolvedPins)
 		{
-			FString DefaultValue = Pin->GetDefaultValue();
-			if (!DefaultValue.IsEmpty())
+			// we don't need to recurse since unknown type pins don't have sub pins
+			TArray<URigVMPin*> LinkedPins = ResolvedPin->GetLinkedSourcePins(false); 
+			LinkedPins.Append(ResolvedPin->GetLinkedTargetPins(false));
+
+			for(URigVMPin* LinkedPin : LinkedPins)
 			{
-				DefaultValues.Add(Pin->GetPinPath(), DefaultValue);
+				if(LinkedPin->IsWildCard())
+				{
+					ResolveWildCardPin(LinkedPin, ResolvedPin, bSetupUndoRedo);
+				}
 			}
 
-			TArray<URigVMLink*> Links = Pin->GetSourceLinks(true);
-			Links.Append(Pin->GetTargetLinks(true));
+		}
+	}
 
-			for (URigVMLink* Link : Links)
+	ResolveTemplateNodeMetaData(InNode);
+}
+
+void URigVMController::ResolveTemplateNodeMetaData(URigVMTemplateNode* InNode)
+{
+#if WITH_EDITOR
+	check(InNode);
+
+	const FLinearColor PreviousColor = InNode->NodeColor;
+	
+	const TArray<int32> ResolvedPermutationIndices = InNode->GetResolvedPermutationIndices();
+	InNode->NodeColor = InNode->GetTemplate()->GetColor(ResolvedPermutationIndices);
+
+	if(!InNode->NodeColor.Equals(PreviousColor, 0.01f))
+	{
+		Notify(ERigVMGraphNotifType::NodeColorChanged, InNode);
+	}
+#endif
+}
+
+bool URigVMController::FullyResolveTemplateNode(URigVMTemplateNode* InNode, const FRigVMFunction* InFunctionToResolve, bool bSetupUndoRedo)
+{
+	if(bIsFullyResolvingTemplateNode)
+	{
+		return false;
+	}
+	TGuardValue<bool> ReentryGuard(bIsFullyResolvingTemplateNode, true);
+	
+	check(InNode);
+	check(InFunctionToResolve);
+
+	const FRigVMTemplate* Template = InNode->GetTemplate();
+
+	const int32 PermutationIndex = Template->FindPermutation(InFunctionToResolve);
+	if(PermutationIndex == INDEX_NONE)
+	{
+		ReportErrorf(TEXT("Template %s doesn't support function %s::%s."),
+			*Template->GetNotation().ToString(),
+			*InFunctionToResolve->Struct->GetStructCPPName(),
+			*InFunctionToResolve->GetMethodName().ToString()
+		);
+		return false;
+	}
+
+	const TArray<int32> PermutationIndices = {PermutationIndex};
+
+	// find all existing pins that we may need to change
+	TArray<FRigVMTemplateArgument> MissingPins;
+	TArray<URigVMPin*> PinsToRemove;
+	TMap<URigVMPin*, FRigVMTemplateArgument::FType> PinTypesToChange;
+	for(int32 ArgIndex = 0; ArgIndex < Template->NumArguments(); ArgIndex++)
+	{
+		const FRigVMTemplateArgument* Argument = Template->GetArgument(ArgIndex);
+		const FRigVMTemplateArgument::FType ResolvedType = Argument->GetSupportedTypes(PermutationIndices)[0];
+
+		URigVMPin* Pin = InNode->FindPin(Argument->GetName().ToString());
+		if(Pin == nullptr)
+		{
+			ReportErrorf(TEXT("Template node %s is missing a pin for argument %s"),
+				*InNode->GetNodePath(),
+				*Argument->GetName().ToString()
+			);
+			return false;
+		}
+
+		if(Pin->GetCPPType() != ResolvedType.CPPType)
+		{
+			PinTypesToChange.Add(Pin, ResolvedType);
+		}
+	}
+
+	// find all missing pins which are not arguments for the template
+	for (TFieldIterator<FProperty> It(InFunctionToResolve->Struct); It; ++It)
+	{
+		const FRigVMTemplateArgument ExpectedArgument(*It);
+		const FRigVMTemplateArgument::FType ExpectedType = ExpectedArgument.GetSupportedTypes()[0];
+
+		if(URigVMPin* Pin = InNode->FindPin(It->GetFName().ToString()))
+		{
+			if(Pin->GetCPPType() != ExpectedType.CPPType)
 			{
-				LinkPaths.Add(TPair<FString, FString>(Link->GetSourcePin()->GetPinPath(), Link->GetTargetPin()->GetPinPath()));
+				PinTypesToChange.Add(Pin, ExpectedType);
 			}
 		}
+		else
+		{
+			MissingPins.Add(ExpectedArgument);
+		}
+	}
+
+	// find all pins which don't have a matching arg on the function
+	for(URigVMPin* Pin : InNode->GetPins())
+	{
+		if(InFunctionToResolve->Struct->FindPropertyByName(Pin->GetFName()) == nullptr)
+		{
+			PinsToRemove.Add(Pin);
+		}
+	}
+
+	// update the cached resolved function name
+	InNode->ResolvedFunctionName = InFunctionToResolve->Name;
 
-		const FRigVMFunction& Function = FRigVMRegistry::Get().GetFunctions()[FunctionIndex];
-		FString NodeName = InNode->GetName();
-		FVector2D NodePosition = InNode->GetPosition();
+	// exit out early if there's nothing to do
+	if(PinTypesToChange.IsEmpty() && MissingPins.IsEmpty() && PinsToRemove.IsEmpty())
+	{
+		ResolveTemplateNodeMetaData(InNode);
+		return true;
+	}
 
-		RemoveNode(InNode, bSetupUndoRedo);
+	if(bSetupUndoRedo)
+	{
+		OpenUndoBracket(TEXT("Resolve Template Node"));
+	}
 
-		if (URigVMNode* NewNode = AddUnitNode(Function.Struct, Function.GetMethodName(), NodePosition, NodeName, bSetupUndoRedo))
+	// update the incorrectly typed pins
+	for(const TPair<URigVMPin*, FRigVMTemplateArgument::FType>& Pair : PinTypesToChange)
+	{
+		URigVMPin* Pin = Pair.Key;
+		const FRigVMTemplateArgument::FType& ExpectedType = Pair.Value;
+		
+		if(!Pin->IsWildCard())
 		{
-			// set default values again
-			for (TPair<FString, FString> Pair : DefaultValues)
+			if((Pin->GetCPPType() != ExpectedType.CPPType) ||
+				(Pin->GetCPPTypeObject() != ExpectedType.CPPTypeObject))
 			{
-				SetPinDefaultValue(Pair.Key, Pair.Value, true, bSetupUndoRedo, false);
+				const FString CPPType = Pin->IsArray() ? RigVMTypeUtils::GetWildCardArrayCPPType() : RigVMTypeUtils::GetWildCardCPPType();
+				if(!ChangePinType(Pin, CPPType, RigVMTypeUtils::GetWildCardCPPTypeObject(), bSetupUndoRedo, false, true, true))
+				{
+					if(bSetupUndoRedo)
+					{
+						CancelUndoBracket();
+					}
+					return false;
+				}
 			}
-
-			// reestablish links
-			for (TPair<FString, FString> Pair : LinkPaths)
+		}
+		
+		if(Pin->IsWildCard())
+		{
+			if(!ResolveWildCardPinImpl(Pin, ExpectedType.CPPType, ExpectedType.CPPTypeObject, bSetupUndoRedo, false, true, true))
 			{
-				AddLink(Pair.Key, Pair.Value, bSetupUndoRedo);
+				if(bSetupUndoRedo)
+				{
+					CancelUndoBracket();
+				}
+				return false;
 			}
 		}
+	}
 
-		return;
+	// remove obsolete pins
+	for(URigVMPin* PinToRemove : PinsToRemove)
+	{
+		RemovePin(PinToRemove, false, true);
 	}
-	else
+
+	// add missing pins
+	for(const FRigVMTemplateArgument& MissingPin : MissingPins)
 	{
-		// update all of the pins that might have changed now as well!
-		for (URigVMPin* Pin : InNode->GetPins())
+		check(MissingPin.GetDirection() == ERigVMPinDirection::Hidden);
+		
+		FProperty* Property = InFunctionToResolve->Struct->FindPropertyByName(MissingPin.GetName());
+		check(Property);
+
+		URigVMPin* Pin = NewObject<URigVMPin>(Cast<UObject>(InNode), MissingPin.GetName());
+		ConfigurePinFromProperty(Property, Pin, MissingPin.GetDirection());
+
+		AddNodePin(InNode, Pin);
+		Notify(ERigVMGraphNotifType::PinAdded, Pin);
+
+		// we don't need to set the default value here since the pin is hidden
+	}
+
+	if (UnitNodeCreatedContext.IsValid())
+	{
+		if (TSharedPtr<FStructOnScope> StructScope = InNode->ConstructStructInstance())
+		{
+			TGuardValue<FName> NodeNameScope(UnitNodeCreatedContext.NodeName, InNode->GetFName());
+			FRigVMStruct* StructInstance = (FRigVMStruct*)StructScope->GetStructMemory();
+			StructInstance->OnUnitNodeCreated(UnitNodeCreatedContext);
+		}
+	}
+	
+	if(bSetupUndoRedo)
+	{
+		CloseUndoBracket();
+	}
+
+	return true;
+}
+
+bool URigVMController::ResolveWildCardPin(const FString& InPinPath, const FString& InCPPType, const FName& InCPPTypeObjectPath, bool bSetupUndoRedo, bool bPrintPythonCommand)
+{
+	if (!IsValidGraph())
+	{
+		return false;
+	}
+
+	URigVMGraph* Graph = GetGraph();
+	check(Graph);
+
+	UObject* CPPTypeObject = nullptr;
+	if (!InCPPTypeObjectPath.IsNone())
+	{
+		CPPTypeObject = URigVMPin::FindObjectFromCPPTypeObjectPath<UObject>(InCPPTypeObjectPath.ToString());
+		if (CPPTypeObject == nullptr)
+		{
+			ReportErrorf(TEXT("Cannot find cpp type object for path '%s'."), *InCPPTypeObjectPath.ToString());
+			return false;
+		}
+	}
+
+	const FString CPPType = PostProcessCPPType(InCPPType, CPPTypeObject);
+
+	if (URigVMPin* Pin = Graph->FindPin(InPinPath))
+	{
+		if(Pin->IsWildCard())
 		{
-			if (Pin->CPPType.IsEmpty())
+			FRigVMBaseAction Action;
+			if (bSetupUndoRedo)
 			{
-				if (const FRigVMPrototypeArg::FType* Type = ResolvedTypes.Find(Pin->GetFName()))
+				Action.Title = TEXT("Resolve Wildcard Pin");
+				ActionStack->BeginAction(Action);
+			}
+			
+			if(ResolveWildCardPinImpl(Pin, CPPType, CPPTypeObject, bSetupUndoRedo))
+			{
+				if(bPrintPythonCommand)
 				{
-					if (!Type->CPPType.IsEmpty())
-					{
-						ChangePinType(Pin, Type->CPPType, Type->GetCPPTypeObjectPath(), bSetupUndoRedo);
-					}
+					const FString GraphName = GetSanitizedGraphName(GetGraph()->GetGraphName());
+					const FString SanitizedPinPath = GetSanitizedPinPath(Pin->GetPinPath());
+					const FString CPPObjectPath = CPPTypeObject ? FString::Printf(TEXT("'%s'"), *CPPTypeObject->GetPathName()) : TEXT("None");
+
+					RigVMPythonUtils::Print(GetGraphOuterName(), 
+						FString::Printf(TEXT("blueprint.get_controller_by_name('%s').resolve_wild_card_pin('%s', '%s', %s)"),
+									*GraphName,
+									*SanitizedPinPath,
+									*CPPType,
+									*CPPObjectPath));
+				}
+
+				if (bSetupUndoRedo)
+				{
+					ActionStack->EndAction(Action);
 				}
+				return true;
+			}
+
+			if (bSetupUndoRedo)
+			{
+				ActionStack->CancelAction(Action);
 			}
 		}
 	}
 
-	// then recursively call
-	TArray<URigVMNode*> LinkedNodes = InNode->GetLinkedSourceNodes();
-	LinkedNodes.Append(InNode->GetLinkedTargetNodes());
-	for (URigVMNode* LinkedNode : LinkedNodes)
+	return false;
+}
+
+bool URigVMController::ResolveWildCardPin(URigVMPin* InPinToResolve, URigVMPin* InTemplatePin,
+										  bool bSetupUndoRedo, bool bTraverseNode, bool bTraverseParentPins, bool bTraverseLinks)
+{
+	check(InTemplatePin);
+	check(!InTemplatePin->IsWildCard());
+	check(InPinToResolve != InTemplatePin);
+	check(InPinToResolve->GetRootPin() != InTemplatePin->GetRootPin());
+
+	FRigVMBaseAction Action;
+	if (bSetupUndoRedo)
+	{
+		Action.Title = TEXT("Resolve Wildcard Pin");
+		ActionStack->BeginAction(Action);
+	}
+
+	const bool bResult = ResolveWildCardPinImpl(InPinToResolve, InTemplatePin->GetCPPType(), InTemplatePin->GetCPPTypeObject(), bSetupUndoRedo, bTraverseNode, bTraverseParentPins, bTraverseLinks);
+	if(!bResult)
 	{
-		PotentiallyResolvePrototypeNode(Cast<URigVMPrototypeNode>(LinkedNode), bSetupUndoRedo, NodesVisited);
+		if(InPinToResolve->IsLinkedTo(InTemplatePin))
+		{
+			BreakLink(InPinToResolve, InTemplatePin, bSetupUndoRedo);
+			BreakLink(InTemplatePin, InPinToResolve, bSetupUndoRedo);
+		}
+	}
+
+	if (bSetupUndoRedo)
+	{
+		ActionStack->EndAction(Action);
 	}
+	
+	return bResult;
 }
-*/
 
-void URigVMController::ResolveUnknownTypePin(URigVMPin* InPinToResolve, const URigVMPin* InTemplatePin,
-	bool bSetupUndoRedo, bool bTraverseNode, bool bTraverseParentPins, bool bTraverseLinks)
+bool URigVMController::ResolveWildCardPinImpl(URigVMPin* InPinToResolve, const FString& InCPPType, UObject* InCPPTypeObject,
+                                             bool bSetupUndoRedo, bool bTraverseNode, bool bTraverseParentPins, bool bTraverseLinks)
 {
 	check(InPinToResolve);
-	check(InTemplatePin);
-	check(InPinToResolve->IsUnknownType());
-	check(!InTemplatePin->IsUnknownType());
+	check(InPinToResolve->IsWildCard());
+
+	// if we are trying to resolve a pin on the template node...
+	FString CPPType = InCPPType;
+	UObject* CPPTypeObject = InCPPTypeObject;
+	
+	if(URigVMTemplateNode* TemplateNode = Cast<URigVMTemplateNode>(InPinToResolve->GetNode()))
+	{
+		const FString BaseCPPType = RigVMTypeUtils::IsArrayType(CPPType) ? RigVMTypeUtils::BaseTypeFromArrayType(CPPType) : CPPType;
+		const FString TypeToResolveAgainst = InPinToResolve->GetRootPin()->IsArray() ? RigVMTypeUtils::ArrayTypeFromBaseType(BaseCPPType) : BaseCPPType; 
+
+		if(!TemplateNode->SupportsType(InPinToResolve->GetRootPin(), TypeToResolveAgainst, &CPPType))
+		{
+			return false;
+		}
+	}
 
 	FRigVMControllerCompileBracketScope CompileScope(this);
 	if(bTraverseParentPins)
 	{
 		URigVMPin* RootPin = InPinToResolve->GetRootPin();
-		if(RootPin != InPinToResolve && RootPin->IsUnknownType())
+		if(RootPin != InPinToResolve && RootPin->IsWildCard())
 		{
-			ResolveUnknownTypePin(RootPin, InTemplatePin, bSetupUndoRedo, bTraverseNode, false, bTraverseLinks);
-			return;
+			const FString BaseCPPType = RigVMTypeUtils::IsArrayType(CPPType) ? RigVMTypeUtils::BaseTypeFromArrayType(CPPType) : CPPType;
+			const FString TypeToResolveAgainst = RootPin->IsArray() ? RigVMTypeUtils::ArrayTypeFromBaseType(BaseCPPType) : BaseCPPType; 
+			return ResolveWildCardPinImpl(RootPin, TypeToResolveAgainst, CPPTypeObject, bSetupUndoRedo, bTraverseNode, false, bTraverseLinks);
 		}
 	}
 
@@ -12805,33 +13192,45 @@ void URigVMController::ResolveUnknownTypePin(URigVMPin* InPinToResolve, const UR
 	{
 		for(URigVMPin* SubPin : InPinToResolve->GetSubPins())
 		{
-			if(SubPin->IsUnknownType())
+			if(SubPin->IsWildCard())
 			{
-				ResolveUnknownTypePin(SubPin, InTemplatePin, bSetupUndoRedo, false, false, bTraverseLinks);
+				const FString BaseCPPType = RigVMTypeUtils::IsArrayType(CPPType) ? RigVMTypeUtils::BaseTypeFromArrayType(CPPType) : CPPType;
+				const FString TypeToResolveAgainst = SubPin->IsArray() ? RigVMTypeUtils::ArrayTypeFromBaseType(BaseCPPType) : BaseCPPType; 
+				ResolveWildCardPinImpl(SubPin, TypeToResolveAgainst, CPPTypeObject, bSetupUndoRedo, false, false, bTraverseLinks);
 			}
 		}
 	}
 
-	FString ResolvedCPPType = InTemplatePin->IsArray() ? InTemplatePin->GetArrayElementCppType() : InTemplatePin->GetCPPType();
+	FString ResolvedCPPType = RigVMTypeUtils::IsArrayType(CPPType) ? RigVMTypeUtils::BaseTypeFromArrayType(CPPType) : CPPType;
 	if(InPinToResolve->IsArray())
 	{
 		ResolvedCPPType = RigVMTypeUtils::ArrayTypeFromBaseType(ResolvedCPPType);
 	}
 
-	if(!ChangePinType(InPinToResolve, ResolvedCPPType, InTemplatePin->GetCPPTypeObject(), bSetupUndoRedo, false, false, false))
+	if(ChangePinType(InPinToResolve, ResolvedCPPType, CPPTypeObject, bSetupUndoRedo, false, false, false))
 	{
-		return;
+		if(URigVMTemplateNode* TemplateNode = Cast<URigVMTemplateNode>(InPinToResolve->GetNode()))
+		{
+			ResolveTemplateNodePins(TemplateNode, bSetupUndoRedo);
+		}
+	}
+	else
+	{
+		return false;
 	}
 
 	if(bTraverseNode)
 	{
 		if(const URigVMNode* Node = InPinToResolve->GetNode())
 		{
-			for(URigVMPin* Pin : Node->GetPins())
+			if(!Node->IsA<URigVMTemplateNode>())
 			{
-				if(Pin->IsUnknownType())
+				for(URigVMPin* Pin : Node->GetPins())
 				{
-					ResolveUnknownTypePin(Pin, InTemplatePin, bSetupUndoRedo, false, bTraverseParentPins, bTraverseLinks);
+					if(Pin->IsWildCard())
+					{
+						ResolveWildCardPin(Pin, InPinToResolve, bSetupUndoRedo, false, bTraverseParentPins, bTraverseLinks);
+					}
 				}
 			}
 		}
@@ -12845,12 +13244,17 @@ void URigVMController::ResolveUnknownTypePin(URigVMPin* InPinToResolve, const UR
 
 		for(URigVMPin* LinkedPin : LinkedPins)
 		{
-			if(LinkedPin->IsUnknownType())
+			if(LinkedPin->IsWildCard())
 			{
-				ResolveUnknownTypePin(LinkedPin, InTemplatePin, bSetupUndoRedo);
+				if(!ResolveWildCardPin(LinkedPin, InPinToResolve, bSetupUndoRedo))
+				{
+					return false;
+				}
 			}
 		}
 	}
+
+	return true;
 }
 
 bool URigVMController::ChangePinType(const FString& InPinPath, const FString& InCPPType, const FName& InCPPTypeObjectPath, bool bSetupUndoRedo, bool bSetupOrphanPins, bool bBreakLinks, bool bRemoveSubPins)
@@ -12916,6 +13320,19 @@ bool URigVMController::ChangePinType(URigVMPin* InPin, const FString& InCPPType,
 		return false;
 	}
 
+	// only allow valid pin cpp types on template nodes
+	FString CPPType = InCPPType;
+	if(URigVMTemplateNode* TemplateNode = Cast<URigVMTemplateNode>(InPin->GetNode()))
+	{
+		if(!TemplateNode->SupportsType(InPin, InCPPType, &CPPType))
+		{
+			return false;
+		}
+
+		const FString BaseCPPType = RigVMTypeUtils::IsArrayType(CPPType) ? RigVMTypeUtils::BaseTypeFromArrayType(CPPType) : CPPType;
+		CPPType = InPin->IsArray() ? RigVMTypeUtils::ArrayTypeFromBaseType(BaseCPPType) : BaseCPPType;
+	}
+
 	FRigVMControllerCompileBracketScope CompileScope(this);
 	FRigVMBaseAction Action;
 	if (bSetupUndoRedo)
@@ -12935,7 +13352,7 @@ bool URigVMController::ChangePinType(URigVMPin* InPin, const FString& InCPPType,
 			BreakAllLinksRecursive(InPin, true, false, true);
 			BreakAllLinksRecursive(InPin, false, false, true);
 		}
-		ActionStack->AddAction(FRigVMChangePinTypeAction(InPin, InCPPType, CPPTypeObjectPath, bSetupOrphanPins, bBreakLinks, bRemoveSubPins));
+		ActionStack->AddAction(FRigVMChangePinTypeAction(InPin, CPPType, CPPTypeObjectPath, bSetupOrphanPins, bBreakLinks, bRemoveSubPins));
 	}
 	
 	if(bSetupOrphanPins)
@@ -12970,11 +13387,17 @@ bool URigVMController::ChangePinType(URigVMPin* InPin, const FString& InCPPType,
 		
 		InPin->SubPins.Reset();
 	}
+
+	// compute the number of remaining wildcard pins
+	auto WildCardPinCountPredicate = [](const URigVMPin* Pin) { return Pin->IsWildCard(); };
+	TArray<URigVMPin*> AllPins = InPin->GetNode()->GetAllPinsRecursively();
+	int32 RemainingWildCardPins = Algo::CountIf(AllPins, WildCardPinCountPredicate);
+	const bool bPinWasWildCard = InPin->IsWildCard();
 	
-	InPin->CPPType = InCPPType;
+	InPin->CPPType = CPPType;
 	InPin->CPPTypeObjectPath = CPPTypeObjectPath;
 	InPin->CPPTypeObject = InCPPTypeObject;
-	InPin->bIsDynamicArray = RigVMTypeUtils::IsArrayType(InCPPType);
+	InPin->bIsDynamicArray = RigVMTypeUtils::IsArrayType(CPPType);
 	// we might want to use GetPinInitialDefaultValue here for a better default value
 	InPin->DefaultValue = FString();
 
@@ -12983,16 +13406,157 @@ bool URigVMController::ChangePinType(URigVMPin* InPin, const FString& InCPPType,
 		InPin->Direction = ERigVMPinDirection::IO;
 	}
 
-	if (InPin->IsStruct())
+	if (InPin->IsStruct() && !InPin->IsArray())
 	{
 		FString DefaultValue = InPin->DefaultValue;
-		CreateDefaultValueForStructIfRequired(InPin->GetScriptStruct(), DefaultValue);
+		if(!bPinWasWildCard)
+		{
+			CreateDefaultValueForStructIfRequired(InPin->GetScriptStruct(), DefaultValue);
+		}
 		AddPinsForStruct(InPin->GetScriptStruct(), InPin->GetNode(), InPin, InPin->Direction, DefaultValue, false, true);
 	}
 
+	// if this is a template clear its caches
+	if(URigVMTemplateNode* TemplateNode = Cast<URigVMTemplateNode>(InPin->GetNode()))
+	{
+		TemplateNode->InvalidateCache();
+	}
+
 	Notify(ERigVMGraphNotifType::PinTypeChanged, InPin);
 	Notify(ERigVMGraphNotifType::PinDefaultValueChanged, InPin);
 
+	// let's see if this was the last resolved wildcard pin
+	if(RemainingWildCardPins > 0)
+	{
+		// compute the number of current wildcard pins
+		RemainingWildCardPins = 0;
+		if(InPin->GetNode()->IsA<URigVMTemplateNode>())
+		{
+			AllPins = InPin->GetNode()->GetAllPinsRecursively();
+			RemainingWildCardPins = Algo::CountIf(AllPins, WildCardPinCountPredicate);
+		}
+
+		// if this is the first time that there are no wild card pins left
+		if(RemainingWildCardPins == 0)
+		{
+			struct Local
+			{
+				static bool IsPinDefaultEmpty(URigVMPin* InPin)
+				{
+					const FString DefaultValue = InPin->GetDefaultValue();
+					static const FString EmptyBraces = TEXT("()");
+					return DefaultValue.IsEmpty() || DefaultValue == EmptyBraces;
+				}
+				
+				static void ApplyResolvedDefaultValue(
+					URigVMController* InController, 
+					URigVMPin* InPin, 
+					const FString& RemainingPinPath, 
+					const FString& InDefaultValue, 
+					bool bSetupUndoRedo)
+				{
+					if(InDefaultValue.IsEmpty())
+					{
+						return;
+					}
+					
+					if(RemainingPinPath.IsEmpty())
+					{
+						InController->SetPinDefaultValue(InPin, InDefaultValue, true, bSetupUndoRedo, false);
+						return;
+					}
+
+					FString PinName;
+					FString SubPinPath;
+					if(!URigVMPin::SplitPinPathAtStart(RemainingPinPath, PinName, SubPinPath))
+					{
+						PinName = RemainingPinPath;
+						SubPinPath.Empty();
+					}
+
+					TArray<FString> MemberValuePairs = URigVMPin::SplitDefaultValue(InDefaultValue);
+					for (const FString& MemberValuePair : MemberValuePairs)
+					{
+						FString MemberName, MemberValue;
+						if (MemberValuePair.Split(TEXT("="), &MemberName, &MemberValue))
+						{
+							if(MemberName.Equals(PinName))
+							{
+								ApplyResolvedDefaultValue(InController, InPin, SubPinPath, MemberValue, bSetupUndoRedo);
+								break;
+							}
+						}
+					}
+				}
+			};
+			
+			for(URigVMPin* Pin : AllPins)
+			{
+				// skip struct pins or array pins
+				if(Pin->GetSubPins().Num() > 0)
+				{
+					continue;
+				}
+				
+				if(!Local::IsPinDefaultEmpty(Pin))
+				{
+					continue;
+				}
+
+				if(URigVMUnitNode* UnitNode = Cast<URigVMUnitNode>(Pin->GetNode()))
+				{
+					if(UnitNode->GetScriptStruct())
+					{
+						TSharedPtr<FStructOnScope> StructOnScope = UnitNode->ConstructStructInstance(true);
+						const FString StructDefaultValue = FRigVMStruct::ExportToFullyQualifiedText(UnitNode->GetScriptStruct(), StructOnScope->GetStructMemory());
+						Local::ApplyResolvedDefaultValue(this, Pin, Pin->GetSegmentPath(true), StructDefaultValue, bSetupUndoRedo);
+						if(!Local::IsPinDefaultEmpty(Pin))
+						{
+							continue;
+						}
+					}
+				}
+
+				// create the default value for the parent struct pin
+				if(Pin->IsStructMember())
+				{
+					const URigVMPin* ParentPin = Pin->GetParentPin();
+					TSharedPtr<FStructOnScope> StructOnScope = MakeShareable(new FStructOnScope(ParentPin->GetScriptStruct()));
+					ParentPin->GetScriptStruct()->InitializeDefaultValue((uint8*)StructOnScope->GetStructMemory());
+					const FString StructDefaultValue = FRigVMStruct::ExportToFullyQualifiedText(ParentPin->GetScriptStruct(), StructOnScope->GetStructMemory());
+					Local::ApplyResolvedDefaultValue(this, Pin, Pin->GetName(), StructDefaultValue, bSetupUndoRedo);
+				}
+				else
+				{
+					// plain types within an array or at the root
+					FString SimpleTypeDefaultValue;
+					if(Pin->GetCPPType() == RigVMTypeUtils::BoolType)
+					{
+						static const FString BoolDefaultValue = TEXT("False");
+						SimpleTypeDefaultValue = BoolDefaultValue;
+					}
+					else if(Pin->GetCPPType() == RigVMTypeUtils::FloatType || Pin->GetCPPType() == RigVMTypeUtils::DoubleType)
+					{
+						static const FString FloatingPointDefaultValue = TEXT("0.000000");
+						SimpleTypeDefaultValue = FloatingPointDefaultValue;
+					}
+					else if(Pin->GetCPPType() == RigVMTypeUtils::Int32Type)
+					{
+						static const FString IntegerDefaultValue = TEXT("0");
+						SimpleTypeDefaultValue = IntegerDefaultValue;
+					}
+					Local::ApplyResolvedDefaultValue(this, Pin, FString(), SimpleTypeDefaultValue, bSetupUndoRedo);
+				}
+			}
+			
+			// todo: update the default value for pins that don't have a default value
+		}
+	}
+
+	// since the resolved pin may affect the node title we need to let
+	// graph views know to invalidate the node title text widget
+	Notify(ERigVMGraphNotifType::NodeDescriptionChanged, InPin->GetNode());
+
 	// in cases were we are just changing the type we have to let the
 	// clients know that the links are still there
 	if(!bSetupOrphanPins && !bBreakLinks && !bRemoveSubPins)
diff --git a/Engine/Source/Developer/RigVMDeveloper/Private/RigVMModel/RigVMControllerActions.cpp b/Engine/Source/Developer/RigVMDeveloper/Private/RigVMModel/RigVMControllerActions.cpp
--- a/Engine/Source/Developer/RigVMDeveloper/Private/RigVMModel/RigVMControllerActions.cpp
+++ b/Engine/Source/Developer/RigVMDeveloper/Private/RigVMModel/RigVMControllerActions.cpp
@@ -588,21 +588,21 @@ bool FRigVMAddEnumNodeAction::Redo(URigVMController* InController)
 	return false;
 }
 
-FRigVMAddPrototypeNodeAction::FRigVMAddPrototypeNodeAction()
-	: PrototypeNotation(NAME_None)
+FRigVMAddTemplateNodeAction::FRigVMAddTemplateNodeAction()
+	: TemplateNotation(NAME_None)
 	, Position(FVector2D::ZeroVector)
 	, NodePath()
 {
 }
 
-FRigVMAddPrototypeNodeAction::FRigVMAddPrototypeNodeAction(URigVMPrototypeNode* InNode)
-	: PrototypeNotation(InNode->GetNotation())
+FRigVMAddTemplateNodeAction::FRigVMAddTemplateNodeAction(URigVMTemplateNode* InNode)
+	: TemplateNotation(InNode->GetNotation())
 	, Position(InNode->GetPosition())
 	, NodePath(InNode->GetNodePath())
 {
 }
 
-bool FRigVMAddPrototypeNodeAction::Undo(URigVMController* InController)
+bool FRigVMAddTemplateNodeAction::Undo(URigVMController* InController)
 {
 	if (!FRigVMBaseAction::Undo(InController))
 	{
@@ -611,10 +611,10 @@ bool FRigVMAddPrototypeNodeAction::Undo(URigVMController* InController)
 	return InController->RemoveNodeByName(*NodePath, false);
 }
 
-bool FRigVMAddPrototypeNodeAction::Redo(URigVMController* InController)
+bool FRigVMAddTemplateNodeAction::Redo(URigVMController* InController)
 {
 #if WITH_EDITOR
-	if (URigVMPrototypeNode* Node = InController->AddPrototypeNode(PrototypeNotation, Position, NodePath, false))
+	if (URigVMTemplateNode* Node = InController->AddTemplateNode(TemplateNotation, Position, NodePath, false))
 	{
 		return FRigVMBaseAction::Redo(InController);
 	}
@@ -714,9 +714,9 @@ FRigVMRemoveNodeAction::FRigVMRemoveNodeAction(URigVMNode* InNode, URigVMControl
 	{
 		InverseAction.AddAction(FRigVMAddSelectNodeAction(SelectNode));
 	}
-	else if (URigVMPrototypeNode* PrototypeNode = Cast<URigVMPrototypeNode>(InNode))
+	else if (URigVMTemplateNode* TemplateNode = Cast<URigVMTemplateNode>(InNode))
 	{
-		InverseAction.AddAction(FRigVMAddPrototypeNodeAction(PrototypeNode));
+		InverseAction.AddAction(FRigVMAddTemplateNodeAction(TemplateNode));
 	}
 	else if (URigVMEnumNode* EnumNode = Cast<URigVMEnumNode>(InNode))
 	{
diff --git a/Engine/Source/Developer/RigVMDeveloper/Private/RigVMModel/RigVMNode.cpp b/Engine/Source/Developer/RigVMDeveloper/Private/RigVMModel/RigVMNode.cpp
--- a/Engine/Source/Developer/RigVMDeveloper/Private/RigVMModel/RigVMNode.cpp
+++ b/Engine/Source/Developer/RigVMDeveloper/Private/RigVMModel/RigVMNode.cpp
@@ -292,11 +292,11 @@ bool URigVMNode::IsMutable() const
 	return false;
 }
 
-bool URigVMNode::HasUnknownTypePin() const
+bool URigVMNode::HasWildCardPin() const
 {
 	for (const URigVMPin* Pin : GetPins())
 	{
-		if (Pin->IsUnknownType())
+		if (Pin->IsWildCard())
 		{
 			return true;
 		}
diff --git a/Engine/Source/Developer/RigVMDeveloper/Private/RigVMModel/RigVMPin.cpp b/Engine/Source/Developer/RigVMDeveloper/Private/RigVMModel/RigVMPin.cpp
--- a/Engine/Source/Developer/RigVMDeveloper/Private/RigVMModel/RigVMPin.cpp
+++ b/Engine/Source/Developer/RigVMDeveloper/Private/RigVMModel/RigVMPin.cpp
@@ -1,11 +1,12 @@
 // Copyright Epic Games, Inc. All Rights Reserved.
 
 #include "RigVMModel/RigVMPin.h"
+
 #include "RigVMModel/RigVMNode.h"
 #include "RigVMModel/RigVMGraph.h"
 #include "RigVMModel/RigVMLink.h"
 #include "RigVMModel/RigVMController.h"
-#include "RigVMModel/Nodes/RigVMPrototypeNode.h"
+#include "RigVMModel/Nodes/RigVMTemplateNode.h"
 #include "RigVMCompiler/RigVMCompiler.h"
 #include "RigVMCore/RigVMExecuteContext.h"
 #include "RigVMCore/RigVMUnknownType.h"
@@ -538,7 +539,7 @@ bool URigVMPin::IsExecuteContext() const
 	return false;
 }
 
-bool URigVMPin::IsUnknownType() const
+bool URigVMPin::IsWildCard() const
 {
 	if (const UScriptStruct* ScriptStruct = GetScriptStruct())
 	{
@@ -1509,72 +1510,45 @@ bool URigVMPin::CanLink(URigVMPin* InSourcePin, URigVMPin* InTargetPin, FString*
 		{
 			if(InSourcePin->IsArray() == InTargetPin->IsArray())
 			{
-				if(InSourcePin->IsUnknownType() && !InTargetPin->IsExecuteContext())
-				{
-					
-					return true;
-				}
-				else if(InTargetPin->IsUnknownType() && !InSourcePin->IsExecuteContext())
+				auto TemplateNodeSupportsType = [](URigVMPin* InPin, const FString& InCPPType, FString* OutFailureReason) -> bool
 				{
-					return true;
-				}
-			}
-		}
-
-		if (bCPPTypesDiffer)
-		{
-			if (OutFailureReason)
-			{
-				*OutFailureReason = TEXT("Source and target pin types are not compatible.");
-			}
 
-			// check if this might be a prototype node
-			if (InSourcePin->CPPType.IsEmpty())
-			{
-				if (URigVMPrototypeNode* PrototypeNode = Cast<URigVMPrototypeNode>(SourceNode))
-				{
-					if (PrototypeNode->SupportsType(InSourcePin, InTargetPin->CPPType))
+					if (URigVMTemplateNode* TemplateNode = Cast<URigVMTemplateNode>(InPin->GetNode()))
 					{
-						if (OutFailureReason)
+						if (TemplateNode->SupportsType(InPin, InCPPType))
 						{
-							*OutFailureReason = FString();
+							if (OutFailureReason)
+							{
+								*OutFailureReason = FString();
+							}
 						}
-					}
-					else
-					{
-						return false;
-					}
-				}
-				else
-				{
-					return false;
-				}
-			}
-			else if (InTargetPin->CPPType.IsEmpty())
-			{
-				if (URigVMPrototypeNode* PrototypeNode = Cast<URigVMPrototypeNode>(TargetNode))
-				{
-					if (PrototypeNode->SupportsType(InTargetPin, InSourcePin->CPPType))
-					{
-						if (OutFailureReason)
+						else
 						{
-							*OutFailureReason = FString();
+							return false;
 						}
 					}
-					else
-					{
-						return false;
-					}
+					return true;
+
+				};
+				
+				if(InSourcePin->IsWildCard() && !InTargetPin->IsExecuteContext())
+				{
+					return TemplateNodeSupportsType(InSourcePin, InTargetPin->GetCPPType(), OutFailureReason);
 				}
-				else
+				else if(InTargetPin->IsWildCard() && !InSourcePin->IsExecuteContext())
 				{
-					return false;
+					return TemplateNodeSupportsType(InTargetPin, InSourcePin->GetCPPType(), OutFailureReason);
 				}
 			}
-			else
+		}
+
+		if (bCPPTypesDiffer)
+		{
+			if (OutFailureReason)
 			{
-				return false;
+				*OutFailureReason = TEXT("Source and target pin types are not compatible.");
 			}
+			return false;
 		}
 	}
 
diff --git a/Engine/Source/Developer/RigVMDeveloper/Public/RigVMDeveloperTypeUtils.h b/Engine/Source/Developer/RigVMDeveloper/Public/RigVMDeveloperTypeUtils.h
--- a/Engine/Source/Developer/RigVMDeveloper/Public/RigVMDeveloperTypeUtils.h
+++ b/Engine/Source/Developer/RigVMDeveloper/Public/RigVMDeveloperTypeUtils.h
@@ -8,6 +8,7 @@
 #include "RigVMTypeUtils.h"
 
 struct FRigVMGraphVariableDescription;
+class URigVMPin;
 
 namespace RigVMTypeUtils
 {
@@ -15,7 +16,7 @@ namespace RigVMTypeUtils
 	RIGVMDEVELOPER_API FRigVMExternalVariable ExternalVariableFromBPVariableDescription(const FBPVariableDescription& InVariableDescription);
 
 	RIGVMDEVELOPER_API FRigVMExternalVariable ExternalVariableFromRigVMVariableDescription(const FRigVMGraphVariableDescription& InVariableDescription);
-	
+
 	RIGVMDEVELOPER_API FRigVMExternalVariable ExternalVariableFromPinType(const FName& InName, const FEdGraphPinType& InPinType, bool bInPublic = false, bool bInReadonly = false);
 
 	RIGVMDEVELOPER_API FRigVMExternalVariable ExternalVariableFromCPPTypePath(const FName& InName, const FString& InCPPTypePath, bool bInPublic = false, bool bInReadonly = false);
@@ -30,10 +31,18 @@ namespace RigVMTypeUtils
 
 	RIGVMDEVELOPER_API FEdGraphPinType SubPinType(const FEdGraphPinType& InPinType, const FString& SegmentPath);
 
+	RIGVMDEVELOPER_API bool CPPTypeFromPin(URigVMPin* InPin, FString& OutCPPType, UObject** OutCPPTypeObject, bool bGetBaseCPPType = false);
+
+	RIGVMDEVELOPER_API bool CPPTypeFromPin(URigVMPin* InPin, FString& OutCPPType, FName& OutCPPTypeObjectPath, bool bGetBaseCPPType = false);
+
+	RIGVMDEVELOPER_API bool CPPTypeFromPin(URigVMPin* InPin, FString& OutCPPType, FString& OutCPPTypeObjectPath, bool bGetBaseCPPType = false);
+
 	RIGVMDEVELOPER_API bool CPPTypeFromPinType(const FEdGraphPinType& InPinType, FString& OutCPPType, UObject** OutCPPTypeObject);
 
 	RIGVMDEVELOPER_API bool CPPTypeFromPinType(const FEdGraphPinType& InPinType, FString& OutCPPType, FName& OutCPPTypeObjectPath);
 
+	RIGVMDEVELOPER_API bool CPPTypeFromPinType(const FEdGraphPinType& InPinType, FString& OutCPPType, FString& OutCPPTypeObjectPath);
+
 	RIGVMDEVELOPER_API bool CPPTypeFromExternalVariable(const FRigVMExternalVariable& InExternalVariable, FString& OutCPPType, UObject** OutCPPTypeObject);
 
 	RIGVMDEVELOPER_API bool AreCompatible(const FString& InCPPTypeA, UObject* InCPPTypeObjectA, const FString& InCPPTypeB, UObject* InCPPTypeObjectB);
diff --git a/Engine/Source/Developer/RigVMDeveloper/Public/RigVMModel/Nodes/RigVMPrototypeNode.h b/Engine/Source/Developer/RigVMDeveloper/Public/RigVMModel/Nodes/RigVMPrototypeNode.h
deleted file mode 100644
diff --git a/Engine/Source/Developer/RigVMDeveloper/Public/RigVMModel/Nodes/RigVMTemplateNode.h b/Engine/Source/Developer/RigVMDeveloper/Public/RigVMModel/Nodes/RigVMTemplateNode.h
new file mode 100644
--- /dev/null
+++ b/Engine/Source/Developer/RigVMDeveloper/Public/RigVMModel/Nodes/RigVMTemplateNode.h
@@ -0,0 +1,72 @@
+// Copyright Epic Games, Inc. All Rights Reserved.
+
+#pragma once
+
+#include "RigVMUnitNode.h"
+#include "RigVMModel/Nodes/RigVMUnitNode.h"
+#include "RigVMCore/RigVMTemplate.h"
+#include "RigVMTemplateNode.generated.h"
+
+/**
+ * The Template Node represents an unresolved function.
+ * Template nodes can morph into all functions implementing
+ * the template's template.
+ */
+UCLASS(BlueprintType)
+class RIGVMDEVELOPER_API URigVMTemplateNode : public URigVMUnitNode
+{
+	GENERATED_BODY()
+
+public:
+
+	// default constructor
+	URigVMTemplateNode();
+
+	// URigVMUnitNode interface
+	virtual UScriptStruct* GetScriptStruct() const;
+	virtual FString GetNodeTitle() const override;
+	virtual FName GetMethodName() const;
+	virtual  FText GetToolTipText() const override;
+
+	virtual FText GetToolTipTextForPin(const URigVMPin* InPin) const override;
+
+	// Returns the notation of the node
+	UFUNCTION(BlueprintCallable, Category = Template)
+	FName GetNotation() const;
+
+	// returns true if a pin supports a given type
+	bool SupportsType(const URigVMPin* InPin, const FString& InCPPType, FString* OutCPPType = nullptr);
+
+	// returns the index of the resolved RigVM function (or INDEX_NONE)
+	TArray<int32> GetResolvedPermutationIndices(FRigVMTemplate::FTypeMap* OutTypes = nullptr) const;
+
+	// returns the resolved functions for the template
+	TArray<const FRigVMFunction*> GetResolvedPermutations(FRigVMTemplate::FTypeMap* OutTypes = nullptr) const;
+
+	// returns the template used for this node
+	const FRigVMTemplate* GetTemplate() const;
+
+	// returns the resolved function or nullptr if there are still unresolved pins left
+	const FRigVMFunction* GetResolvedFunction() const;
+
+	// returns true if the template node is resolved
+	bool IsResolved() const;
+
+private:
+
+	void InvalidateCache();
+
+	UPROPERTY()
+	FName TemplateNotation;
+
+	UPROPERTY()
+	FString ResolvedFunctionName;
+
+	TMap<FString, TPair<bool, FRigVMTemplateArgument::FType>> SupportedTypesCache;
+
+	mutable const FRigVMTemplate* CachedTemplate;
+	mutable const FRigVMFunction* CachedFunction;
+
+	friend class URigVMController;
+};
+
diff --git a/Engine/Source/Developer/RigVMDeveloper/Public/RigVMModel/Nodes/RigVMUnitNode.h b/Engine/Source/Developer/RigVMDeveloper/Public/RigVMModel/Nodes/RigVMUnitNode.h
--- a/Engine/Source/Developer/RigVMDeveloper/Public/RigVMModel/Nodes/RigVMUnitNode.h
+++ b/Engine/Source/Developer/RigVMDeveloper/Public/RigVMModel/Nodes/RigVMUnitNode.h
@@ -33,11 +33,11 @@ public:
 	// Returns the UStruct for this unit node
 	// (the struct declaring the RIGVM_METHOD)
 	UFUNCTION(BlueprintCallable, Category = RigVMUnitNode)
-	UScriptStruct* GetScriptStruct() const;
+	virtual UScriptStruct* GetScriptStruct() const;
 
 	// Returns the name of the declared RIGVM_METHOD
 	UFUNCTION(BlueprintCallable, Category = RigVMUnitNode)
-	FName GetMethodName() const;
+	virtual FName GetMethodName() const;
 
 	// Returns the default value for the struct as text
 	UFUNCTION(BlueprintCallable, Category = RigVMUnitNode)
diff --git a/Engine/Source/Developer/RigVMDeveloper/Public/RigVMModel/RigVMController.h b/Engine/Source/Developer/RigVMDeveloper/Public/RigVMModel/RigVMController.h
--- a/Engine/Source/Developer/RigVMDeveloper/Public/RigVMModel/RigVMController.h
+++ b/Engine/Source/Developer/RigVMDeveloper/Public/RigVMModel/RigVMController.h
@@ -12,7 +12,7 @@
 #include "RigVMModel/Nodes/RigVMBranchNode.h"
 #include "RigVMModel/Nodes/RigVMIfNode.h"
 #include "RigVMModel/Nodes/RigVMSelectNode.h"
-#include "RigVMModel/Nodes/RigVMPrototypeNode.h"
+#include "RigVMModel/Nodes/RigVMTemplateNode.h"
 #include "RigVMModel/Nodes/RigVMEnumNode.h"
 #include "RigVMModel/Nodes/RigVMCollapseNode.h"
 #include "RigVMModel/Nodes/RigVMFunctionReferenceNode.h"
@@ -20,6 +20,10 @@
 #include "RigVMModel/RigVMBuildData.h"
 #include "RigVMController.generated.h"
 
+#ifndef UE_RIGVM_ENABLE_TEMPLATE_NODES
+#define UE_RIGVM_ENABLE_TEMPLATE_NODES 0
+#endif
+
 class URigVMActionStack;
 
 UENUM()
@@ -209,6 +213,10 @@ public:
 	UFUNCTION(BlueprintCallable, Category = RigVMController)
 	URigVMVariableNode* ReplaceParameterNodeWithVariable(const FName& InNodeName, const FName& InVariableName, const FString& InCPPType, UObject* InCPPTypeObject, bool bSetupUndoRedo);
 
+	// Refreshes the unit node to be backed up by a template instead 
+	UFUNCTION(BlueprintCallable, Category = RigVMController)
+	URigVMTemplateNode* ReplaceUnitNodeWithTemplateNode(const FName& InNodeName, bool bSetupUndoRedo);
+
 	// Adds a Parameter Node to the edited Graph.
 	// Parameters represent input or output arguments to the Graph / Function.
 	// Input Parameters are constant values / literals.
@@ -274,10 +282,13 @@ public:
 	UFUNCTION(BlueprintCallable, Category = RigVMController)
 	URigVMSelectNode* AddSelectNodeFromStruct(UScriptStruct* InScriptStruct, const FVector2D& InPosition = FVector2D::ZeroVector, const FString& InNodeName = TEXT(""), bool bSetupUndoRedo = true);
 
+	// Adds a template node to the graph.
+	UFUNCTION(BlueprintCallable, Category = RigVMController)
+	URigVMTemplateNode* AddTemplateNode(const FName& InNotation, const FVector2D& InPosition = FVector2D::ZeroVector, const FString& InNodeName = TEXT(""), bool bSetupUndoRedo = true, bool bPrintPythonCommand = false);
 
-	// Adds a prototype node to the graph.
+	// Resolves a wildcard pin on any node
 	UFUNCTION(BlueprintCallable, Category = RigVMController)
-	URigVMPrototypeNode* AddPrototypeNode(const FName& InNotation, const FVector2D& InPosition = FVector2D::ZeroVector, const FString& InNodeName = TEXT(""), bool bSetupUndoRedo = true, bool bPrintPythonCommand = false);
+	bool ResolveWildCardPin(const FString& InPinPath, const FString& InCPPType, const FName& InCPPTypeObjectPath, bool bSetupUndoRedo = true, bool bPrintPythonCommand = false);
 
 	// Adds a Function / Struct Node to the edited Graph as an injected node
 	// UnitNode represent a RIGVM_METHOD declaration on a USTRUCT.
@@ -687,7 +698,7 @@ public:
 	TArray<FAssetData> GetAffectedAssets(ERigVMControllerBulkEditType InEditType, bool bForceLoad = false, bool bNotify = true);
 
 	// Sets the execute context struct type to use
-	void SetExecuteContextStruct(UStruct* InExecuteContextStruct);
+	void SetExecuteContextStruct(UScriptStruct* InExecuteContextStruct);
 
 	// A delegate that can be set to change the struct unfolding behaviour
 	FRigVMController_ShouldStructUnfoldDelegate UnfoldStructDelegate;
@@ -825,7 +836,7 @@ private:
 
 	// try to reconnect source and target pins after a node deletion
 	void RelinkSourceAndTargetPins(URigVMNode* RigNode, bool bSetupUndoRedo = true);
-	
+
 public:
 	bool AddLink(URigVMPin* OutputPin, URigVMPin* InputPin, bool bSetupUndoRedo = true);
 	bool BreakLink(URigVMPin* OutputPin, URigVMPin* InputPin, bool bSetupUndoRedo = true);
@@ -871,11 +882,11 @@ private:
 	static void PostProcessDefaultValue(URigVMPin* Pin, FString& OutDefaultValue);
 	static FString PostProcessCPPType(const FString& InCPPType, UObject* InCPPTypeObject);
 
-	/*
-	void PotentiallyResolvePrototypeNode(URigVMPrototypeNode* InNode, bool bSetupUndoRedo);
-	void PotentiallyResolvePrototypeNode(URigVMPrototypeNode* InNode, bool bSetupUndoRedo, TArray<URigVMNode*>& NodesVisited);
-	*/
-	void ResolveUnknownTypePin(URigVMPin* InPinToResolve, const URigVMPin* InTemplatePin, bool bSetupUndoRedo, bool bTraverseNode = true, bool bTraverseParentPins = true, bool bTraverseLinks = true);
+	bool ResolveWildCardPin(URigVMPin* InPinToResolve, URigVMPin* InTemplatePin, bool bSetupUndoRedo, bool bTraverseNode = true, bool bTraverseParentPins = true, bool bTraverseLinks = true);
+	bool ResolveWildCardPinImpl(URigVMPin* InPinToResolve, const FString& InCPPType, UObject* InCPPTypeObject, bool bSetupUndoRedo, bool bTraverseNode = true, bool bTraverseParentPins = true, bool bTraverseLinks = true);
+	void ResolveTemplateNodePins(URigVMTemplateNode* InNode, bool bSetupUndoRedo);
+	void ResolveTemplateNodeMetaData(URigVMTemplateNode* InNode);
+	bool FullyResolveTemplateNode(URigVMTemplateNode* InNode, const FRigVMFunction* InFunctionToResolve, bool bSetupUndoRedo);
 	bool ChangePinType(const FString& InPinPath, const FString& InCPPType, const FName& InCPPTypeObjectPath, bool bSetupUndoRedo, bool bSetupOrphanPins = true, bool bBreakLinks = true, bool bRemoveSubPins = true);
 	bool ChangePinType(URigVMPin* InPin, const FString& InCPPType, const FName& InCPPTypeObjectPath, bool bSetupUndoRedo, bool bSetupOrphanPins = true, bool bBreakLinks = true, bool bRemoveSubPins = true);
 	bool ChangePinType(URigVMPin* InPin, const FString& InCPPType,UObject* InCPPTypeObject, bool bSetupUndoRedo, bool bSetupOrphanPins = true, bool bBreakLinks = true, bool bRemoveSubPins = true);
@@ -918,7 +929,7 @@ private:
 	TArray<TObjectPtr<URigVMGraph>> Graphs;
 
 	UPROPERTY()
-	TObjectPtr<UStruct> ExecuteContextStruct;
+	TObjectPtr<UScriptStruct> ExecuteContextStruct;
 
 	UPROPERTY(transient)
 	TObjectPtr<URigVMActionStack> ActionStack;
@@ -969,6 +980,7 @@ private:
 	FRigVMUnitNodeCreatedContext UnitNodeCreatedContext;
 
 	bool bIsRunningUnitTest;
+	bool bIsFullyResolvingTemplateNode;
 
 	friend class URigVMGraph;
 	friend class URigVMPin;
diff --git a/Engine/Source/Developer/RigVMDeveloper/Public/RigVMModel/RigVMControllerActions.h b/Engine/Source/Developer/RigVMDeveloper/Public/RigVMModel/RigVMControllerActions.h
--- a/Engine/Source/Developer/RigVMDeveloper/Public/RigVMModel/RigVMControllerActions.h
+++ b/Engine/Source/Developer/RigVMDeveloper/Public/RigVMModel/RigVMControllerActions.h
@@ -596,23 +596,23 @@ public:
 };
 
 /**
- * An action adding a prototype node to the graph.
+ * An action adding a template node to the graph.
  */
 USTRUCT()
-struct FRigVMAddPrototypeNodeAction : public FRigVMBaseAction
+struct FRigVMAddTemplateNodeAction : public FRigVMBaseAction
 {
 	GENERATED_BODY()
 
 public:
 
-	FRigVMAddPrototypeNodeAction();
-	FRigVMAddPrototypeNodeAction(URigVMPrototypeNode* InNode);
-	virtual ~FRigVMAddPrototypeNodeAction() {};
+	FRigVMAddTemplateNodeAction();
+	FRigVMAddTemplateNodeAction(URigVMTemplateNode* InNode);
+	virtual ~FRigVMAddTemplateNodeAction() {};
 	virtual bool Undo(URigVMController* InController) override;
 	virtual bool Redo(URigVMController* InController) override;
 
 	UPROPERTY()
-	FName PrototypeNotation;
+	FName TemplateNotation;
 
 	UPROPERTY()
 	FVector2D Position;
diff --git a/Engine/Source/Developer/RigVMDeveloper/Public/RigVMModel/RigVMNode.h b/Engine/Source/Developer/RigVMDeveloper/Public/RigVMModel/RigVMNode.h
--- a/Engine/Source/Developer/RigVMDeveloper/Public/RigVMModel/RigVMNode.h
+++ b/Engine/Source/Developer/RigVMDeveloper/Public/RigVMModel/RigVMNode.h
@@ -140,7 +140,7 @@ public:
 	virtual bool IsMutable() const;
 
 	// Returns true if this node has an unknown type pin
-	bool HasUnknownTypePin() const;
+	bool HasWildCardPin() const;
 
 	virtual bool ContributesToResult() const { return IsMutable(); }
 
diff --git a/Engine/Source/Developer/RigVMDeveloper/Public/RigVMModel/RigVMPin.h b/Engine/Source/Developer/RigVMDeveloper/Public/RigVMModel/RigVMPin.h
--- a/Engine/Source/Developer/RigVMDeveloper/Public/RigVMModel/RigVMPin.h
+++ b/Engine/Source/Developer/RigVMDeveloper/Public/RigVMModel/RigVMPin.h
@@ -216,7 +216,7 @@ public:
 
 	// Returns true if the C++ data type is unknown
 	UFUNCTION(BlueprintCallable, Category = RigVMPin)
-	bool IsUnknownType() const;
+	bool IsWildCard() const;
 
 	// Returns the default value of the Pin as a string.
 	// Note that this value is computed based on the Pin's
diff --git a/Engine/Source/Runtime/RigVM/Private/RigVMCore/RigVM.cpp b/Engine/Source/Runtime/RigVM/Private/RigVMCore/RigVM.cpp
--- a/Engine/Source/Runtime/RigVM/Private/RigVMCore/RigVM.cpp
+++ b/Engine/Source/Runtime/RigVM/Private/RigVMCore/RigVM.cpp
@@ -596,14 +596,14 @@ int32 URigVM::AddRigVMFunction(UScriptStruct* InRigVMStruct, const FName& InMeth
 		return FunctionIndex;
 	}
 
-	FRigVMFunctionPtr Function = FRigVMRegistry::Get().FindFunction(*FunctionKey);
+	const FRigVMFunction* Function = FRigVMRegistry::Get().FindFunction(*FunctionKey);
 	if (Function == nullptr)
 	{
 		return INDEX_NONE;
 	}
 
 	GetFunctionNames().Add(*FunctionKey);
-	return GetFunctions().Add(Function);
+	return GetFunctions().Add(Function->FunctionPtr);
 }
 
 FString URigVM::GetRigVMFunctionName(int32 InFunctionIndex) const
@@ -834,7 +834,11 @@ void URigVM::ResolveFunctionsIfRequired()
 
 		for (int32 FunctionIndex = 0; FunctionIndex < GetFunctionNames().Num(); FunctionIndex++)
 		{
-			GetFunctions()[FunctionIndex] = FRigVMRegistry::Get().FindFunction(*GetFunctionNames()[FunctionIndex].ToString());
+			GetFunctions()[FunctionIndex] = nullptr;
+			if(const FRigVMFunction* Function = FRigVMRegistry::Get().FindFunction(*GetFunctionNames()[FunctionIndex].ToString()))
+			{
+				GetFunctions()[FunctionIndex] = Function->FunctionPtr;
+			}
 			ensureMsgf(GetFunctions()[FunctionIndex], TEXT("Function %s is not valid"), *GetFunctionNames()[FunctionIndex].ToString());
 		}
 	}
diff --git a/Engine/Source/Runtime/RigVM/Private/RigVMCore/RigVMFunction.cpp b/Engine/Source/Runtime/RigVM/Private/RigVMCore/RigVMFunction.cpp
--- a/Engine/Source/Runtime/RigVM/Private/RigVMCore/RigVMFunction.cpp
+++ b/Engine/Source/Runtime/RigVM/Private/RigVMCore/RigVMFunction.cpp
@@ -1,6 +1,8 @@
 // Copyright Epic Games, Inc. All Rights Reserved.
 
 #include "RigVMCore/RigVMFunction.h"
+#include "RigVMCore/RigVMTemplate.h"
+#include "RigVMCore/RigVMRegistry.h"
 
 FName FRigVMFunction::GetMethodName() const
 {
@@ -42,3 +44,18 @@ FString FRigVMFunction::GetModuleRelativeHeaderPath() const
 	return FString();
 }
 
+const FRigVMTemplate* FRigVMFunction::GetTemplate() const
+{
+	if(TemplateIndex == INDEX_NONE)
+	{
+		return nullptr;
+	}
+
+	const FRigVMTemplate* Template = &FRigVMRegistry::Get().GetTemplates()[TemplateIndex];
+	if(Template->NumPermutations() <= 1)
+	{
+		return nullptr;
+	}
+
+	return Template;
+}
diff --git a/Engine/Source/Runtime/RigVM/Private/RigVMCore/RigVMPrototype.cpp b/Engine/Source/Runtime/RigVM/Private/RigVMCore/RigVMPrototype.cpp
deleted file mode 100644
diff --git a/Engine/Source/Runtime/RigVM/Private/RigVMCore/RigVMRegistry.cpp b/Engine/Source/Runtime/RigVM/Private/RigVMCore/RigVMRegistry.cpp
--- a/Engine/Source/Runtime/RigVM/Private/RigVMCore/RigVMRegistry.cpp
+++ b/Engine/Source/Runtime/RigVM/Private/RigVMCore/RigVMRegistry.cpp
@@ -5,7 +5,7 @@
 #include "UObject/UObjectIterator.h"
 
 FRigVMRegistry FRigVMRegistry::s_RigVMRegistry;
-const FName FRigVMRegistry::PrototypeNameMetaName = TEXT("PrototypeName");
+const FName FRigVMRegistry::TemplateNameMetaName = TEXT("PrototypeName");
 
 FRigVMRegistry& FRigVMRegistry::Get()
 {
@@ -23,13 +23,14 @@ void FRigVMRegistry::Register(const TCHAR* InName, FRigVMFunctionPtr InFunctionP
 		return;
 	}
 
-	FRigVMFunction Function(InName, InFunctionPtr, InStruct, Functions.Num());
+	const FRigVMFunction Function(InName, InFunctionPtr, InStruct, Functions.Num());
 	Functions.Add(Function);
+	FunctionNameToIndex.Add(InName, Function.Index);
 
 #if WITH_EDITOR
 	
-	FString PrototypeMetadata;
-	if (InStruct->GetStringMetaDataHierarchical(PrototypeNameMetaName, &PrototypeMetadata))
+	FString TemplateMetadata;
+	if (InStruct->GetStringMetaDataHierarchical(TemplateNameMetaName, &TemplateMetadata))
 	{
 		if(InStruct->HasMetaData(FRigVMStruct::DeprecatedMetaName))
 		{
@@ -39,25 +40,33 @@ void FRigVMRegistry::Register(const TCHAR* InName, FRigVMFunctionPtr InFunctionP
 		FString MethodName;
 		if (FString(InName).Split(TEXT("::"), nullptr, &MethodName))
 		{
-			FString PrototypeName = FString::Printf(TEXT("%s::%s"), *PrototypeMetadata, *MethodName);
-			FRigVMPrototype Prototype(InStruct, PrototypeName, Function.Index);
-			if (Prototype.IsValid())
+			const FString TemplateName = FString::Printf(TEXT("%s::%s"), *TemplateMetadata, *MethodName);
+			FRigVMTemplate Template(InStruct, TemplateName, Function.Index);
+			if (Template.IsValid())
 			{
 				bool bWasMerged = false;
-				for (FRigVMPrototype& ExistingPrototype : Prototypes)
+
+				const int32* ExistingTemplateIndexPtr = TemplateNotationToIndex.Find(Template.GetNotation());
+				if(ExistingTemplateIndexPtr)
 				{
-					if (ExistingPrototype.Merge(Prototype))
+					FRigVMTemplate& ExistingTemplate = Templates[*ExistingTemplateIndexPtr];
+					if (ExistingTemplate.Merge(Template))
 					{
-						Functions[Function.Index].PrototypeIndex = ExistingPrototype.Index;
+						Functions[Function.Index].TemplateIndex = ExistingTemplate.Index;
 						bWasMerged = true;
 					}
 				}
 
 				if (!bWasMerged)
 				{
-					Prototype.Index = Prototypes.Num();
-					Functions[Function.Index].PrototypeIndex = Prototype.Index;
-					Prototypes.Add(Prototype);
+					Template.Index = Templates.Num();
+					Functions[Function.Index].TemplateIndex = Template.Index;
+					Templates.Add(Template);
+
+					if(ExistingTemplateIndexPtr == nullptr)
+					{
+						TemplateNotationToIndex.Add(Template.GetNotation(), Template.Index);
+					}
 				}
 			}
 		}
@@ -66,65 +75,47 @@ void FRigVMRegistry::Register(const TCHAR* InName, FRigVMFunctionPtr InFunctionP
 #endif
 }
 
-FRigVMFunctionPtr FRigVMRegistry::FindFunction(const TCHAR* InName) const
+const FRigVMFunction* FRigVMRegistry::FindFunction(const TCHAR* InName) const
 {
-	for (const FRigVMFunction& Function : Functions)
+	if(const int32* FunctionIndexPtr = FunctionNameToIndex.Find(InName))
 	{
-		if (FCString::Strcmp(Function.Name, InName) == 0)
-		{
-			return Function.FunctionPtr;
-		}
+		return &Functions[*FunctionIndexPtr];
 	}
-
 	return nullptr;
 }
 
-const FRigVMPrototype* FRigVMRegistry::FindPrototype(const FName& InNotation) const
+const FRigVMFunction* FRigVMRegistry::FindFunction(UScriptStruct* InStruct, const TCHAR* InName) const
 {
-	if (InNotation.IsNone())
-	{
-		return nullptr;
-	}
-
-	for (const FRigVMPrototype& Prototype : Prototypes)
-	{
-		if (Prototype.GetNotation() == InNotation)
-		{
-			return &Prototype;
-		}
-	}
-
-	return nullptr;
+	check(InStruct);
+	check(InName);
+	
+	const FString FunctionName = FString::Printf(TEXT("%s::%s"), *InStruct->GetStructCPPName(), InName);
+	return FindFunction(*FunctionName);
 }
 
-// Returns a prototype pointer given its notation (or nullptr)
-const FRigVMPrototype* FRigVMRegistry::FindPrototype(UScriptStruct* InStruct, const FString& InPrototypeName) const
+const TArray<FRigVMFunction>& FRigVMRegistry::GetFunctions() const
 {
-	FName Notation = FRigVMPrototype::GetNotationFromStruct(InStruct, InPrototypeName);
-	return FindPrototype(Notation);
+	return Functions;
 }
 
-FRigVMFunction FRigVMRegistry::FindFunctionInfo(const TCHAR* InName) const
+const FRigVMTemplate* FRigVMRegistry::FindTemplate(const FName& InNotation) const
 {
-	for (const FRigVMFunction& Function : Functions)
+	if (InNotation.IsNone())
 	{
-		if (FCString::Strcmp(Function.Name, InName) == 0)
-		{
-			return Function;
-		}
-	} 
+		return nullptr;
+	}
 
-	return FRigVMFunction();
-}
+	if(const int32* TemplateIndexPtr = TemplateNotationToIndex.Find(InNotation))
+	{
+		return &Templates[*TemplateIndexPtr];
+	}
 
-const TArray<FRigVMFunction>& FRigVMRegistry::GetFunctions() const
-{
-	return Functions;
+	return nullptr;
 }
 
-const TArray<FRigVMPrototype>& FRigVMRegistry::GetPrototypes() const
+const TArray<FRigVMTemplate>& FRigVMRegistry::GetTemplates() const
 {
-	return Prototypes;
+	return Templates;
 }
 
 
diff --git a/Engine/Source/Runtime/RigVM/Private/RigVMCore/RigVMStruct.cpp b/Engine/Source/Runtime/RigVM/Private/RigVMCore/RigVMStruct.cpp
--- a/Engine/Source/Runtime/RigVM/Private/RigVMCore/RigVMStruct.cpp
+++ b/Engine/Source/Runtime/RigVM/Private/RigVMCore/RigVMStruct.cpp
@@ -1,6 +1,7 @@
 // Copyright Epic Games, Inc. All Rights Reserved.
 
 #include "RigVMCore/RigVMStruct.h"
+#include "RigVMCore/RigVMRegistry.h"
 #include "UObject/StructOnScope.h"
 
 ////////////////////////////////////////////////////////////////////////////////////////////////////////////////////////
@@ -103,7 +104,7 @@ const FName FRigVMStruct::TitleColorMetaName("TitleColor");
 const FName FRigVMStruct::NodeColorMetaName("NodeColor");
 const FName FRigVMStruct::IconMetaName("Icon");
 const FName FRigVMStruct::KeywordsMetaName("Keywords");
-const FName FRigVMStruct::PrototypeNameMetaName("PrototypeName");
+const FName FRigVMStruct::TemplateNameMetaName = FRigVMRegistry::TemplateNameMetaName;
 const FName FRigVMStruct::ExpandPinByDefaultMetaName("ExpandByDefault");
 const FName FRigVMStruct::DefaultArraySizeMetaName("DefaultArraySize");
 const FName FRigVMStruct::VaryingMetaName("Varying");
diff --git a/Engine/Source/Runtime/RigVM/Private/RigVMCore/RigVMTemplate.cpp b/Engine/Source/Runtime/RigVM/Private/RigVMCore/RigVMTemplate.cpp
new file mode 100644
--- /dev/null
+++ b/Engine/Source/Runtime/RigVM/Private/RigVMCore/RigVMTemplate.cpp
@@ -0,0 +1,698 @@
+// Copyright Epic Games, Inc. All Rights Reserved.
+
+#include "RigVMCore/RigVMTemplate.h"
+#include "RigVMCore/RigVMRegistry.h"
+#include "RigVMCore/RigVMStruct.h"
+#include "RigVMModule.h"
+#include "Algo/Sort.h"
+
+////////////////////////////////////////////////////////////////////////////////////////////////
+
+FRigVMTemplateArgument::FRigVMTemplateArgument()
+	: Index(INDEX_NONE)
+	, Name(NAME_None)
+	, Direction(ERigVMPinDirection::IO)
+	, bSingleton(true)
+{
+}
+
+FRigVMTemplateArgument::FRigVMTemplateArgument(FProperty* InProperty)
+	: Index(INDEX_NONE)
+	, Name(InProperty->GetFName())
+	, Direction(ERigVMPinDirection::IO)
+	, bSingleton(true)
+{
+#if WITH_EDITOR
+	Direction = FRigVMStruct::GetPinDirectionFromProperty(InProperty);
+#endif
+
+	FString ExtendedType;
+	const FString CPPType = InProperty->GetCPPType(&ExtendedType);
+	FType Type(CPPType + ExtendedType);
+
+	if (FArrayProperty* ArrayProperty = CastField<FArrayProperty>(InProperty))
+	{
+		InProperty = ArrayProperty->Inner;
+	}
+
+	if (FStructProperty* StructProperty = CastField<FStructProperty>(InProperty))
+	{
+		Type.CPPTypeObject = StructProperty->Struct;
+	}
+	else if (FEnumProperty* EnumProperty = CastField<FEnumProperty>(InProperty))
+	{
+		Type.CPPTypeObject = EnumProperty->GetEnum();
+	}
+
+	Types.Add(Type);
+}
+
+bool FRigVMTemplateArgument::SupportsType(const FString& InCPPType, const TArray<int32>& InPermutationIndices, FType* OutType) const
+{
+	bool bFoundMatch = false;
+	bool bFoundPerfectMatch = false;
+	
+	auto VisitType = [&bFoundMatch, &bFoundPerfectMatch, InCPPType, &OutType](const FType& InType)
+	{
+		if (InType.Matches(InCPPType))
+		{
+			if(!bFoundPerfectMatch)
+			{
+				if(OutType)
+				{
+					*OutType = InType;
+				}
+				if(InType.CPPType == InCPPType)
+				{
+					bFoundPerfectMatch = true;
+				}
+			}
+			bFoundMatch = true;
+		}
+	};
+	
+	if(InPermutationIndices.IsEmpty())
+	{
+		for (const FType& Type : Types)
+		{
+			VisitType(Type);
+		}
+	}
+	else
+	{
+		for (const int32 PermutationIndex : InPermutationIndices)
+		{
+			VisitType(Types[PermutationIndex]);
+		}
+	}
+	return bFoundMatch;
+}
+
+bool FRigVMTemplateArgument::IsSingleton(const TArray<int32>& InPermutationIndices) const
+{
+	if (bSingleton)
+	{
+		return true;
+	}
+	else if(InPermutationIndices.Num() == 0)
+	{
+		return false;
+	}
+
+	const FType TypeToCheck = Types[InPermutationIndices[0]];
+	for (int32 PermutationIndex = 1; PermutationIndex < InPermutationIndices.Num(); PermutationIndex++)
+	{
+		if (Types[InPermutationIndices[PermutationIndex]] != TypeToCheck)
+		{
+			return false;
+		}
+	}
+	return true;
+}
+
+bool FRigVMTemplateArgument::IsArray() const
+{
+	if(Types.Num() > 0)
+	{
+		return Types[0].IsArray();
+	}
+	return false;
+}
+
+const TArray<FRigVMTemplateArgument::FType>& FRigVMTemplateArgument::GetTypes() const
+{
+	return Types;
+}
+
+TArray<FRigVMTemplateArgument::FType> FRigVMTemplateArgument::GetSupportedTypes(const TArray<int32>& InPermutationIndices) const
+{
+	TArray<FType> SupportedTypes;
+	if(InPermutationIndices.IsEmpty())
+	{
+		for (const FType& Type : Types)
+		{
+			SupportedTypes.AddUnique(Type);
+		}
+	}
+	else
+	{
+		for(const int32 PermutationIndex : InPermutationIndices)
+		{
+			SupportedTypes.AddUnique(Types[PermutationIndex]);
+		}
+	}
+	return SupportedTypes;
+}
+
+TArray<FString> FRigVMTemplateArgument::GetSupportedTypeStrings(const TArray<int32>& InPermutationIndices) const
+{
+	TArray<FString> SupportedTypes;
+	if(InPermutationIndices.IsEmpty())
+	{
+		for (const FType& Type : Types)
+		{
+			SupportedTypes.AddUnique(Type.CPPType);
+		}
+	}
+	else
+	{
+		for(const int32 PermutationIndex : InPermutationIndices)
+		{
+			SupportedTypes.AddUnique(Types[PermutationIndex].CPPType);
+		}
+	}
+	return SupportedTypes;
+}
+
+////////////////////////////////////////////////////////////////////////////////////////////////
+
+FRigVMTemplate::FRigVMTemplate()
+	: Index(INDEX_NONE)
+	, Notation(NAME_None)
+{
+
+}
+
+FRigVMTemplate::FRigVMTemplate(UScriptStruct* InStruct, const FString& InTemplateName, int32 InFunctionIndex)
+	: Index(INDEX_NONE)
+	, Notation(NAME_None)
+{
+	TArray<FString> ArgumentNotations;
+	for (TFieldIterator<FProperty> It(InStruct); It; ++It)
+	{
+		FRigVMTemplateArgument Argument(*It);
+		Argument.Index = Arguments.Num();
+
+		if(IsValidArgumentForTemplate(Argument))
+		{
+			Arguments.Add(Argument);
+			ArgumentNotations.Add(GetArgumentNotation(Argument));
+		}
+	}
+
+	if (ArgumentNotations.Num() > 0)
+	{
+		FString NotationStr = FString::Printf(TEXT("%s(%s)"), *InTemplateName, *FString::Join(ArgumentNotations, TEXT(",")));
+		Notation = *NotationStr;
+		Permutations.Add(InFunctionIndex);
+	}
+}
+
+bool FRigVMTemplate::IsValidArgumentForTemplate(const FRigVMTemplateArgument& InArgument)
+{
+	static const TArray<ERigVMPinDirection> ValidDirections = {
+		ERigVMPinDirection::Input,
+		ERigVMPinDirection::Output,
+		ERigVMPinDirection::IO,
+		ERigVMPinDirection::Visible
+	};
+
+	if(!ValidDirections.Contains(InArgument.Direction))
+	{
+		return false;
+	}
+	return true;
+}
+
+
+const FString& FRigVMTemplate::GetArgumentNotationPrefix(const FRigVMTemplateArgument& InArgument)
+{
+	static const FString EmptyPrefix = FString();
+	static const FString InPrefix = TEXT("in ");
+	static const FString OutPrefix = TEXT("out ");
+	static const FString IOPrefix = TEXT("io ");
+
+	switch(InArgument.Direction)
+	{
+		case ERigVMPinDirection::Input:
+		case ERigVMPinDirection::Visible:
+		{
+			return InPrefix;
+		}
+		case ERigVMPinDirection::Output:
+		{
+			return OutPrefix;
+		}
+		case ERigVMPinDirection::IO:
+		{
+			return IOPrefix;
+		}
+		default:
+		{
+			break;
+		}
+	}
+
+	return EmptyPrefix;
+}
+
+const FString& FRigVMTemplate::GetArgumentNotationSuffix(const FRigVMTemplateArgument& InArgument)
+{
+	static const FString EmptySuffix = FString();
+	static const FString ArraySuffix = TEXT("[]");
+	return InArgument.IsArray() ? ArraySuffix : EmptySuffix;
+}
+
+FString FRigVMTemplate::GetArgumentNotation(const FRigVMTemplateArgument& InArgument)
+{
+	return FString::Printf(TEXT("%s%s%s"),
+		*GetArgumentNotationPrefix(InArgument),
+		*InArgument.GetName().ToString(),
+		*GetArgumentNotationSuffix(InArgument));
+}
+
+bool FRigVMTemplate::IsValid() const
+{
+	return !Notation.IsNone();
+}
+
+const FName& FRigVMTemplate::GetNotation() const
+{
+	return Notation;
+}
+
+FName FRigVMTemplate::GetName() const
+{
+	FString Left;
+	if (GetNotation().ToString().Split(TEXT("::"), &Left, nullptr))
+	{
+		return *Left;
+	}
+	return NAME_None;
+}
+
+#if WITH_EDITOR
+
+FLinearColor FRigVMTemplate::GetColor(const TArray<int32>& InPermutationIndices) const
+{
+	bool bFirstColorFound = false;
+	FLinearColor ResolvedColor = FLinearColor::White;
+
+	auto GetColorFromMetadata = [] (FString InMetadata) -> FLinearColor
+	{
+		FLinearColor Color = FLinearColor::Black;
+
+		FString Metadata = InMetadata;
+		Metadata.TrimStartAndEndInline();
+		FString SplitString(TEXT(" "));
+		FString Red, Green, Blue, GreenAndBlue;
+		if (Metadata.Split(SplitString, &Red, &GreenAndBlue))
+		{
+			Red.TrimEndInline();
+			GreenAndBlue.TrimStartInline();
+			if (GreenAndBlue.Split(SplitString, &Green, &Blue))
+			{
+				Green.TrimEndInline();
+				Blue.TrimStartInline();
+
+				float RedValue = FCString::Atof(*Red);
+				float GreenValue = FCString::Atof(*Green);
+				float BlueValue = FCString::Atof(*Blue);
+				Color = FLinearColor(RedValue, GreenValue, BlueValue);
+			}
+		}
+
+		return Color;
+	};
+
+	auto VisitPermutation = [&bFirstColorFound, &ResolvedColor, GetColorFromMetadata, this](int32 InPermutationIndex) -> bool
+	{
+		static const FName NodeColorName = TEXT("NodeColor");
+		FString NodeColorMetadata;
+
+		const FRigVMFunction* ResolvedFunction = GetPermutation(InPermutationIndex);
+		ResolvedFunction->Struct->GetStringMetaDataHierarchical(NodeColorName, &NodeColorMetadata);
+		if (!NodeColorMetadata.IsEmpty())
+		{
+			if(bFirstColorFound)
+			{
+				const FLinearColor NodeColor = GetColorFromMetadata(NodeColorMetadata);
+				if(!ResolvedColor.Equals(NodeColor, 0.01f))
+				{
+					ResolvedColor = FLinearColor::White;
+					return false;
+				}
+			}
+			else
+			{
+				ResolvedColor = GetColorFromMetadata(NodeColorMetadata);
+				bFirstColorFound = true;
+			}
+		}
+		return true;
+	};
+
+	if(InPermutationIndices.IsEmpty())
+	{
+		for(int32 PermutationIndex = 0; PermutationIndex < Permutations.Num(); PermutationIndex++)
+		{
+			if(!VisitPermutation(PermutationIndex))
+			{
+				break;
+			}
+		}
+	}
+	else
+	{
+		for(const int32 PermutationIndex : InPermutationIndices)
+		{
+			if(!VisitPermutation(PermutationIndex))
+			{
+				break;
+			}
+		}
+	}
+	return ResolvedColor;
+}
+
+FText FRigVMTemplate::GetTooltipText(const TArray<int32>& InPermutationIndices) const
+{
+	FText ResolvedTooltipText;
+
+	auto VisitPermutation = [&ResolvedTooltipText, this](int32 InPermutationIndex) -> bool
+	{
+		const FRigVMFunction* ResolvedFunction = GetPermutation(InPermutationIndex);
+		const FText TooltipText = ResolvedFunction->Struct->GetToolTipText();
+		
+		if (!ResolvedTooltipText.IsEmpty())
+		{
+			if(!ResolvedTooltipText.EqualTo(TooltipText))
+			{
+				ResolvedTooltipText = FText::FromName(GetName());
+				return false;
+			}
+		}
+		else
+		{
+			ResolvedTooltipText = TooltipText;
+		}
+		return true;
+	};
+
+	if(InPermutationIndices.IsEmpty())
+	{
+		for(int32 PermutationIndex = 0; PermutationIndex < Permutations.Num(); PermutationIndex++)
+		{
+			if(!VisitPermutation(PermutationIndex))
+			{
+				break;
+			}
+		}
+	}
+	else
+	{
+		for(const int32 PermutationIndex : InPermutationIndices)
+		{
+			if(!VisitPermutation(PermutationIndex))
+			{
+				break;
+			}
+		}
+	}
+
+	return ResolvedTooltipText;
+}
+
+#endif
+
+bool FRigVMTemplate::IsCompatible(const FRigVMTemplate& InOther) const
+{
+	if (!IsValid() || !InOther.IsValid())
+	{
+		return false;
+	}
+
+	return Notation == InOther.Notation;
+}
+
+bool FRigVMTemplate::Merge(const FRigVMTemplate& InOther)
+{
+	if (!IsCompatible(InOther))
+	{
+		return false;
+	}
+
+	if (InOther.Permutations.Num() != 1)
+	{
+		return false;
+	}
+
+	// find colliding permutations
+	for(int32 PermutationIndex = 0;PermutationIndex<NumPermutations();PermutationIndex++)
+	{
+		int32 MatchingArguments = 0;
+		for(int32 ArgumentIndex = 0; ArgumentIndex < Arguments.Num(); ArgumentIndex++)
+		{
+			if(Arguments[ArgumentIndex].GetTypes()[PermutationIndex] ==
+				InOther.Arguments[ArgumentIndex].Types[0])
+			{
+				MatchingArguments++;
+			}
+		}
+		if(MatchingArguments == Arguments.Num())
+		{
+			// find the previously defined permutation.
+			UE_LOG(LogRigVM, Warning, TEXT("RigVMFunction '%s' cannot be merged into the '%s' template. It collides with '%s'."),
+				InOther.GetPermutation(0)->Name,
+				*GetNotation().ToString(),
+				GetPermutation(PermutationIndex)->Name);
+			return false;
+		}
+	}
+
+	TArray<FRigVMTemplateArgument> NewArgs;
+
+	for (int32 ArgumentIndex = 0; ArgumentIndex < Arguments.Num(); ArgumentIndex++)
+	{
+		if (InOther.Arguments[ArgumentIndex].Types.Num() != 1)
+		{
+			return false;
+		}
+
+		NewArgs.Add(Arguments[ArgumentIndex]);
+		NewArgs[ArgumentIndex].Types.Add(InOther.Arguments[ArgumentIndex].Types[0]);
+
+		if (!Arguments[ArgumentIndex].Types.Contains(InOther.Arguments[ArgumentIndex].Types[0]))
+		{
+			NewArgs[ArgumentIndex].bSingleton = false;
+		}
+	}
+
+	Arguments = NewArgs;
+
+	Permutations.Add(InOther.Permutations[0]);
+	return true;
+}
+
+const FRigVMTemplateArgument* FRigVMTemplate::FindArgument(const FName& InArgumentName) const
+{
+	return Arguments.FindByPredicate([InArgumentName](const FRigVMTemplateArgument& Argument) -> bool
+	{
+		return Argument.GetName() == InArgumentName;
+	});
+}
+
+bool FRigVMTemplate::ArgumentSupportsType(const FName& InArgumentName, const FString& InCPPType, const FRigVMTemplate::FTypeMap& InTypes, FRigVMTemplateArgument::FType* OutType) const
+{
+	if (const FRigVMTemplateArgument* Argument = FindArgument(InArgumentName))
+	{
+		if(InTypes.Num() == 0)
+		{
+			return Argument->SupportsType(InCPPType, TArray<int32>(), OutType);
+		}
+
+		FTypeMap Types = InTypes;
+		Types.FindOrAdd(InArgumentName) = InCPPType;
+
+		TArray<int32> PermutationIndices;
+		if(Resolve(Types, PermutationIndices, true))
+		{
+			if(OutType)
+			{
+				*OutType = Types.FindChecked(InArgumentName);
+			}
+			return true;
+		}
+		return false;
+	}
+	return false;
+}
+
+const FRigVMFunction* FRigVMTemplate::GetPermutation(int32 InIndex) const
+{
+	const FRigVMRegistry& Registry = FRigVMRegistry::Get();
+	return &Registry.GetFunctions()[Permutations[InIndex]];
+}
+
+bool FRigVMTemplate::ContainsPermutation(const FRigVMFunction* InPermutation) const
+{
+	return FindPermutation(InPermutation) != INDEX_NONE;
+}
+
+int32 FRigVMTemplate::FindPermutation(const FRigVMFunction* InPermutation) const
+{
+	check(InPermutation);
+	return Permutations.Find(InPermutation->Index);
+}
+
+bool FRigVMTemplate::FullyResolve(FRigVMTemplate::FTypeMap& InOutTypes, int32& OutPermutationIndex) const
+{
+	TArray<int32> PermutationIndices;
+	Resolve(InOutTypes, PermutationIndices, false);
+	if(PermutationIndices.Num() == 1)
+	{
+		OutPermutationIndex = PermutationIndices[0];
+	}
+	else
+	{
+		OutPermutationIndex = INDEX_NONE;
+	}
+	return OutPermutationIndex != INDEX_NONE;
+}
+
+bool FRigVMTemplate::Resolve(FTypeMap& InOutTypes, TArray<int32>& OutPermutationIndices, bool bAllowFloatingPointCasts) const
+{
+	FTypeMap InputTypes = InOutTypes;
+	InOutTypes.Reset();
+
+	OutPermutationIndices.Reset();
+	for (int32 PermutationIndex = 0; PermutationIndex < Permutations.Num(); PermutationIndex++)
+	{
+		OutPermutationIndices.Add(PermutationIndex);
+	}
+	
+	for (const FRigVMTemplateArgument& Argument : Arguments)
+	{
+		if (Argument.bSingleton)
+		{
+			InOutTypes.Add(Argument.Name, Argument.Types[0]);
+			continue;
+		}
+		else if (const FRigVMTemplateArgument::FType* InputType = InputTypes.Find(Argument.Name))
+		{
+			FRigVMTemplateArgument::FType MatchedType = *InputType;
+			bool bFoundMatch = false;
+			bool bFoundPerfectMatch = false;
+			for (int32 PermutationIndex = 0; PermutationIndex < Argument.Types.Num(); PermutationIndex++)
+			{
+				if (!Argument.Types[PermutationIndex].Matches(InputType->CPPType, bAllowFloatingPointCasts))
+				{
+					OutPermutationIndices.Remove(PermutationIndex);
+				}
+				else
+				{
+					bFoundMatch = true;
+
+					// if the type matches - but it's not the exact same
+					if(!bFoundPerfectMatch)
+					{
+						MatchedType = Argument.Types[PermutationIndex];
+
+						// if we found the perfect match - let's stop here
+						if(Argument.Types[PermutationIndex] == *InputType)
+						{
+							bFoundPerfectMatch = true;
+						}
+					}
+				}
+			}
+			
+			if(bFoundMatch)
+			{
+				InOutTypes.Add(Argument.Name,MatchedType);
+				continue;
+			}
+		}
+		
+		if(Argument.IsArray())
+		{
+			InOutTypes.Add(Argument.Name, FRigVMTemplateArgument::FType::Array());
+		}
+		else
+		{
+			InOutTypes.Add(Argument.Name, FRigVMTemplateArgument::FType());
+		}
+	}
+
+	if (OutPermutationIndices.Num() == 1)
+	{
+		InOutTypes.Reset();
+		for (const FRigVMTemplateArgument& Argument : Arguments)
+		{
+			InOutTypes.Add(Argument.Name, Argument.Types[OutPermutationIndices[0]]);
+		}
+	}
+	else if (OutPermutationIndices.Num() > 1)
+	{
+		for (const FRigVMTemplateArgument& Argument : Arguments)
+		{
+			if (Argument.IsSingleton(OutPermutationIndices))
+			{
+				InOutTypes.FindChecked(Argument.Name) = Argument.Types[OutPermutationIndices[0]];
+			}
+		}
+	}
+
+	return !OutPermutationIndices.IsEmpty();
+}
+
+#if WITH_EDITOR
+
+FString FRigVMTemplate::GetCategory() const
+{
+	FString Category;
+	GetPermutation(0)->Struct->GetStringMetaDataHierarchical(FRigVMStruct::CategoryMetaName, &Category);
+
+	if (Category.IsEmpty())
+	{
+		return Category;
+	}
+
+	for (int32 PermutationIndex = 1; PermutationIndex < NumPermutations(); PermutationIndex++)
+	{
+		FString OtherCategory;
+		if (GetPermutation(PermutationIndex)->Struct->GetStringMetaDataHierarchical(FRigVMStruct::CategoryMetaName, &OtherCategory))
+		{
+			while (!OtherCategory.StartsWith(Category, ESearchCase::IgnoreCase))
+			{
+				FString Left;
+				if (Category.Split(TEXT("|"), &Left, nullptr, ESearchCase::IgnoreCase, ESearchDir::FromEnd))
+				{
+					Category = Left;
+				}
+				else
+				{
+					return FString();
+				}
+
+			}
+		}
+	}
+
+	return Category;
+}
+
+FString FRigVMTemplate::GetKeywords() const
+{
+	TArray<FString> KeywordsMetadata;
+	KeywordsMetadata.Add(GetName().ToString());
+
+	for (int32 PermutationIndex = 0; PermutationIndex < NumPermutations(); PermutationIndex++)
+	{
+		if (const FRigVMFunction* Function = GetPermutation(PermutationIndex))
+		{
+			KeywordsMetadata.Add(Function->Struct->GetDisplayNameText().ToString());
+			
+			FString FunctionKeyWordsMetadata;
+			Function->Struct->GetStringMetaDataHierarchical(FRigVMStruct::KeywordsMetaName, &FunctionKeyWordsMetadata);
+			if (!FunctionKeyWordsMetadata.IsEmpty())
+			{
+				KeywordsMetadata.Add(FunctionKeyWordsMetadata);
+			}
+		}
+	}
+
+	return FString::Join(KeywordsMetadata, TEXT(","));
+}
+
+#endif
\ No newline at end of file
diff --git a/Engine/Source/Runtime/RigVM/Public/RigVMCore/RigVMFunction.h b/Engine/Source/Runtime/RigVM/Public/RigVMCore/RigVMFunction.h
--- a/Engine/Source/Runtime/RigVM/Public/RigVMCore/RigVMFunction.h
+++ b/Engine/Source/Runtime/RigVM/Public/RigVMCore/RigVMFunction.h
@@ -12,6 +12,8 @@ typedef TArrayView<void*> FRigVMUserDataArray;
 
 typedef void (*FRigVMFunctionPtr)(FRigVMExecuteContext& RigVMExecuteContext, FRigVMMemoryHandleArray RigVMMemoryHandles);
 
+struct FRigVMTemplate;
+
 /**
  * The Pin Direction is used to differentiate different kinds of 
  * pins in the data flow graph - inputs, outputs etc.
@@ -37,14 +39,14 @@ struct RIGVM_API FRigVMFunction
 	UScriptStruct* Struct;
 	FRigVMFunctionPtr FunctionPtr;
 	int32 Index;
-	int32 PrototypeIndex;
+	int32 TemplateIndex;
 
 	FRigVMFunction()
 		: Name(nullptr)
 		, Struct(nullptr)
 		, FunctionPtr(nullptr)
 		, Index(INDEX_NONE)
-		, PrototypeIndex(INDEX_NONE)
+		, TemplateIndex(INDEX_NONE)
 	{
 	}
 
@@ -53,7 +55,7 @@ struct RIGVM_API FRigVMFunction
 		, Struct(InStruct)
 		, FunctionPtr(InFunctionPtr)
 		, Index(InIndex)
-		, PrototypeIndex(INDEX_NONE)
+		, TemplateIndex(INDEX_NONE)
 	{
 	}
 
@@ -61,4 +63,5 @@ struct RIGVM_API FRigVMFunction
 	FName GetMethodName() const;
 	FString GetModuleName() const;
 	FString GetModuleRelativeHeaderPath() const;
+	const FRigVMTemplate* GetTemplate() const;
 };
diff --git a/Engine/Source/Runtime/RigVM/Public/RigVMCore/RigVMPrototype.h b/Engine/Source/Runtime/RigVM/Public/RigVMCore/RigVMPrototype.h
deleted file mode 100644
diff --git a/Engine/Source/Runtime/RigVM/Public/RigVMCore/RigVMRegistry.h b/Engine/Source/Runtime/RigVM/Public/RigVMCore/RigVMRegistry.h
--- a/Engine/Source/Runtime/RigVM/Public/RigVMCore/RigVMRegistry.h
+++ b/Engine/Source/Runtime/RigVM/Public/RigVMCore/RigVMRegistry.h
@@ -5,7 +5,7 @@
 #include "CoreMinimal.h"
 #include "RigVMMemory.h"
 #include "RigVMFunction.h"
-#include "RigVMPrototype.h"
+#include "RigVMTemplate.h"
 
 /**
  * The FRigVMRegistry is used to manage all known function pointers
@@ -29,27 +29,24 @@ public:
 	// based on the names.
 	void Refresh();
 
-	// Returns a function pointer given its name
-	FRigVMFunctionPtr FindFunction(const TCHAR* InName) const;
+	// Returns the function given its name (or nullptr)
+	const FRigVMFunction* FindFunction(const TCHAR* InName) const;
 
-	// Returns a prototype pointer given its notation (or nullptr)
-	const FRigVMPrototype* FindPrototype(const FName& InNotation) const;
-
-	// Returns a prototype pointer given its notation (or nullptr)
-	const FRigVMPrototype* FindPrototype(UScriptStruct* InStruct, const FString& InPrototypeName) const;
-
-	// Returns registry info about a function given its name
-	FRigVMFunction FindFunctionInfo(const TCHAR* InName) const;
+	// Returns the function given its backing up struct and method name
+	const FRigVMFunction* FindFunction(UScriptStruct* InStruct, const TCHAR* InName) const;
 
 	// Returns all current rigvm functions
 	const TArray<FRigVMFunction>& GetFunctions() const;
 
+	// Returns a template pointer given its notation (or nullptr)
+	const FRigVMTemplate* FindTemplate(const FName& InNotation) const;
+
 	// Returns all current rigvm functions
-	const TArray<FRigVMPrototype>& GetPrototypes() const;
+	const TArray<FRigVMTemplate>& GetTemplates() const;
 
 private:
 
-	static const FName PrototypeNameMetaName;
+	static const FName TemplateNameMetaName;
 
 	// disable default constructor
 	FRigVMRegistry() {}
@@ -61,8 +58,15 @@ private:
 	// memory for all functions
 	TArray<FRigVMFunction> Functions;
 
-	// memory for all prototypes
-	TArray<FRigVMPrototype> Prototypes;
+	// memory for all templates
+	TArray<FRigVMTemplate> Templates;
+
+	// name lookup for functions
+	TMap<FName, int32> FunctionNameToIndex;
+
+	// name lookup for templates
+	TMap<FName, int32> TemplateNotationToIndex;
 
 	static FRigVMRegistry s_RigVMRegistry;
+	friend struct FRigVMStruct;
 };
diff --git a/Engine/Source/Runtime/RigVM/Public/RigVMCore/RigVMStruct.h b/Engine/Source/Runtime/RigVM/Public/RigVMCore/RigVMStruct.h
--- a/Engine/Source/Runtime/RigVM/Public/RigVMCore/RigVMStruct.h
+++ b/Engine/Source/Runtime/RigVM/Public/RigVMCore/RigVMStruct.h
@@ -188,7 +188,7 @@ public:
 	// Example: Icon="EditorStyle|GraphEditor.Sequence_16x"
 	static const FName IconMetaName;
 	static const FName KeywordsMetaName;
-	static const FName PrototypeNameMetaName;
+	static const FName TemplateNameMetaName;
 	static const FName ExpandPinByDefaultMetaName;
 	static const FName DefaultArraySizeMetaName;
 	static const FName VaryingMetaName;
diff --git a/Engine/Source/Runtime/RigVM/Public/RigVMCore/RigVMTemplate.h b/Engine/Source/Runtime/RigVM/Public/RigVMCore/RigVMTemplate.h
new file mode 100644
--- /dev/null
+++ b/Engine/Source/Runtime/RigVM/Public/RigVMCore/RigVMTemplate.h
@@ -0,0 +1,238 @@
+// Copyright Epic Games, Inc. All Rights Reserved.
+
+#pragma once
+
+#include "CoreMinimal.h"
+#include "RigVMFunction.h"
+#include "UObject/Object.h"
+#include "RigVMTypeUtils.h"
+
+/**
+ * The template argument represents a single parameter
+ * in a function call and all of its possible types
+ */
+struct RIGVM_API FRigVMTemplateArgument
+{
+	struct FType
+	{
+		FString CPPType;
+		UObject* CPPTypeObject;
+
+		FType()
+			: CPPType()
+			, CPPTypeObject(nullptr)
+		{
+			CPPType = RigVMTypeUtils::GetWildCardCPPType();
+			CPPTypeObject = RigVMTypeUtils::GetWildCardCPPTypeObject();
+		}
+
+		FType(const FString& InCPPType, UObject* InCPPTypeObject = nullptr)
+			: CPPType(InCPPType)
+			, CPPTypeObject(InCPPTypeObject)
+		{
+			check(!CPPType.IsEmpty());
+		}
+
+		static FType Array()
+		{
+			return FType(RigVMTypeUtils::GetWildCardArrayCPPType(), RigVMTypeUtils::GetWildCardCPPTypeObject());
+		}
+
+		bool operator == (const FType& InOther) const
+		{
+			return CPPType == InOther.CPPType;
+		}
+
+		bool operator != (const FType& InOther) const
+		{
+			return CPPType != InOther.CPPType;
+		}
+
+		FORCEINLINE bool Matches(const FString& InCPPType, bool bAllowFloatingPointCasts = true) const
+		{
+			if(CPPType == InCPPType)
+			{
+				return true;
+			}
+			if(bAllowFloatingPointCasts)
+			{
+				if(InCPPType == RigVMTypeUtils::FloatType && CPPType == RigVMTypeUtils::DoubleType)
+				{
+					return true;
+				}
+				if(InCPPType == RigVMTypeUtils::DoubleType && CPPType == RigVMTypeUtils::FloatType)
+				{
+					return true;
+				}
+				if(InCPPType == RigVMTypeUtils::FloatArrayType && CPPType == RigVMTypeUtils::DoubleArrayType)
+				{
+					return true;
+				}
+				if(InCPPType == RigVMTypeUtils::DoubleArrayType && CPPType == RigVMTypeUtils::FloatArrayType)
+				{
+					return true;
+				}
+			}
+			return false;
+		}
+
+		FName GetCPPTypeObjectPath() const
+		{
+			if(CPPTypeObject)
+			{
+				return *CPPTypeObject->GetPathName();
+			}
+			return NAME_None;
+		}
+
+		bool IsWildCard() const
+		{
+			return CPPTypeObject == RigVMTypeUtils::GetWildCardCPPTypeObject();
+		}
+
+		bool IsArray() const
+		{
+			return RigVMTypeUtils::IsArrayType(CPPType);
+		}
+	};
+
+	// returns the name of the argument
+	const FName& GetName() const { return Name; }
+
+	// returns the direction of the argument
+	ERigVMPinDirection GetDirection() const { return Direction; }
+
+	// returns true if this argument supports a given type across a set of permutations
+	bool SupportsType(const FString& InCPPType, const TArray<int32>& InPermutationIndices = TArray<int32>(), FType* OutType = nullptr) const;
+
+	// returns the flat list of types (including duplicates) of this argument
+	const TArray<FType>& GetTypes() const;
+
+	// returns an array of all of the supported types
+	TArray<FType> GetSupportedTypes(const TArray<int32>& InPermutationIndices = TArray<int32>()) const;
+
+	// returns an array of all supported types as strings
+	TArray<FString> GetSupportedTypeStrings(const TArray<int32>& InPermutationIndices = TArray<int32>()) const;
+
+	// returns true if an argument is singleton (same type for all variants)
+	bool IsSingleton(const TArray<int32>& InPermutationIndices = TArray<int32>()) const;
+
+	// returns true if the argument uses an array container
+	bool IsArray() const;
+
+protected:
+
+	int32 Index;
+	FName Name;
+	ERigVMPinDirection Direction;
+	bool bSingleton;
+	TArray<FType> Types;
+
+	// default constructor
+	FRigVMTemplateArgument();
+
+	// constructor from a property
+	FRigVMTemplateArgument(FProperty* InProperty);
+
+	friend struct FRigVMTemplate;
+	friend class URigVMController;
+};
+
+/**
+ * The template is used to group multiple rigvm functions
+ * that share the same notation. Templates can then be used
+ * to build polymorphic nodes (RigVMTemplateNode) that can
+ * take on any of the permutations supported by the template.
+ */
+struct RIGVM_API FRigVMTemplate
+{
+public:
+
+	typedef TMap<FName, FRigVMTemplateArgument::FType> FTypeMap;
+
+	// returns true if this is a valid template
+	bool IsValid() const;
+
+	// Returns the notation of this template
+	const FName& GetNotation() const;
+
+	// Returns the name of the template
+	FName GetName() const;
+
+	// returns true if this template is compatible with another one
+	bool IsCompatible(const FRigVMTemplate& InOther) const;
+
+	// returns true if this template can merge another one
+	bool Merge(const FRigVMTemplate& InOther);
+
+	// returns the number of args of this template
+	int32 NumArguments() const { return Arguments.Num(); }
+
+	// returns an argument for a given index
+	const FRigVMTemplateArgument* GetArgument(int32 InIndex) const { return &Arguments[InIndex]; }
+
+		// returns an argument given a name (or nullptr)
+	const FRigVMTemplateArgument* FindArgument(const FName& InArgumentName) const;
+
+	// returns true if a given arg supports a type
+	bool ArgumentSupportsType(const FName& InArgumentName, const FString& InCPPType, const FTypeMap& InTypes = FTypeMap(), FRigVMTemplateArgument::FType* OutType = nullptr) const;
+
+	// returns the number of permutations supported by this template
+	int32 NumPermutations() const { return Permutations.Num(); }
+
+	// returns a permutation given an index
+	const FRigVMFunction* GetPermutation(int32 InIndex) const;
+
+	// returns true if a given function is a permutation of this template
+	bool ContainsPermutation(const FRigVMFunction* InPermutation) const;
+
+	// returns the index of the permutation within the template of a given function (or INDEX_NONE)
+	int32 FindPermutation(const FRigVMFunction* InPermutation) const;
+
+	// returns true if the template was able to resolve to single permutation
+	bool FullyResolve(FTypeMap& InOutTypes, int32& OutPermutationIndex) const;
+
+	// returns true if the template was able to resolve to at least one permutation
+	bool Resolve(FTypeMap& InOutTypes, TArray<int32> & OutPermutationIndices, bool bAllowFloatingPointCasts) const;
+
+	// returns true if a given argument is valid for a template
+	static bool IsValidArgumentForTemplate(const FRigVMTemplateArgument& InArgument);
+
+	// returns the prefix for an argument in the notation
+	static const FString& GetArgumentNotationPrefix(const FRigVMTemplateArgument& InArgument);
+
+	// returns the suffix for an argument in the notation
+	static const FString& GetArgumentNotationSuffix(const FRigVMTemplateArgument& InArgument);
+
+	// returns the notation of an argument
+	static FString GetArgumentNotation(const FRigVMTemplateArgument& InArgument);
+
+#if WITH_EDITOR
+
+	// Returns the color based on the permutation's metadata
+	FLinearColor GetColor(const TArray<int32>& InPermutationIndices = TArray<int32>()) const;
+
+	// Returns the tooltip based on the permutation's metadata
+	FText GetTooltipText(const TArray<int32>& InPermutationIndices = TArray<int32>()) const;
+
+	FString GetCategory() const;
+	FString GetKeywords() const;
+
+#endif
+
+private:
+
+	// Default constructor
+	FRigVMTemplate();
+
+	// Constructor from a struct, a template name and a function index
+	FRigVMTemplate(UScriptStruct* InStruct, const FString& InTemplateName, int32 InFunctionIndex);
+
+	int32 Index;
+	FName Notation;
+	TArray<FRigVMTemplateArgument> Arguments;
+	TArray<int32> Permutations;
+
+	friend struct FRigVMRegistry;
+};
+
diff --git a/Engine/Source/Runtime/RigVM/Public/RigVMCore/RigVMUnknownType.h b/Engine/Source/Runtime/RigVM/Public/RigVMCore/RigVMUnknownType.h
--- a/Engine/Source/Runtime/RigVM/Public/RigVMCore/RigVMUnknownType.h
+++ b/Engine/Source/Runtime/RigVM/Public/RigVMCore/RigVMUnknownType.h
@@ -9,7 +9,7 @@
 /**
  * The unknown type is used to identify untyped nodes
  */
-USTRUCT()
+USTRUCT(meta=(DisplayName="Wildcard"))
 struct RIGVM_API FRigVMUnknownType
 {
 	GENERATED_BODY()
diff --git a/Engine/Source/Runtime/RigVM/Public/RigVMTypeUtils.h b/Engine/Source/Runtime/RigVM/Public/RigVMTypeUtils.h
--- a/Engine/Source/Runtime/RigVM/Public/RigVMTypeUtils.h
+++ b/Engine/Source/Runtime/RigVM/Public/RigVMTypeUtils.h
@@ -3,6 +3,7 @@
 #pragma once
 
 #include "CoreMinimal.h"
+#include "RigVMCore/RigVMUnknownType.h"
 
 namespace RigVMTypeUtils
 {
@@ -10,7 +11,20 @@ namespace RigVMTypeUtils
 	const TCHAR TObjectPtrPrefix[] = TEXT("TObjectPtr<");
 	const TCHAR TArrayTemplate[] = TEXT("TArray<%s>");
 	const TCHAR TObjectPtrTemplate[] = TEXT("TObjectPtr<%s%s>");
-	
+
+	const FString BoolType = TEXT("bool");
+	const FString FloatType = TEXT("float");
+	const FString DoubleType = TEXT("double");
+	const FString Int32Type = TEXT("int32");
+	const FString FNameType = TEXT("FName");
+	const FString FStringType = TEXT("FString");
+	const FString BoolArrayType = TEXT("TArray<bool>");
+	const FString FloatArrayType = TEXT("TArray<float>");
+	const FString DoubleArrayType = TEXT("TArray<double>");
+	const FString Int32ArrayType = TEXT("TArray<int32>");
+	const FString FNameArrayType = TEXT("TArray<FName>");
+	const FString FStringArrayType = TEXT("TArray<FString>");
+
 	// Returns true if the type specified is an array
 	FORCEINLINE bool IsArrayType(const FString& InCPPType)
 	{
@@ -27,8 +41,38 @@ namespace RigVMTypeUtils
 		return InCPPType.RightChop(7).LeftChop(1);
 	}
 
+	FORCEINLINE FString CPPTypeFromEnum(UEnum* InEnum)
+	{
+		check(InEnum);
+
+		FString CPPType = InEnum->CppType;
+		if(CPPType.IsEmpty()) // this might be a user defined enum
+		{
+			CPPType = InEnum->GetName();
+		}
+		return CPPType;
+	}
+
 	FORCEINLINE bool IsUObjectType(const FString& InCPPType)
 	{
 		return InCPPType.StartsWith(TObjectPtrPrefix);
 	}
+
+	static UScriptStruct* GetWildCardCPPTypeObject()
+	{
+		static UScriptStruct* WildCardTypeObject = FRigVMUnknownType::StaticStruct();
+		return WildCardTypeObject;
+	}
+
+	static const FString GetWildCardCPPType()
+	{
+		static const FString WildCardCPPType = FRigVMUnknownType::StaticStruct()->GetStructCPPName(); 
+		return WildCardCPPType;
+	}
+
+	static const FString GetWildCardArrayCPPType()
+	{
+		static const FString WildCardArrayCPPType = ArrayTypeFromBaseType(GetWildCardCPPType()); 
+		return WildCardArrayCPPType;
+	}
 }
\ No newline at end of file
