Control Rig: Remove parameters

diff --git a/Engine/Plugins/Animation/ControlRig/Source/ControlRig/Public/ControlRigObjectVersion.h b/Engine/Plugins/Animation/ControlRig/Source/ControlRig/Public/ControlRigObjectVersion.h
--- a/Engine/Plugins/Animation/ControlRig/Source/ControlRig/Public/ControlRigObjectVersion.h
+++ b/Engine/Plugins/Animation/ControlRig/Source/ControlRig/Public/ControlRigObjectVersion.h
@@ -54,6 +54,9 @@ struct CONTROLRIG_API FControlRigObjectVersion
 		// Removed the parent cache for multi parent elements 
 		RemovedMultiParentParentCache,
 
+		// Deprecation of parameters
+		RemoveParameters,
+
 		// -----<new versions can be added above this line>-------------------------------------------------
 		VersionPlusOne,
 		LatestVersion = VersionPlusOne - 1
diff --git a/Engine/Plugins/Animation/ControlRig/Source/ControlRigDeveloper/Private/ControlRigBlueprint.cpp b/Engine/Plugins/Animation/ControlRig/Source/ControlRigDeveloper/Private/ControlRigBlueprint.cpp
--- a/Engine/Plugins/Animation/ControlRig/Source/ControlRigDeveloper/Private/ControlRigBlueprint.cpp
+++ b/Engine/Plugins/Animation/ControlRig/Source/ControlRigDeveloper/Private/ControlRigBlueprint.cpp
@@ -464,6 +464,7 @@ void UControlRigBlueprint::PostLoad()
 #if UE_RIGVM_ENABLE_TEMPLATE_NODES
 		ConvertUnitNodesToTemplateNodes();
 #endif
+		PatchParameterNodesOnLoad();
 
 #if WITH_EDITOR
 
@@ -2722,7 +2723,7 @@ void UControlRigBlueprint::PopulateModelFromGraphForBackwardsCompatibility(UCont
 							bIsInput = false;
 						}
 
-						ModelNode = GetOrCreateController()->AddParameterNode(PropertyName, DataType.ToString(), DataTypeObject, bIsInput, FString(), NodePosition, PropertyName.ToString(), false);
+						ModelNode = GetOrCreateController()->AddVariableNode(PropertyName, DataType.ToString(), DataTypeObject, bIsInput, FString(), NodePosition, PropertyName.ToString(), false);
 					}
 				}
 				else
@@ -3201,9 +3202,6 @@ void UControlRigBlueprint::HandleModifiedEvent(ERigVMGraphNotifType InNotifType,
 				}
 				break;
 			}
-			case ERigVMGraphNotifType::ParameterAdded:
-			case ERigVMGraphNotifType::ParameterRemoved:
-			case ERigVMGraphNotifType::ParameterRenamed:
 			case ERigVMGraphNotifType::PinBoundVariableChanged:
 			case ERigVMGraphNotifType::VariableRemappingChanged:
 			{
@@ -3368,6 +3366,8 @@ void UControlRigBlueprint::CreateMemberVariablesOnLoad()
 				}
 			}
 
+			// Leaving this for backwards compatibility, even though we don't support parameters anymore
+			// When a parameter node is found, we will create a variable
 			if (URigVMParameterNode* ParameterNode = Cast<URigVMParameterNode>(Node))
 			{
 				if (URigVMPin* ParameterPin = ParameterNode->FindPin(TEXT("Parameter")))
@@ -3753,6 +3753,82 @@ void UControlRigBlueprint::ConvertUnitNodesToTemplateNodes()
 	}
 }
 
+void UControlRigBlueprint::PatchParameterNodesOnLoad()
+{	
+#if WITH_EDITOR
+
+	// setup variables on the blueprint based on the previous "parameters"
+	if (GetLinkerCustomVersion(FControlRigObjectVersion::GUID) < FControlRigObjectVersion::RemoveParameters)
+	{
+		
+		TGuardValue<bool> GuardNotifsSelf(bSuspendModelNotificationsForSelf, true);
+
+		GetOrCreateController()->ReattachLinksToPinObjects();
+
+		check(Model);
+
+		TArray<URigVMNode*> Nodes = Model->GetNodes();
+		for (URigVMNode* Node : Nodes)
+		{
+			if (URigVMParameterNode* ParameterNode = Cast<URigVMParameterNode>(Node))
+			{
+				FRigVMGraphParameterDescription Description = ParameterNode->GetParameterDescription();
+
+				const FEdGraphPinType PinType = RigVMTypeUtils::PinTypeFromCPPType(Description.CPPType, Description.CPPTypeObject);
+
+				int32 VariableIndex = INDEX_NONE;
+				bool bFoundName = false;
+				for (int32 i=0; i<NewVariables.Num(); ++i)
+				{
+					if (NewVariables[i].VarName == Description.Name)
+					{
+						if (NewVariables[i].VarType == PinType &&
+							NewVariables[i].DefaultValue == Description.DefaultValue)
+						{
+							VariableIndex = i;
+						}
+						bFoundName = true;
+						break;
+					}
+				}
+
+				if (VariableIndex == INDEX_NONE)
+				{
+					FName NewVariableName = Description.Name;
+					if (bFoundName)
+					{
+						bool bFound = true;
+						int32 Count = 0;
+						while (bFound)
+						{
+							bFound = false;
+							for (int32 i=0; i<NewVariables.Num(); ++i)
+							{
+								if (NewVariables[i].VarName == NewVariableName)
+								{
+									bFound = true;
+									NewVariableName = *FString::Printf(TEXT("%s_%d"), *Description.Name.ToString(), ++Count);
+									break;
+								}
+							}
+						}
+					}
+					
+					VariableIndex = AddCRMemberVariable(this, NewVariableName, PinType, false, false, Description.DefaultValue);
+				}
+				
+				FName VarName = NewVariables[VariableIndex].VarName;
+				GetOrCreateController()->ReplaceParameterNodeWithVariable(ParameterNode->GetFName(), VarName, Description.CPPType, Description.CPPTypeObject, false);
+				bDirtyDuringLoad = true;
+			}
+		}		
+
+		LastNewVariables = NewVariables;
+	}
+
+#endif	
+}
+
 void UControlRigBlueprint::PropagatePoseFromInstanceToBP(UControlRig* InControlRig)
 {
 	check(InControlRig);
diff --git a/Engine/Plugins/Animation/ControlRig/Source/ControlRigDeveloper/Private/Graph/ControlRigGraph.cpp b/Engine/Plugins/Animation/ControlRig/Source/ControlRigDeveloper/Private/Graph/ControlRigGraph.cpp
--- a/Engine/Plugins/Animation/ControlRig/Source/ControlRigDeveloper/Private/Graph/ControlRigGraph.cpp
+++ b/Engine/Plugins/Animation/ControlRig/Source/ControlRigDeveloper/Private/Graph/ControlRigGraph.cpp
@@ -237,8 +237,6 @@ void UControlRigGraph::HandleModifiedEvent(ERigVMGraphNotifType InNotifType, URi
 			case ERigVMGraphNotifType::PinRenamed:
 			case ERigVMGraphNotifType::VariableRemoved:
 			case ERigVMGraphNotifType::VariableRenamed:
-			case ERigVMGraphNotifType::ParameterRemoved:
-			case ERigVMGraphNotifType::ParameterRenamed:
 			case ERigVMGraphNotifType::GraphChanged:
 			{
 				return;
@@ -594,14 +592,6 @@ void UControlRigGraph::HandleModifiedEvent(ERigVMGraphNotifType InNotifType, URi
 							RigNode->ReconstructNode_Internal(true);
 						}
 					}
-					else if (Cast<URigVMParameterNode>(ModelPin->GetNode()))
-					{
-						if (ModelPin->GetName() == TEXT("Parameter"))
-						{
-							RigNode->InvalidateNodeTitle();
-							RigNode->ReconstructNode_Internal(true);
-						}
-					}
 					else if (Cast<URigVMUnitNode>(ModelPin->GetNode()))
 					{
 						RigNode->InvalidateNodeTitle();
diff --git a/Engine/Plugins/Animation/ControlRig/Source/ControlRigDeveloper/Private/Graph/ControlRigGraphNode.cpp b/Engine/Plugins/Animation/ControlRig/Source/ControlRigDeveloper/Private/Graph/ControlRigGraphNode.cpp
--- a/Engine/Plugins/Animation/ControlRig/Source/ControlRigDeveloper/Private/Graph/ControlRigGraphNode.cpp
+++ b/Engine/Plugins/Animation/ControlRig/Source/ControlRigDeveloper/Private/Graph/ControlRigGraphNode.cpp
@@ -745,7 +745,7 @@ FLinearColor UControlRigGraphNode::GetNodeOpacityColor() const
 {
 	if (URigVMNode* ModelNode = GetModelNode())
 	{
-		if (Cast<URigVMParameterNode>(ModelNode) || Cast<URigVMVariableNode>(ModelNode))
+		if (Cast<URigVMVariableNode>(ModelNode))
 		{
 			return FLinearColor::White;
 		}
diff --git a/Engine/Plugins/Animation/ControlRig/Source/ControlRigDeveloper/Public/ControlRigBlueprint.h b/Engine/Plugins/Animation/ControlRig/Source/ControlRigDeveloper/Public/ControlRigBlueprint.h
--- a/Engine/Plugins/Animation/ControlRig/Source/ControlRigDeveloper/Public/ControlRigBlueprint.h
+++ b/Engine/Plugins/Animation/ControlRig/Source/ControlRigDeveloper/Public/ControlRigBlueprint.h
@@ -552,6 +552,7 @@ private:
 	void PatchBoundVariables();
 	void PatchVariableNodesWithIncorrectType();
 	void ConvertUnitNodesToTemplateNodes();
+	void PatchParameterNodesOnLoad();
 
 	TMap<FName, int32> AddedMemberVariableMap;
 	TArray<FBPVariableDescription> LastNewVariables;
diff --git a/Engine/Plugins/Animation/ControlRig/Source/ControlRigEditor/Private/Graph/ControlRigGraphPanelPinFactory.cpp b/Engine/Plugins/Animation/ControlRig/Source/ControlRigEditor/Private/Graph/ControlRigGraphPanelPinFactory.cpp
--- a/Engine/Plugins/Animation/ControlRig/Source/ControlRigEditor/Private/Graph/ControlRigGraphPanelPinFactory.cpp
+++ b/Engine/Plugins/Animation/ControlRig/Source/ControlRigEditor/Private/Graph/ControlRigGraphPanelPinFactory.cpp
@@ -7,7 +7,6 @@
 #include "Graph/SControlRigGraphPinNameList.h"
 #include "Graph/SControlRigGraphPinCurveFloat.h"
 #include "Graph/SControlRigGraphPinVariableName.h"
-#include "Graph/SControlRigGraphPinParameterName.h"
 #include "Graph/SControlRigGraphPinVariableBinding.h"
 #include "KismetPins/SGraphPinExec.h"
 #include "SGraphPinComboBox.h"
@@ -95,10 +94,6 @@ TSharedPtr<SGraphPin> FControlRigGraphPanelPinFactory::CreatePin(UEdGraphPin* In
 				{
 					return SNew(SControlRigGraphPinVariableName, InPin);
 				}
-				else if (CustomWidgetName == TEXT("ParameterName"))
-				{
-					return SNew(SControlRigGraphPinParameterName, InPin);
-				}
 			}
 
 			if (InPin->PinType.PinCategory == UEdGraphSchema_K2::PC_Struct)
diff --git a/Engine/Plugins/Animation/ControlRig/Source/ControlRigEditor/Private/Graph/SControlRigGraphPinParameterName.cpp b/Engine/Plugins/Animation/ControlRig/Source/ControlRigEditor/Private/Graph/SControlRigGraphPinParameterName.cpp
deleted file mode 100644
diff --git a/Engine/Plugins/Animation/ControlRig/Source/ControlRigEditor/Private/Graph/SControlRigGraphPinParameterName.h b/Engine/Plugins/Animation/ControlRig/Source/ControlRigEditor/Private/Graph/SControlRigGraphPinParameterName.h
deleted file mode 100644
diff --git a/Engine/Source/Developer/RigVMDeveloper/Private/RigVMCompiler/RigVMAST.cpp b/Engine/Source/Developer/RigVMDeveloper/Private/RigVMCompiler/RigVMAST.cpp
--- a/Engine/Source/Developer/RigVMDeveloper/Private/RigVMCompiler/RigVMAST.cpp
+++ b/Engine/Source/Developer/RigVMDeveloper/Private/RigVMCompiler/RigVMAST.cpp
@@ -511,15 +511,6 @@ bool FRigVMVarExprAST::IsExecuteContext() const
 	return GetPin()->IsExecuteContext();
 }
 
-bool FRigVMVarExprAST::IsGraphParameter() const
-{
-	if (Cast<URigVMParameterNode>(GetPin()->GetNode()))
-	{
-		return GetPin()->GetName() == TEXT("Value");
-	}
-	return false;
-}
-
 bool FRigVMVarExprAST::IsGraphVariable() const
 {
 	if (Cast<URigVMVariableNode>(GetPin()->GetNode()))
@@ -937,7 +928,6 @@ FRigVMExprAST* FRigVMParserAST::CreateExpressionForNode(const FRigVMASTProxy& In
 	else
 	{
 		if (InNodeProxy.IsA<URigVMRerouteNode>() ||
-			InNodeProxy.IsA<URigVMParameterNode>() ||
 			InNodeProxy.IsA<URigVMVariableNode>() ||
 			InNodeProxy.IsA<URigVMEnumNode>() ||
 			InNodeProxy.IsA<URigVMLibraryNode>() ||
@@ -1040,8 +1030,7 @@ FRigVMExprAST* FRigVMParserAST::TraversePin(const FRigVMASTProxy& InPinProxy, FR
 			return nullptr;
 		}
 	}
-	else if (Cast<URigVMParameterNode>(Pin->GetNode()) ||
-		Cast<URigVMEnumNode>(Pin->GetNode()))
+	else if (Cast<URigVMEnumNode>(Pin->GetNode()))
 	{
 		if (Pin->GetDirection() == ERigVMPinDirection::Visible)
 		{
@@ -1053,8 +1042,7 @@ FRigVMExprAST* FRigVMParserAST::TraversePin(const FRigVMASTProxy& InPinProxy, FR
 		Pin->GetDirection() == ERigVMPinDirection::Visible) &&
 		Links.Num() == 0)
 	{
-		if (Cast<URigVMParameterNode>(Pin->GetNode()) ||
-				 Cast<URigVMVariableNode>(Pin->GetNode()))
+		if (Cast<URigVMVariableNode>(Pin->GetNode()))
 		{
 			PinExpr = MakeExpr<FRigVMVarExprAST>(FRigVMExprAST::EType::Var, InPinProxy);
 			FRigVMExprAST* PinLiteralExpr = MakeExpr<FRigVMLiteralExprAST>(InPinProxy);
@@ -1154,8 +1142,7 @@ FRigVMExprAST* FRigVMParserAST::TraverseLink(const FRigVMPinProxyPair& InLink, F
 	bool bRequiresCopy = SourceRootPin != SourcePin || TargetRootPin != TargetPin;
 	if (!bRequiresCopy)
 	{
-		if(Cast<URigVMParameterNode>(TargetRootPin->GetNode()) ||
-			Cast<URigVMVariableNode>(TargetRootPin->GetNode()))
+		if(Cast<URigVMVariableNode>(TargetRootPin->GetNode()))
 		{
 			bRequiresCopy = true;
 		}
@@ -1432,13 +1419,6 @@ void FRigVMParserAST::FoldNoOps()
 		{
 			if (URigVMNode* Node = Expression->To<FRigVMNoOpExprAST>()->GetNode())
 			{
-				if (URigVMParameterNode* ParameterNode = Cast<URigVMParameterNode>(Node))
-				{
-					if (!ParameterNode->IsInput())
-					{
-						continue;
-					}
-				}
 				if (URigVMVariableNode* VariableNode = Cast<URigVMVariableNode>(Node))
 				{
 					if (!VariableNode->IsGetter())
@@ -1628,8 +1608,7 @@ bool FRigVMParserAST::FoldConstantValuesToLiterals(URigVMGraph* InGraph, URigVMC
 
 	for (const FRigVMASTProxy& NodeProxy : NodeProxies)
 	{
-		if (NodeProxy.IsA<URigVMParameterNode>() ||
-			NodeProxy.IsA<URigVMVariableNode>() ||
+		if (NodeProxy.IsA<URigVMVariableNode>() ||
 			NodeProxy.IsA<URigVMEnumNode>())
 		{
 			continue;
@@ -1701,8 +1680,7 @@ bool FRigVMParserAST::FoldConstantValuesToLiterals(URigVMGraph* InGraph, URigVMC
 
 				check(SourceNodeProxy.IsValid());
 
-				if (SourceNodeProxy.IsA<URigVMParameterNode>() ||
-					SourceNodeProxy.IsA<URigVMVariableNode>() ||
+				if (SourceNodeProxy.IsA<URigVMVariableNode>() ||
 					SourceNodeProxy.IsA<URigVMRerouteNode>() ||
 					SourceNodeProxy.IsA<URigVMEnumNode>())
 				{
@@ -1947,8 +1925,7 @@ bool FRigVMParserAST::FoldUnreachableBranches(URigVMGraph* InGraph)
 
 	for (const FRigVMASTProxy& NodeProxy : NodeProxies)
 	{
-		if (NodeProxy.IsA<URigVMParameterNode>() ||
-			NodeProxy.IsA<URigVMVariableNode>())
+		if (NodeProxy.IsA<URigVMVariableNode>())
 		{
 			continue;
 		}
@@ -2430,13 +2407,7 @@ FString FRigVMParserAST::DumpDot() const
 				}
 				case FRigVMExprAST::EType::Var:
 				{
-					if (InExpr->To<FRigVMVarExprAST>()->IsGraphParameter())
-					{
-						URigVMParameterNode* ParameterNode = Cast<URigVMParameterNode>(InExpr->To<FRigVMVarExprAST>()->GetPin()->GetNode());
-						check(ParameterNode);
-						Label = FString::Printf(TEXT("Param %s"), *ParameterNode->GetParameterName().ToString());
-					}
-					else if (InExpr->To<FRigVMVarExprAST>()->IsGraphVariable())
+					if (InExpr->To<FRigVMVarExprAST>()->IsGraphVariable())
 					{
 						URigVMVariableNode* VariableNode = Cast<URigVMVariableNode>(InExpr->To<FRigVMVarExprAST>()->GetPin()->GetNode());
 						check(VariableNode);
diff --git a/Engine/Source/Developer/RigVMDeveloper/Private/RigVMCompiler/RigVMCompiler.cpp b/Engine/Source/Developer/RigVMDeveloper/Private/RigVMCompiler/RigVMCompiler.cpp
--- a/Engine/Source/Developer/RigVMDeveloper/Private/RigVMCompiler/RigVMCompiler.cpp
+++ b/Engine/Source/Developer/RigVMDeveloper/Private/RigVMCompiler/RigVMCompiler.cpp
@@ -429,57 +429,12 @@ bool URigVMCompiler::Compile(URigVMGraph* InGraph, URigVMController* InControlle
 		}
 	}
 
-	// define all parameters independent from sorted nodes
-	for (const FRigVMExprAST* Expr : WorkData.AST->Expressions)
-	{
-		if (Expr->IsA(FRigVMExprAST::EType::Literal))
-		{
-			continue;
-		}
-		if (Expr->IsA(FRigVMExprAST::EType::Var))
-		{
-			const FRigVMVarExprAST* VarExpr = Expr->To<FRigVMVarExprAST>();
-			{
-				if(Cast<URigVMParameterNode>(VarExpr->GetPin()->GetNode()))
-				{
-					if (VarExpr->GetPin()->GetName() == TEXT("Value"))
-					{
-						FindOrAddRegister(VarExpr, WorkData, false /* watchvalue */);
-					}
-				}
-			}
-		}
-	}
-
 	WorkData.ExprComplete.Reset();
 	for (FRigVMExprAST* RootExpr : *WorkData.AST)
 	{
 		TraverseExpression(RootExpr, WorkData);
 	}
 
-	// If a parameter node has no corresponding AST node, because it's not on the execution path,
-	// then add a dummy parameter with no register index, so that we can properly add the pins on the 
-	// ControlRig node later and not lose any existing external connections.
-	for (URigVMNode* Node : InGraph->GetNodes())
-	{
-		if (URigVMParameterNode* ParameterNode = Cast<URigVMParameterNode>(Node))
-		{
-			FRigVMASTProxy ParameterNodeProxy = FRigVMASTProxy::MakeFromUObject(ParameterNode);
-			const FRigVMExprAST* ParameterExpr = WorkData.AST->GetExprForSubject(ParameterNodeProxy);
-			if (!ParameterExpr || ParameterExpr->IsA(FRigVMExprAST::NoOp))
-			{
-				FName Name = ParameterNode->GetParameterName();
-
-				if (!WorkData.VM->ParametersNameMap.Contains(Name))
-				{
-					ERigVMParameterType ParameterType = ParameterNode->IsInput() ? ERigVMParameterType::Input : ERigVMParameterType::Output;
-					FRigVMParameter Parameter(ParameterType, Name, INDEX_NONE, ParameterNode->GetCPPType(), nullptr);
-					WorkData.VM->ParametersNameMap.Add(Parameter.Name, WorkData.VM->Parameters.Add(Parameter));
-				}
-			}
-		}
-	}
-
 	if(WorkData.WatchedPins.Num() > 0)
 	{
 		for(int32 GraphIndex=0; GraphIndex<VisitedGraphs.Num(); GraphIndex++)
@@ -1945,7 +1900,6 @@ FString URigVMCompiler::GetPinHash(const URigVMPin* InPin, const FRigVMVarExprAS
 
 	bool bIsExecutePin = false;
 	bool bIsLiteral = false;
-	bool bIsParameter = false;
 	bool bIsVariable = false;
 
 	if (InVarExpr != nullptr && !bIsDebugValue)
@@ -1969,11 +1923,11 @@ FString URigVMCompiler::GetPinHash(const URigVMPin* InPin, const FRigVMVarExprAS
 		bIsExecutePin = InPin->IsExecuteContext();
 		bIsLiteral = InVarExpr->GetType() == FRigVMExprAST::EType::Literal;
 
-		bIsParameter = Cast<URigVMParameterNode>(Node) != nullptr;
 		bIsVariable = Cast<URigVMVariableNode>(Node) != nullptr || InVarExpr->IsA(FRigVMExprAST::ExternalVar);
 
 		// determine if this is an initialization for an IO pin
-		if (!bIsLiteral && !bIsParameter && !bIsVariable &&
+		if (!bIsLiteral &&
+			!bIsVariable &&
 			!bIsExecutePin && (InPin->GetDirection() == ERigVMPinDirection::IO ||
 			(InPin->GetDirection() == ERigVMPinDirection::Input && InPin->GetSourceLinks().Num() == 0)))
 		{
@@ -1986,14 +1940,7 @@ FString URigVMCompiler::GetPinHash(const URigVMPin* InPin, const FRigVMVarExprAS
 	}
 
 	bool bUseFullNodePath = true;
-	if (URigVMParameterNode* ParameterNode = Cast<URigVMParameterNode>(Node))
-	{
-		if (InPin->GetName() == TEXT("Value") && !bIsLiteral && !bIsDebugValue)
-		{
-			return FString::Printf(TEXT("%sParameter::%s%s"), *Prefix, *ParameterNode->GetParameterName().ToString(), *Suffix);
-		}
-	}
-	else if (URigVMVariableNode* VariableNode = Cast<URigVMVariableNode>(Node))
+	if (URigVMVariableNode* VariableNode = Cast<URigVMVariableNode>(Node))
 	{
 		if (InPin->GetName() == TEXT("Value") && !bIsDebugValue)
 		{
@@ -2379,7 +2326,6 @@ FRigVMOperand URigVMCompiler::FindOrAddRegister(const FRigVMVarExprAST* InVarExp
 
 	bool bIsExecutePin = Pin->IsExecuteContext();
 	bool bIsLiteral = InVarExpr->GetType() == FRigVMExprAST::EType::Literal && !bIsDebugValue;
-	bool bIsParameter = InVarExpr->IsGraphParameter();
 	bool bIsVariable = Pin->IsRootPin() && (Pin->GetName() == URigVMVariableNode::ValueName) &&
 		InVarExpr->GetPin()->GetNode()->IsA<URigVMVariableNode>();
 
@@ -2528,16 +2474,6 @@ FRigVMOperand URigVMCompiler::FindOrAddRegister(const FRigVMVarExprAST* InVarExp
 		}
 	
 		Operand = WorkData.AddProperty(MemoryType, RegisterName, CPPType, Pin->GetCPPTypeObject(), JoinedDefaultValue);
-
-		if (bIsParameter && !bIsLiteral && !bIsDebugValue)
-		{
-			URigVMParameterNode* ParameterNode = Cast<URigVMParameterNode>(Pin->GetNode());
-			check(ParameterNode);
-			FName Name = ParameterNode->GetParameterName();
-			ERigVMParameterType ParameterType = ParameterNode->IsInput() ? ERigVMParameterType::Input : ERigVMParameterType::Output;
-			FRigVMParameter Parameter(ParameterType, Name, Operand.GetRegisterIndex(), Pin->GetCPPType(), ScriptStruct);
-			WorkData.VM->ParametersNameMap.FindOrAdd(Parameter.Name) = WorkData.VM->Parameters.Add(Parameter);
-		}
 	}
 	ensure(Operand.IsValid());
 
@@ -2601,7 +2537,7 @@ const FRigVMCompilerWorkData::FRigVMASTProxyArray& URigVMCompiler::FindProxiesWi
 		{
 			if (URigVMPin* Pin = Cast<URigVMPin>(CurrentProxy.GetSubject()))
 			{
-				if (Pin->GetNode()->IsA<URigVMVariableNode>() || Pin->GetNode()->IsA<URigVMParameterNode>())
+				if (Pin->GetNode()->IsA<URigVMVariableNode>())
 				{
 					if (Pin->GetDirection() == ERigVMPinDirection::Input)
 					{
diff --git a/Engine/Source/Developer/RigVMDeveloper/Private/RigVMModel/Nodes/RigVMParameterNode.cpp b/Engine/Source/Developer/RigVMDeveloper/Private/RigVMModel/Nodes/RigVMParameterNode.cpp
--- a/Engine/Source/Developer/RigVMDeveloper/Private/RigVMModel/Nodes/RigVMParameterNode.cpp
+++ b/Engine/Source/Developer/RigVMDeveloper/Private/RigVMModel/Nodes/RigVMParameterNode.cpp
@@ -58,7 +58,7 @@ UObject* URigVMParameterNode::GetCPPTypeObject() const
 FString URigVMParameterNode::GetDefaultValue() const
 {
 
-	URigVMPin* ValuePin = FindPin(ValueName);
+	URigVMPin* ValuePin = FindPin(DefaultName);
 	if (ValuePin == nullptr)
 	{
 		return FString();
diff --git a/Engine/Source/Developer/RigVMDeveloper/Private/RigVMModel/RigVMController.cpp b/Engine/Source/Developer/RigVMDeveloper/Private/RigVMModel/RigVMController.cpp
--- a/Engine/Source/Developer/RigVMDeveloper/Private/RigVMModel/RigVMController.cpp
+++ b/Engine/Source/Developer/RigVMDeveloper/Private/RigVMModel/RigVMController.cpp
@@ -255,8 +255,6 @@ void URigVMController::HandleModifiedEvent(ERigVMGraphNotifType InNotifType, URi
 		case ERigVMGraphNotifType::LinkAdded:
 		case ERigVMGraphNotifType::LinkRemoved:
 		case ERigVMGraphNotifType::PinArraySizeChanged:
-		case ERigVMGraphNotifType::ParameterAdded:
-		case ERigVMGraphNotifType::ParameterRemoved:
 		{
 			if (InGraph)
 			{
@@ -598,35 +596,6 @@ TArray<FString> URigVMController::GetAddNodePythonCommands(URigVMNode* Node) con
 			}
 		}
 	}
-	else if (const URigVMParameterNode* ParameterNode = Cast<URigVMParameterNode>(Node))
-	{
-		const FString ParameterName = GetSanitizedVariableName(ParameterNode->GetParameterName().ToString());
-
-		// add_parameter_node_from_object_path(parameter_name, cpp_type, cpp_type_object_path, is_input, default_value, position=[0.0, 0.0], node_name='', undo=True)
-		if (ParameterNode->GetParameterDescription().CPPTypeObject)
-		{
-			Commands.Add(FString::Printf(TEXT("blueprint.get_controller_by_name('%s').add_parameter_node_from_object_path('%s', '%s', '%s', %s, '%s', %s, '%s')"),
-					*GraphName,
-					*ParameterName,
-					*ParameterNode->GetParameterDescription().CPPType,
-					*ParameterNode->GetParameterDescription().CPPTypeObject->GetPathName(),
-					ParameterNode->IsInput() ? TEXT("True") : TEXT("False"),
-					*ParameterNode->GetParameterDescription().DefaultValue,
-					*RigVMPythonUtils::Vector2DToPythonString(ParameterNode->GetPosition()),
-					*NodeName));	
-		}
-		else
-		{
-			Commands.Add(FString::Printf(TEXT("blueprint.get_controller_by_name('%s').add_parameter_node('%s', '%s', None, %s, str(%s), %s, '%s')"),
-					*GraphName,
-					*ParameterName,
-					*ParameterNode->GetParameterDescription().CPPType,
-					ParameterNode->IsInput() ? TEXT("True") : TEXT("False"),
-					*ParameterNode->GetParameterDescription().DefaultValue,
-					*RigVMPythonUtils::Vector2DToPythonString(ParameterNode->GetPosition()),
-					*NodeName));	
-		}
-	}
 	else if (const URigVMCommentNode* CommentNode = Cast<URigVMCommentNode>(Node))
 	{
 		// add_comment_node(comment_text, position=[0.0, 0.0], size=[400.0, 300.0], color=[0.0, 0.0, 0.0, 0.0], node_name='', undo=True)
@@ -1664,168 +1633,9 @@ URigVMTemplateNode* URigVMController::ReplaceUnitNodeWithTemplateNode(const FNam
 
 URigVMParameterNode* URigVMController::AddParameterNode(const FName& InParameterName, const FString& InCPPType, UObject* InCPPTypeObject, bool bIsInput, const FString& InDefaultValue, const FVector2D& InPosition, const FString& InNodeName, bool bSetupUndoRedo, bool bPrintPythonCommand)
 {
-	if (!IsValidGraph())
-	{
-		return nullptr;
-	}
-
-	URigVMGraph* Graph = GetGraph();
-	check(Graph);
-
-	if (Graph->IsA<URigVMFunctionLibrary>())
-	{
-		ReportError(TEXT("Cannot add parameter nodes to function library graphs."));
-		return nullptr;
-	}
-
-	if (InCPPTypeObject == nullptr)
-	{
-		InCPPTypeObject = URigVMCompiler::GetScriptStructForCPPType(InCPPType);
-	}
-	if (InCPPTypeObject == nullptr)
-	{
-		InCPPTypeObject = URigVMPin::FindObjectFromCPPTypeObjectPath<UObject>(InCPPType);
-	}
-
-	TArray<FRigVMGraphParameterDescription> ExistingParameters = Graph->GetParameterDescriptions();
-	for (const FRigVMGraphParameterDescription& ExistingParameter : ExistingParameters)
-	{
-		if (ExistingParameter.Name == InParameterName)
-		{
-			if (ExistingParameter.CPPType != InCPPType ||
-				ExistingParameter.CPPTypeObject != InCPPTypeObject ||
-				ExistingParameter.bIsInput != bIsInput)
-			{
-				ReportErrorf(TEXT("Cannot add parameter '%s' - parameter already exists."), *InParameterName.ToString());
-				return nullptr;
-			}
-		}
-	}
-
-	FString Name = GetValidNodeName(InNodeName.IsEmpty() ? FString(TEXT("ParameterNode")) : InNodeName);
-	URigVMParameterNode* Node = NewObject<URigVMParameterNode>(Graph, *Name);
-	Node->Position = InPosition;
-
-	if (!bIsInput)
-	{
-		URigVMPin* ExecutePin = NewObject<URigVMPin>(Node, FRigVMStruct::ExecuteContextName);
-		ExecutePin->CPPType = FString::Printf(TEXT("F%s"), *ExecuteContextStruct->GetName());
-		ExecutePin->CPPTypeObject = ExecuteContextStruct;
-		ExecutePin->CPPTypeObjectPath = *ExecutePin->CPPTypeObject->GetPathName();
-		ExecutePin->Direction = ERigVMPinDirection::IO;
-		AddNodePin(Node, ExecutePin);
-	}
-
-	URigVMPin* ParameterPin = NewObject<URigVMPin>(Node, *URigVMParameterNode::ParameterName);
-	ParameterPin->CPPType = RigVMTypeUtils::FNameType;
-	ParameterPin->Direction = ERigVMPinDirection::Visible;
-	ParameterPin->DefaultValue = InParameterName.ToString();
-	ParameterPin->CustomWidgetName = TEXT("ParameterName");
-
-	AddNodePin(Node, ParameterPin);
-
-	URigVMPin* DefaultValuePin = nullptr;
-	if (bIsInput)
-	{
-		DefaultValuePin = NewObject<URigVMPin>(Node, *URigVMParameterNode::DefaultName);
-	}
-	URigVMPin* ValuePin = NewObject<URigVMPin>(Node, *URigVMParameterNode::ValueName);
-
-	if (DefaultValuePin)
-	{
-		DefaultValuePin->CPPType = PostProcessCPPType(InCPPType, InCPPTypeObject);
-		DefaultValuePin->CPPTypeObject = InCPPTypeObject;
-		if(DefaultValuePin->CPPTypeObject)
-		{
-			DefaultValuePin->CPPTypeObjectPath = *DefaultValuePin->CPPTypeObject->GetPathName();
-		}
-	}
-	
-	ValuePin->CPPType = PostProcessCPPType(InCPPType, InCPPTypeObject);
-	ValuePin->CPPTypeObject = InCPPTypeObject;
-	if(ValuePin->CPPTypeObject)
-	{
-		ValuePin->CPPTypeObjectPath = *ValuePin->CPPTypeObject->GetPathName();
-	}
-
-	if (DefaultValuePin)
-	{
-		DefaultValuePin->Direction = ERigVMPinDirection::Visible;
-	}
-	ValuePin->Direction = bIsInput ? ERigVMPinDirection::Output : ERigVMPinDirection::Input;
-
-	if (bIsInput)
-	{
-		if (ValuePin->CPPType == RigVMTypeUtils::FNameType)
-		{
-			ValuePin->bIsConstant = true;
-		}
-	}
-
-	if (DefaultValuePin)
-	{
-		AddNodePin(Node, DefaultValuePin);
-	}
-	AddNodePin(Node, ValuePin);
-
-	Graph->Nodes.Add(Node);
-
-	if (ValuePin->IsStruct())
-	{
-		FString DefaultValue = InDefaultValue;
-		CreateDefaultValueForStructIfRequired(ValuePin->GetScriptStruct(), DefaultValue);
-		if (DefaultValuePin)
-		{
-			AddPinsForStruct(DefaultValuePin->GetScriptStruct(), Node, DefaultValuePin, DefaultValuePin->Direction, DefaultValue, false);
-		}
-		AddPinsForStruct(ValuePin->GetScriptStruct(), Node, ValuePin, ValuePin->Direction, DefaultValue, false);
-	}
-	else if (!InDefaultValue.IsEmpty() && InDefaultValue != TEXT("()"))
-	{
-		if (DefaultValuePin)
-		{
-			SetPinDefaultValue(DefaultValuePin, InDefaultValue, true, false, false);
-		}
-		SetPinDefaultValue(ValuePin, InDefaultValue, true, false, false);
-	}
-
-	ForEveryPinRecursively(Node, [](URigVMPin* Pin) {
-		Pin->bIsExpanded = false;
-	});
-
-	if (!bSuspendNotifications)
-	{
-		Graph->MarkPackageDirty();
-	}
-
-	FRigVMControllerCompileBracketScope CompileScope(this);
-	FRigVMAddParameterNodeAction Action;
-	if (bSetupUndoRedo)
-	{
-		Action = FRigVMAddParameterNodeAction(Node);
-		Action.Title = FString::Printf(TEXT("Add %s Parameter"), *InParameterName.ToString());
-		ActionStack->BeginAction(Action);
-	}
-
-	Notify(ERigVMGraphNotifType::NodeAdded, Node);
-	Notify(ERigVMGraphNotifType::ParameterAdded, Node);
-
-	if (bSetupUndoRedo)
-	{
-		ActionStack->EndAction(Action);
-	}
-
-	if (bPrintPythonCommand)
-	{
-		TArray<FString> Commands = GetAddNodePythonCommands(Node);
-		for (const FString& Command : Commands)
-		{
-			RigVMPythonUtils::Print(GetGraphOuterName(), 
-								FString::Printf(TEXT("%s"), *Command));
-		}
-	}
-
-	return Node;
+	AddVariableNode(InParameterName, InCPPType, InCPPTypeObject, bIsInput, InDefaultValue, InPosition, InNodeName, bSetupUndoRedo, bPrintPythonCommand);
+	ReportWarning(TEXT("AddParameterNode has been deprecated. Adding a variable node instead."));
+	return nullptr;
 }
 
 URigVMParameterNode* URigVMController::AddParameterNodeFromObjectPath(const FName& InParameterName, const FString& InCPPType, const FString& InCPPTypeObjectPath, bool bIsInput, const FString& InDefaultValue, const FVector2D& InPosition, const FString& InNodeName, bool bSetupUndoRedo, bool bPrintPythonCommand)
@@ -5378,11 +5188,7 @@ bool URigVMController::RemoveNode(URigVMNode* InNode, bool bSetupUndoRedo, bool
 	{
 		Notify(ERigVMGraphNotifType::VariableRemoved, VariableNode);
 	}
-	if (URigVMParameterNode* ParameterNode = Cast<URigVMParameterNode>(InNode))
-	{
-		Notify(ERigVMGraphNotifType::ParameterRemoved, ParameterNode);
-	}
-
+	
 	if (URigVMInjectionInfo* InjectionInfo = InNode->GetInjectionInfo())
 	{
 		DestroyObject(InjectionInfo);
@@ -6191,74 +5997,8 @@ bool URigVMController::RenameVariable(const FName& InOldName, const FName& InNew
 
 bool URigVMController::RenameParameter(const FName& InOldName, const FName& InNewName, bool bSetupUndoRedo)
 {
-	if(!IsValidGraph())
-	{
-		return false;
-	}
-
-	if (InOldName == InNewName)
-	{
-		ReportWarning(TEXT("RenameParameter: InOldName and InNewName are equal."));
-		return false;
-	}
-
-	URigVMGraph* Graph = GetGraph();
-	check(Graph);
-
-	TArray<FRigVMGraphParameterDescription> ExistingParameters = Graph->GetParameterDescriptions();
-	for (const FRigVMGraphParameterDescription& ExistingParameter : ExistingParameters)
-	{
-		if (ExistingParameter.Name == InNewName)
-		{
-			ReportErrorf(TEXT("Cannot rename parameter to '%s' - parameter already exists."), *InNewName.ToString());
-			return false;
-		}
-	}
-
-	FRigVMControllerCompileBracketScope CompileScope(this);
-	FRigVMRenameParameterAction Action;
-	if (bSetupUndoRedo)
-	{
-		Action = FRigVMRenameParameterAction(InOldName, InNewName);
-		Action.Title = FString::Printf(TEXT("Rename Parameter"));
-		ActionStack->BeginAction(Action);
-	}
-
-	TArray<URigVMNode*> RenamedNodes;
-	for (URigVMNode* Node : Graph->Nodes)
-	{
-		if(URigVMParameterNode* ParameterNode = Cast<URigVMParameterNode>(Node))
-		{
-			if (ParameterNode->GetParameterName() == InOldName)
-			{
-				ParameterNode->FindPin(URigVMParameterNode::ParameterName)->DefaultValue = InNewName.ToString();
-				RenamedNodes.Add(Node);
-			}
-		}
-	}
-
-	for (URigVMNode* RenamedNode : RenamedNodes)
-	{
-		Notify(ERigVMGraphNotifType::ParameterRenamed, RenamedNode);
-		if (!bSuspendNotifications)
-		{
-			Graph->MarkPackageDirty();
-		}
-	}
-
-	if (bSetupUndoRedo)
-	{
-		if (RenamedNodes.Num() > 0)
-		{
-			ActionStack->EndAction(Action);
-		}
-		else
-		{
-			ActionStack->CancelAction(Action);
-		}
-	}
-
-	return RenamedNodes.Num() > 0;
+	ReportWarning(TEXT("RenameParameter has been deprecated. Please use RenameVariable instead."));
+	return false;
 }
 
 void URigVMController::UpdateRerouteNodeAfterChangingLinks(URigVMPin* PinChanged, bool bSetupUndoRedo)
@@ -6474,14 +6214,7 @@ bool URigVMController::SetPinDefaultValue(const FString& InPinPath, const FStrin
 			return SetVariableName(VariableNode, *InDefaultValue, bSetupUndoRedo);
 		}
 	}
-	if (URigVMParameterNode* ParameterNode = Cast<URigVMParameterNode>(Pin->GetNode()))
-	{
-		if (Pin->GetName() == URigVMParameterNode::ParameterName)
-		{
-			return SetParameterName(ParameterNode, *InDefaultValue, bSetupUndoRedo);
-		}
-	}
-
+	
 	if (!SetPinDefaultValue(Pin, InDefaultValue, bResizeArrays, bSetupUndoRedo, bMergeUndoAction))
 	{
 		return false;
@@ -9881,66 +9614,6 @@ bool URigVMController::SetVariableName(URigVMVariableNode* InVariableNode, const
 	return true;
 }
 
-bool URigVMController::SetParameterName(URigVMParameterNode* InParameterNode, const FName& InParameterName, bool bSetupUndoRedo)
-{
-	if (!IsValidNodeForGraph(InParameterNode))
-	{
-		return false;
-	}
-
-	if (InParameterNode->GetParameterName() == InParameterName)
-	{
-		return false;
-	}
-
-	if (InParameterName == NAME_None)
-	{
-		return false;
-	}
-
-	URigVMGraph* Graph = GetGraph();
-	check(Graph);
-
-	TArray<FRigVMGraphParameterDescription> Descriptions = Graph->GetParameterDescriptions();
-	TMap<FName, int32> NameToIndex;
-	for (int32 ParameterIndex = 0; ParameterIndex < Descriptions.Num(); ParameterIndex++)
-	{
-		NameToIndex.Add(Descriptions[ParameterIndex].Name, ParameterIndex);
-	}
-
-	FName ParameterName = GetUniqueName(InParameterName, [Descriptions, NameToIndex, InParameterNode](const FName& InName) {
-		const int32* FoundIndex = NameToIndex.Find(InName);
-		if (FoundIndex == nullptr)
-		{
-			return true;
-		}
-		return InParameterNode->GetCPPType() == Descriptions[*FoundIndex].CPPType && InParameterNode->IsInput() == Descriptions[*FoundIndex].bIsInput;
-	}, false, true);
-
-	int32 NodesSharingName = 0;
-	for (URigVMNode* Node : Graph->Nodes)
-	{
-		if (URigVMParameterNode* OtherParameterNode = Cast<URigVMParameterNode>(Node))
-		{
-			if (OtherParameterNode->GetParameterName() == InParameterNode->GetParameterName())
-			{
-				NodesSharingName++;
-			}
-		}
-	}
-
-	if (NodesSharingName == 1)
-	{
-		Notify(ERigVMGraphNotifType::ParameterRemoved, InParameterNode);
-	}
-
-	SetPinDefaultValue(InParameterNode->FindPin(URigVMParameterNode::ParameterName), ParameterName.ToString(), false, bSetupUndoRedo, false);
-
-	Notify(ERigVMGraphNotifType::ParameterAdded, InParameterNode);
-
-	return true;
-}
-
 URigVMRerouteNode* URigVMController::AddFreeRerouteNode(bool bShowAsFullNode, const FString& InCPPType, const FName& InCPPTypeObjectPath, bool bIsConstant, const FName& InCustomWidgetName, const FString& InDefaultValue, const FVector2D& InPosition, const FString& InNodeName, bool bSetupUndoRedo)
 {
 	if (!IsValidGraph())
diff --git a/Engine/Source/Developer/RigVMDeveloper/Private/RigVMModel/RigVMControllerActions.cpp b/Engine/Source/Developer/RigVMDeveloper/Private/RigVMModel/RigVMControllerActions.cpp
--- a/Engine/Source/Developer/RigVMDeveloper/Private/RigVMModel/RigVMControllerActions.cpp
+++ b/Engine/Source/Developer/RigVMDeveloper/Private/RigVMModel/RigVMControllerActions.cpp
@@ -301,52 +301,6 @@ bool FRigVMAddVariableNodeAction::Redo(URigVMController* InController)
 	return false;
 }
 
-FRigVMAddParameterNodeAction::FRigVMAddParameterNodeAction()
-	: ParameterName(NAME_None)
-	, CPPType()
-	, CPPTypeObjectPath()
-	, bIsInput(false)
-	, DefaultValue()
-	, Position(FVector2D::ZeroVector)
-	, NodePath()
-{
-}
-
-FRigVMAddParameterNodeAction::FRigVMAddParameterNodeAction(URigVMParameterNode* InNode)
-	: ParameterName(InNode->GetParameterName())
-	, CPPType(InNode->GetCPPType())
-	, CPPTypeObjectPath()
-	, bIsInput(InNode->IsInput())
-	, DefaultValue(InNode->GetDefaultValue())
-	, Position(InNode->GetPosition())
-	, NodePath(InNode->GetNodePath())
-{
-	if (InNode->GetCPPTypeObject())
-	{
-		CPPTypeObjectPath = InNode->GetCPPTypeObject()->GetPathName();
-	}
-}
-
-bool FRigVMAddParameterNodeAction::Undo(URigVMController* InController)
-{
-	if (!FRigVMBaseAction::Undo(InController))
-	{
-		return false;
-	}
-	return InController->RemoveNodeByName(*NodePath, false);
-}
-
-bool FRigVMAddParameterNodeAction::Redo(URigVMController* InController)
-{
-#if WITH_EDITOR
-	if (URigVMParameterNode* Node = InController->AddParameterNodeFromObjectPath(ParameterName, CPPType, CPPTypeObjectPath, bIsInput, DefaultValue, Position, NodePath, false))
-	{
-		return FRigVMBaseAction::Redo(InController);
-	}
-#endif
-	return false;
-}
-
 FRigVMAddCommentNodeAction::FRigVMAddCommentNodeAction()
 	: CommentText()
 	, Position(FVector2D::ZeroVector)
@@ -688,12 +642,6 @@ FRigVMRemoveNodeAction::FRigVMRemoveNodeAction(URigVMNode* InNode, URigVMControl
 		URigVMPin* ValuePin = VariableNode->FindPin(TEXT("Value"));
 		InverseAction.AddAction(FRigVMSetPinDefaultValueAction(ValuePin, ValuePin->GetDefaultValue()));
 	}
-	else if (URigVMParameterNode* ParameterNode = Cast<URigVMParameterNode>(InNode))
-	{
-		InverseAction.AddAction(FRigVMAddParameterNodeAction(ParameterNode));
-		URigVMPin* ValuePin = ParameterNode->FindPin(TEXT("Value"));
-		InverseAction.AddAction(FRigVMSetPinDefaultValueAction(ValuePin, ValuePin->GetDefaultValue()));
-	}
 	else if (URigVMCommentNode* CommentNode = Cast<URigVMCommentNode>(InNode))
 	{
 		InverseAction.AddAction(FRigVMAddCommentNodeAction(CommentNode));
@@ -1176,30 +1124,6 @@ bool FRigVMRenameVariableAction::Redo(URigVMController* InController)
 	return FRigVMBaseAction::Redo(InController);
 }
 
-FRigVMRenameParameterAction::FRigVMRenameParameterAction(const FName& InOldParameterName, const FName& InNewParameterName)
-: OldParameterName(InOldParameterName.ToString())
-, NewParameterName(InNewParameterName.ToString())
-{
-}
-
-bool FRigVMRenameParameterAction::Undo(URigVMController* InController)
-{
-	if(!FRigVMBaseAction::Undo(InController))
-	{
-		return false;
-	}
-	return InController->RenameParameter(*NewParameterName, *OldParameterName, false);
-}
-
-bool FRigVMRenameParameterAction::Redo(URigVMController* InController)
-{
-	if(!InController->RenameParameter(*OldParameterName, *NewParameterName, false))
-	{
-		return false;
-	}
-	return FRigVMBaseAction::Redo(InController);
-}
-
 FRigVMSetPinExpansionAction::FRigVMSetPinExpansionAction(URigVMPin* InPin, bool bNewIsExpanded)
 : PinPath(InPin->GetPinPath())
 , OldIsExpanded(InPin->IsExpanded())
diff --git a/Engine/Source/Developer/RigVMDeveloper/Private/RigVMModel/RigVMGraph.cpp b/Engine/Source/Developer/RigVMDeveloper/Private/RigVMModel/RigVMGraph.cpp
--- a/Engine/Source/Developer/RigVMDeveloper/Private/RigVMModel/RigVMGraph.cpp
+++ b/Engine/Source/Developer/RigVMDeveloper/Private/RigVMModel/RigVMGraph.cpp
@@ -102,19 +102,6 @@ TArray<FRigVMGraphVariableDescription> URigVMGraph::GetVariableDescriptions() co
 	return Variables;
 }
 
-TArray<FRigVMGraphParameterDescription> URigVMGraph::GetParameterDescriptions() const
-{
-	TArray<FRigVMGraphParameterDescription> Parameters;
-	for (URigVMNode* Node : Nodes)
-	{
-		if (URigVMParameterNode* ParameterNode = Cast<URigVMParameterNode>(Node))
-		{
-			Parameters.AddUnique(ParameterNode->GetParameterDescription());
-		}
-	}
-	return Parameters;
-}
-
 FString URigVMGraph::GetNodePath() const
 {
 	if (URigVMCollapseNode* CollapseNode = Cast<URigVMCollapseNode>(GetOuter()))
diff --git a/Engine/Source/Developer/RigVMDeveloper/Private/RigVMModel/RigVMPin.cpp b/Engine/Source/Developer/RigVMDeveloper/Private/RigVMModel/RigVMPin.cpp
--- a/Engine/Source/Developer/RigVMDeveloper/Private/RigVMModel/RigVMPin.cpp
+++ b/Engine/Source/Developer/RigVMDeveloper/Private/RigVMModel/RigVMPin.cpp
@@ -256,7 +256,7 @@ void URigVMPin::GetExposedPinChain(TArray<const URigVMPin*>& OutExposedPins) con
 void URigVMPin::GetExposedPinChainImpl(TArray<const URigVMPin*>& OutExposedPins, TArray<const URigVMPin*>& VisitedPins) const
 {
 	// Variable nodes do not share the operand with their source link
-	if ((GetNode()->IsA<URigVMVariableNode>() || GetNode()->IsA<URigVMParameterNode>()) && GetDirection() == ERigVMPinDirection::Input)
+	if (GetNode()->IsA<URigVMVariableNode>() && GetDirection() == ERigVMPinDirection::Input)
 	{
 		OutExposedPins.Add(this);
 		return;
@@ -299,7 +299,7 @@ void URigVMPin::GetExposedPinChainImpl(TArray<const URigVMPin*>& OutExposedPins,
 			}
 		}
 		// Variable nodes do not share the operand with their source link
-		else if (SourcePin->GetNode()->IsA<URigVMVariableNode>() || SourcePin->GetNode()->IsA<URigVMParameterNode>())
+		else if (SourcePin->GetNode()->IsA<URigVMVariableNode>())
 		{
 			continue;
 		}
@@ -331,7 +331,7 @@ void URigVMPin::GetExposedPinChainImpl(TArray<const URigVMPin*>& OutExposedPins,
 			URigVMPin* TargetPin = Link->GetTargetPin();
 
 			// Variable nodes do not share the operand with their source link
-			if (TargetPin->GetNode()->IsA<URigVMVariableNode>() || TargetPin->GetNode()->IsA<URigVMParameterNode>())
+			if (TargetPin->GetNode()->IsA<URigVMVariableNode>())
 			{
 				continue;
 			}
diff --git a/Engine/Source/Developer/RigVMDeveloper/Public/RigVMCompiler/RigVMAST.h b/Engine/Source/Developer/RigVMDeveloper/Public/RigVMCompiler/RigVMAST.h
--- a/Engine/Source/Developer/RigVMDeveloper/Public/RigVMCompiler/RigVMAST.h
+++ b/Engine/Source/Developer/RigVMDeveloper/Public/RigVMCompiler/RigVMAST.h
@@ -687,10 +687,6 @@ public:
 	// @return true if this variable is an execute context
 	bool IsExecuteContext() const;
 
-	// returns true if this variable is a graph parameter
-	// @return true if this variable is a graph parameter
-	bool IsGraphParameter() const;
-
 	// returns true if this variable is a graph variable
 	// @return true if this variable is a graph variable
 	bool IsGraphVariable() const;
diff --git a/Engine/Source/Developer/RigVMDeveloper/Public/RigVMModel/Nodes/RigVMParameterNode.h b/Engine/Source/Developer/RigVMDeveloper/Public/RigVMModel/Nodes/RigVMParameterNode.h
--- a/Engine/Source/Developer/RigVMDeveloper/Public/RigVMModel/Nodes/RigVMParameterNode.h
+++ b/Engine/Source/Developer/RigVMDeveloper/Public/RigVMModel/Nodes/RigVMParameterNode.h
@@ -12,7 +12,7 @@
  * about unique parameters within a Graph. Multiple Parameter
  * Nodes can share the same parameter description.
  */
-USTRUCT(BlueprintType)
+USTRUCT(BlueprintType, meta = (Deprecated = "5.1"))
 struct FRigVMGraphParameterDescription
 {
 	GENERATED_BODY()
@@ -75,7 +75,7 @@ public:
  * The Parameter Node represents an input or output argument / parameter
  * of the Function / Graph. Parameter Node have only a single value pin.
  */
-UCLASS(BlueprintType)
+UCLASS(BlueprintType, meta = (Deprecated = "5.1"))
 class RIGVMDEVELOPER_API URigVMParameterNode : public URigVMNode
 {
 	GENERATED_BODY()
diff --git a/Engine/Source/Developer/RigVMDeveloper/Public/RigVMModel/RigVMController.h b/Engine/Source/Developer/RigVMDeveloper/Public/RigVMModel/RigVMController.h
--- a/Engine/Source/Developer/RigVMDeveloper/Public/RigVMModel/RigVMController.h
+++ b/Engine/Source/Developer/RigVMDeveloper/Public/RigVMModel/RigVMController.h
@@ -221,14 +221,14 @@ public:
 	// Parameters represent input or output arguments to the Graph / Function.
 	// Input Parameters are constant values / literals.
 	// This causes a NodeAdded modified event.
-	UFUNCTION(BlueprintCallable, Category = RigVMController)
+	UFUNCTION(BlueprintCallable, Category = RigVMController, meta=(DeprecatedFunction))
 	URigVMParameterNode* AddParameterNode(const FName& InParameterName, const FString& InCPPType, UObject* InCPPTypeObject, bool bIsInput, const FString& InDefaultValue, const FVector2D& InPosition = FVector2D::ZeroVector, const FString& InNodeName = TEXT(""), bool bSetupUndoRedo = true, bool bPrintPythonCommand = false);
 
 	// Adds a Parameter Node to the edited Graph given a struct object path name.
 	// Parameters represent input or output arguments to the Graph / Function.
 	// Input Parameters are constant values / literals.
 	// This causes a NodeAdded modified event.
-	UFUNCTION(BlueprintCallable, Category = RigVMController)
+	UFUNCTION(BlueprintCallable, Category = RigVMController, meta=(DeprecatedFunction))
 	URigVMParameterNode* AddParameterNodeFromObjectPath(const FName& InParameterName, const FString& InCPPType, const FString& InCPPTypeObjectPath, bool bIsInput, const FString& InDefaultValue, const FVector2D& InPosition = FVector2D::ZeroVector, const FString& InNodeName = TEXT(""), bool bSetupUndoRedo = true, bool bPrintPythonCommand = false);
 
 	// Adds a Comment Node to the edited Graph.
@@ -532,7 +532,7 @@ public:
 
 	// Renames a parameter in the graph.
 	// This causes a ParameterRenamed modified event.
-	UFUNCTION(BlueprintCallable, Category = RigVMController)
+	UFUNCTION(BlueprintCallable, Category = RigVMController, meta=(DeprecatedFunction))
 	bool RenameParameter(const FName& InOldName, const FName& InNewName, bool bSetupUndoRedo = true);
 
 	// Sets the pin to be expanded or not
@@ -849,7 +849,6 @@ private:
 	void ExpandPinRecursively(URigVMPin* InPin, bool bSetupUndoRedo);
 	bool SetPinIsWatched(URigVMPin* InPin, bool bIsWatched, bool bSetupUndoRedo);
 	bool SetVariableName(URigVMVariableNode* InVariableNode, const FName& InVariableName, bool bSetupUndoRedo);
-	bool SetParameterName(URigVMParameterNode* InParameterNode, const FName& InParameterName, bool bSetupUndoRedo);
 	static void ForEveryPinRecursively(URigVMPin* InPin, TFunction<void(URigVMPin*)> OnEachPinFunction);
 	static void ForEveryPinRecursively(URigVMNode* InNode, TFunction<void(URigVMPin*)> OnEachPinFunction);
 	URigVMCollapseNode* CollapseNodes(const TArray<URigVMNode*>& InNodes, const FString& InCollapseNodeName, bool bSetupUndoRedo);
diff --git a/Engine/Source/Developer/RigVMDeveloper/Public/RigVMModel/RigVMControllerActions.h b/Engine/Source/Developer/RigVMDeveloper/Public/RigVMModel/RigVMControllerActions.h
--- a/Engine/Source/Developer/RigVMDeveloper/Public/RigVMModel/RigVMControllerActions.h
+++ b/Engine/Source/Developer/RigVMDeveloper/Public/RigVMModel/RigVMControllerActions.h
@@ -374,44 +374,6 @@ public:
 	FString NodePath;
 };
 
-/**
- * An action adding a parameter node to the graph.
- */
-USTRUCT()
-struct FRigVMAddParameterNodeAction : public FRigVMBaseAction
-{
-	GENERATED_BODY()
-
-public:
-
-	FRigVMAddParameterNodeAction();
-	FRigVMAddParameterNodeAction(URigVMParameterNode* InNode);
-	virtual ~FRigVMAddParameterNodeAction() {};
-	virtual bool Undo(URigVMController* InController) override;
-	virtual bool Redo(URigVMController* InController) override;
-
-	UPROPERTY()
-	FName ParameterName;
-	
-	UPROPERTY()
-	FString CPPType;
-
-	UPROPERTY()
-	FString CPPTypeObjectPath;
-
-	UPROPERTY()
-	bool bIsInput;
-
-	UPROPERTY()
-	FString DefaultValue;
-
-	UPROPERTY()
-	FVector2D Position;
-
-	UPROPERTY()
-	FString NodePath;
-};
-
 /**
  * An action adding a comment node to the graph.
  */
@@ -970,29 +932,6 @@ public:
 	FString NewVariableName;
 };
 
-/**
- * An action renaming a parameter in the graph.
- */
-USTRUCT()
-struct FRigVMRenameParameterAction : public FRigVMBaseAction
-{
-	GENERATED_BODY()
-
-public:
-
-	FRigVMRenameParameterAction() {}
-	FRigVMRenameParameterAction(const FName& InOldParameterName, const FName& InNewParameterName);
-	virtual ~FRigVMRenameParameterAction() {};
-	virtual bool Undo(URigVMController* InController) override;
-	virtual bool Redo(URigVMController* InController) override;
-
-	UPROPERTY()
-	FString OldParameterName;
-
-	UPROPERTY()
-	FString NewParameterName;
-};
-
 /**
  * An action setting a pin's expansion state in the graph.
  */
diff --git a/Engine/Source/Developer/RigVMDeveloper/Public/RigVMModel/RigVMGraph.h b/Engine/Source/Developer/RigVMDeveloper/Public/RigVMModel/RigVMGraph.h
--- a/Engine/Source/Developer/RigVMDeveloper/Public/RigVMModel/RigVMGraph.h
+++ b/Engine/Source/Developer/RigVMDeveloper/Public/RigVMModel/RigVMGraph.h
@@ -69,11 +69,6 @@ public:
 	UFUNCTION(BlueprintCallable, Category = RigVMGraph)
 	TArray<FRigVMGraphVariableDescription> GetVariableDescriptions() const;
 
-	// Returns a list of unique Parameter descriptions within this Graph.
-	// Multiple Parameter Nodes can share the same description.
-	UFUNCTION(BlueprintCallable, Category = RigVMGraph)
-	TArray<FRigVMGraphParameterDescription> GetParameterDescriptions() const;
-
 	// Returns the path of this graph as defined by its invoking nodes
 	UFUNCTION(BlueprintCallable, Category = RigVMGraph)
 	virtual FString GetNodePath() const;
diff --git a/Engine/Source/Developer/RigVMDeveloper/Public/RigVMModel/RigVMNotifications.h b/Engine/Source/Developer/RigVMDeveloper/Public/RigVMModel/RigVMNotifications.h
--- a/Engine/Source/Developer/RigVMDeveloper/Public/RigVMModel/RigVMNotifications.h
+++ b/Engine/Source/Developer/RigVMDeveloper/Public/RigVMModel/RigVMNotifications.h
@@ -41,9 +41,6 @@ enum class ERigVMGraphNotifType : uint8
 	VariableAdded, // A variable has been added (Subject == URigVMVariableNode)
 	VariableRemoved, // A variable has been removed (Subject == URigVMVariableNode)
 	VariableRenamed, // A variable has been renamed (Subject == URigVMVariableNode)
-	ParameterAdded, // A parameter has been added (Subject == URigVMParameterNode)
-	ParameterRemoved, // A parameter has been removed (Subject == URigVMParameterNode)
-	ParameterRenamed, // A parameter has been renamed (Subject == URigVMParameterNode)
 	InteractionBracketOpened, // A bracket has been opened (Subject == nullptr)
 	InteractionBracketClosed, // A bracket has been opened (Subject == nullptr)
 	InteractionBracketCanceled, // A bracket has been canceled (Subject == nullptr)
