RigVM: Node upgrade functionality

As part of debugging the last unit test I've added functionality to support debugging of the ActionStack. The controller is now able to upgrade a node as per user interaction. next up I'll integrate this into the user interface.

diff --git a/Engine/Source/Developer/RigVMDeveloper/Private/RigVMModel/Nodes/RigVMUnitNode.cpp b/Engine/Source/Developer/RigVMDeveloper/Private/RigVMModel/Nodes/RigVMUnitNode.cpp
--- a/Engine/Source/Developer/RigVMDeveloper/Private/RigVMModel/Nodes/RigVMUnitNode.cpp
+++ b/Engine/Source/Developer/RigVMDeveloper/Private/RigVMModel/Nodes/RigVMUnitNode.cpp
@@ -175,3 +175,16 @@ TSharedPtr<FStructOnScope> URigVMUnitNode::ConstructStructInstance(bool bUseDefa
 	}
 	return nullptr;
 }
+
+FRigVMStructUpgradeInfo URigVMUnitNode::GetUpgradeInfo() const
+{
+	if(UScriptStruct* Struct = GetScriptStruct())
+	{
+		check(Struct->IsChildOf(FRigVMStruct::StaticStruct()));
+
+		const TSharedPtr<FStructOnScope> StructOnScope = ConstructStructInstance();
+		const FRigVMStruct* StructMemory = (const FRigVMStruct*)StructOnScope->GetStructMemory();
+		return StructMemory->GetUpgradeInfo();
+	}
+	return FRigVMStructUpgradeInfo();
+}
diff --git a/Engine/Source/Developer/RigVMDeveloper/Private/RigVMModel/RigVMController.cpp b/Engine/Source/Developer/RigVMDeveloper/Private/RigVMModel/RigVMController.cpp
--- a/Engine/Source/Developer/RigVMDeveloper/Private/RigVMModel/RigVMController.cpp
+++ b/Engine/Source/Developer/RigVMDeveloper/Private/RigVMModel/RigVMController.cpp
@@ -1631,6 +1631,254 @@ URigVMTemplateNode* URigVMController::ReplaceUnitNodeWithTemplateNode(const FNam
 	return TemplateNode;
 }
 
+TArray<URigVMNode*> URigVMController::UpgradeNodes(const TArray<FName>& InNodeNames, bool bRecursive, bool bSetupUndoRedo,
+	bool bPrintPythonCommand)
+{
+	TArray<URigVMNode*> Nodes;
+	if (!IsValidGraph())
+	{
+		return Nodes;
+	}
+
+	for (const FName& NodeName : InNodeNames)
+	{
+		if (URigVMNode* Node = GetGraph()->FindNodeByName(NodeName))
+		{
+			Nodes.Add(Node);
+		}
+	}
+
+	Nodes = UpgradeNodes(Nodes, bRecursive, bSetupUndoRedo);
+
+	if(bPrintPythonCommand)
+	{
+		const FString GraphName = GetSanitizedGraphName(GetGraph()->GetGraphName());
+		TArray<FString> NodeNames;
+		for(const FName& NodeName : InNodeNames)
+		{
+			NodeNames.Add(GetSanitizedNodeName(NodeName.ToString()));
+		}
+		const FString NodeNamesJoined = FString::Join(NodeNames, TEXT("','"));
+
+		// UpgradeNodes(const TArray<FName>& InNodeNames)
+		RigVMPythonUtils::Print(GetGraphOuterName(),
+				FString::Printf(TEXT("blueprint.get_controller_by_name('%s').upgrade_nodes(['%s'])"),
+				*GraphName,
+				*NodeNamesJoined));
+	}
+
+	return Nodes;
+}
+
+TArray<URigVMNode*> URigVMController::UpgradeNodes(const TArray<URigVMNode*>& InNodes, bool bRecursive, bool bSetupUndoRedo)
+{
+	if (!IsValidGraph())
+	{
+		return TArray<URigVMNode*>();
+	}
+
+	bool bFoundAnyNodeToUpgrade = false;
+	for(URigVMNode* Node : InNodes)
+	{
+		if(!IsValidNodeForGraph(Node))
+		{
+			return TArray<URigVMNode*>();
+		}
+
+		bFoundAnyNodeToUpgrade |= Node->CanBeUpgraded();
+	}
+
+	if(!bFoundAnyNodeToUpgrade)
+	{
+		return InNodes;
+	}
+
+	FRigVMBaseAction Action;
+	if(bSetupUndoRedo)
+	{
+		Action.Title = TEXT("Upgrade nodes");
+		ActionStack->BeginAction(Action);
+	}
+
+	// find all links affecting the nodes to upgrade
+	TArray<TPair<FString, FString>> LinkedPaths;
+	for(URigVMNode* Node : InNodes)
+	{
+		TArray<URigVMLink*> Links = Node->GetLinks();
+		for(URigVMLink* Link : Links)
+		{
+			TPair<FString, FString> LinkedPath(Link->GetSourcePin()->GetPinPath(), Link->GetTargetPin()->GetPinPath());
+			LinkedPaths.AddUnique(LinkedPath);
+		}
+	}
+
+	for(const TPair<FString, FString>& LinkedPath : LinkedPaths)
+	{
+		if(!BreakLink(LinkedPath.Key, LinkedPath.Value, bSetupUndoRedo))
+		{
+			if(bSetupUndoRedo)
+			{
+				ActionStack->CancelAction(Action);
+			}
+			ReportErrorf(TEXT("Couldn't remove link '%s' -> '%s'"), *LinkedPath.Key, *LinkedPath.Value);
+			return TArray<URigVMNode*>();
+		}
+	}
+
+	TArray<URigVMNode*> UpgradedNodes;
+	TMap<FString,FRigVMController_PinPathRemapDelegate> RemapPinDelegates;
+	for(URigVMNode* Node : InNodes)
+	{
+		FRigVMController_PinPathRemapDelegate RemapPinDelegate;
+		URigVMNode* UpgradedNode = UpgradeNode(Node, bSetupUndoRedo, &RemapPinDelegate);
+		UpgradedNodes.Add(UpgradedNode);
+		if(RemapPinDelegate.IsBound())
+		{
+			RemapPinDelegates.Add(UpgradedNode->GetName(), RemapPinDelegate);
+		}
+	}
+
+	for(const TPair<FString, FString>& LinkedPath : LinkedPaths)
+	{
+		FString SourcePath = LinkedPath.Key;
+		FString TargetPath = LinkedPath.Value;
+
+		FString SourceNodeName, SourceSegmentPath;
+		if(!URigVMPin::SplitPinPathAtStart(SourcePath, SourceNodeName, SourceSegmentPath))
+		{
+			continue;
+		}
+		
+		FString TargetNodeName, TargetSegmentPath;
+		if(!URigVMPin::SplitPinPathAtStart(TargetPath, TargetNodeName, TargetSegmentPath))
+		{
+			continue;
+		}
+
+		if(FRigVMController_PinPathRemapDelegate* SourceRemapDelegate = RemapPinDelegates.Find(SourceNodeName))
+		{
+			SourcePath = SourceRemapDelegate->Execute(SourcePath, false);
+		}
+		if(FRigVMController_PinPathRemapDelegate* TargetRemapDelegate = RemapPinDelegates.Find(TargetNodeName))
+		{
+			TargetPath = TargetRemapDelegate->Execute(TargetPath, true);
+		}
+		
+		if(!AddLink(SourcePath, TargetPath, bSetupUndoRedo))
+		{
+			ReportWarningf(TEXT("Couldn't recreate link '%s' -> '%s'"), *SourcePath, *TargetPath);
+		}
+	}
+
+	if(bRecursive)
+	{
+		UpgradedNodes = UpgradeNodes(UpgradedNodes, bRecursive, bSetupUndoRedo);
+	}
+
+	if(bSetupUndoRedo)
+	{
+		ActionStack->EndAction(Action);
+	}
+
+	return UpgradedNodes;
+}
+
+URigVMNode* URigVMController::UpgradeNode(URigVMNode* InNode, bool bSetupUndoRedo, FRigVMController_PinPathRemapDelegate* OutRemapPinDelegate)
+{
+	if(!IsValidNodeForGraph(InNode))
+	{
+		return nullptr;
+	}
+
+	if(!InNode->CanBeUpgraded())
+	{
+		return InNode; 
+	}
+
+	TMap<FString, FPinState> PinStates = GetPinStates(InNode);
+
+	const FString NodeName = InNode->GetName();
+	const FVector2D NodePosition = InNode->GetPosition();
+
+	FRigVMBaseAction Action;
+	if(bSetupUndoRedo)
+	{
+		Action.Title = TEXT("Upgrade node");
+		ActionStack->BeginAction(Action);
+	}
+
+	URigVMNode* UpgradedNode = nullptr;
+
+	if(URigVMUnitNode* UnitNode = Cast<URigVMUnitNode>(InNode))
+	{
+		const FName MethodName = UnitNode->GetMethodName();
+		
+		const FRigVMStructUpgradeInfo UpgradeInfo = UnitNode->GetUpgradeInfo();
+		check(UpgradeInfo.IsValid());
+
+		if(OutRemapPinDelegate)
+		{
+			*OutRemapPinDelegate = FRigVMController_PinPathRemapDelegate::CreateLambda([UpgradeInfo](const FString& InPinPath, bool bIsInput) -> FString
+			{
+				return UpgradeInfo.RemapPin(InPinPath, bIsInput, true);
+			});
+		}
+
+		if(!RemoveNode(InNode, bSetupUndoRedo, true, false, false))
+		{
+			if(bSetupUndoRedo)
+			{
+				ActionStack->CancelAction(Action);
+			}
+			ReportErrorf(TEXT("Unable to remove node %s."), *NodeName);
+			return nullptr;
+		}
+
+		URigVMUnitNode* NewUnitNode = AddUnitNode(UpgradeInfo.GetNewStruct(), MethodName, NodePosition, NodeName, bSetupUndoRedo, false);
+		if(NewUnitNode == nullptr)
+		{
+			if(bSetupUndoRedo)
+			{
+				ActionStack->CancelAction(Action);
+			}
+			ReportErrorf(TEXT("Unable to upgrade node %s."), *NodeName);
+			return nullptr;
+		}
+
+		for(URigVMPin* Pin : NewUnitNode->GetPins())
+		{
+			const FString DefaultValue = UpgradeInfo.GetDefaultValueForPin(Pin->GetFName());
+			if(!DefaultValue.IsEmpty())
+			{
+				SetPinDefaultValue(Pin, DefaultValue, true, false, false);
+			}
+		}
+
+		UpgradedNode = NewUnitNode;
+	}
+	else
+	{
+		// for now we don't allow to upgrade anything else but unit nodes
+		checkNoEntry();
+	}
+
+	check(UpgradedNode);
+
+	// reapply the pin states but don't touch defaults
+	for(TPair<FString, FPinState>& PinState : PinStates)
+	{
+		PinState.Value.DefaultValue.Reset();
+	}
+	ApplyPinStates(UpgradedNode, PinStates);
+
+	if(bSetupUndoRedo)
+	{
+		ActionStack->EndAction(Action);
+	}
+
+	return UpgradedNode;
+}
+
 URigVMParameterNode* URigVMController::AddParameterNode(const FName& InParameterName, const FString& InCPPType, UObject* InCPPTypeObject, bool bIsInput, const FString& InDefaultValue, const FVector2D& InPosition, const FString& InNodeName, bool bSetupUndoRedo, bool bPrintPythonCommand)
 {
 	AddVariableNode(InParameterName, InCPPType, InCPPTypeObject, bIsInput, InDefaultValue, InPosition, InNodeName, bSetupUndoRedo, bPrintPythonCommand);
@@ -7767,6 +8015,11 @@ bool URigVMController::BreakAllLinks(const FString& InPinPath, bool bAsInput, bo
 
 bool URigVMController::BreakAllLinks(URigVMPin* Pin, bool bAsInput, bool bSetupUndoRedo)
 {
+	if(!Pin->IsLinked(false))
+	{
+		return false;
+	}
+	
 	FRigVMControllerCompileBracketScope CompileScope(this);
 	FRigVMBaseAction Action;
 	if (bSetupUndoRedo)
@@ -12307,6 +12560,20 @@ void URigVMController::ApplyPinStates(URigVMNode* InNode, const TMap<FString, UR
 	}
 }
 
+void URigVMController::ReportInfo(const FString& InMessage) const
+{
+	if (URigVMGraph* Graph = GetGraph())
+	{
+		if (UPackage* Package = Cast<UPackage>(Graph->GetOutermost()))
+		{
+			UE_LOG(LogRigVMDeveloper, Display, TEXT("%s : %s"), *Package->GetPathName(), *InMessage);
+			return;
+		}
+	}
+
+	UE_LOG(LogRigVMDeveloper, Display, TEXT("%s"), *InMessage);
+}
+
 void URigVMController::ReportWarning(const FString& InMessage) const
 {
 	if(!bReportWarningsAndErrors)
diff --git a/Engine/Source/Developer/RigVMDeveloper/Private/RigVMModel/RigVMControllerActions.cpp b/Engine/Source/Developer/RigVMDeveloper/Private/RigVMModel/RigVMControllerActions.cpp
--- a/Engine/Source/Developer/RigVMDeveloper/Private/RigVMModel/RigVMControllerActions.cpp
+++ b/Engine/Source/Developer/RigVMDeveloper/Private/RigVMModel/RigVMControllerActions.cpp
@@ -7,36 +7,45 @@
 #include "Misc/ITransaction.h"
 #endif
 
+UScriptStruct* FRigVMActionKey::GetScriptStruct() const
+{
+	return FindObjectChecked<UScriptStruct>(ANY_PACKAGE, *ScriptStructPath);		
+}
+
+TSharedPtr<FStructOnScope> FRigVMActionKey::GetAction() const
+{
+	UScriptStruct* ScriptStruct = GetScriptStruct();
+	TSharedPtr<FStructOnScope> StructOnScope = MakeShareable(new FStructOnScope(ScriptStruct));
+	ScriptStruct->ImportText(*ExportedText, StructOnScope->GetStructMemory(), nullptr, PPF_None, nullptr, ScriptStruct->GetName());
+	return StructOnScope;
+}
+
 FRigVMActionWrapper::FRigVMActionWrapper(const FRigVMActionKey& Key)
 {
-	ScriptStruct = FindObjectChecked<UScriptStruct>(ANY_PACKAGE, *Key.ScriptStructPath);
-	Data.SetNumUninitialized(ScriptStruct->GetStructureSize());
-	ScriptStruct->InitializeStruct(Data.GetData(), 1);
-	ScriptStruct->ImportText(*Key.ExportedText, Data.GetData(), nullptr, PPF_None, nullptr, ScriptStruct->GetName());
+	StructOnScope = Key.GetAction();
 }
 FRigVMActionWrapper::~FRigVMActionWrapper()
 {
-	if (Data.Num() > 0 && ScriptStruct != nullptr)
-	{
-		ScriptStruct->DestroyStruct(Data.GetData(), 1);
-	}
 }
 
-FRigVMBaseAction* FRigVMActionWrapper::GetAction()
+const UScriptStruct* FRigVMActionWrapper::GetScriptStruct() const
+{
+	return CastChecked<UScriptStruct>(StructOnScope->GetStruct());
+}
+
+FRigVMBaseAction* FRigVMActionWrapper::GetAction() const
 {
-	return (FRigVMBaseAction*)Data.GetData();
+	return (FRigVMBaseAction*)StructOnScope->GetStructMemory();
 }
 
-FString FRigVMActionWrapper::ExportText()
+FString FRigVMActionWrapper::ExportText() const
 {
 	FString ExportedText;
-	if (Data.Num() > 0 && ScriptStruct != nullptr)
+	if (StructOnScope.IsValid() && StructOnScope->IsValid())
 	{
-		TArray<uint8, TAlignedHeapAllocator<16>> DefaultStructData;
-		DefaultStructData.AddZeroed(ScriptStruct->GetStructureSize());
-		ScriptStruct->InitializeDefaultValue(DefaultStructData.GetData());
-		
-		ScriptStruct->ExportText(ExportedText, GetAction(), DefaultStructData.GetData(), nullptr, PPF_None, nullptr);
+		const UScriptStruct* ScriptStruct = GetScriptStruct();
+		FStructOnScope DefaultScope(ScriptStruct);
+		ScriptStruct->ExportText(ExportedText, GetAction(), DefaultScope.GetStructMemory(), nullptr, PPF_None, nullptr);
 	}
 	return ExportedText;
 }
@@ -86,6 +95,12 @@ bool URigVMActionStack::Undo(URigVMController* InController)
 	ActionIndex = UndoActions.Num();
 	
 	FRigVMActionWrapper Wrapper(KeyToUndo);
+
+#if RIGVM_ACTIONSTACK_VERBOSE_LOG		
+	TGuardValue<int32> TabDepthGuard(LogActionDepth, 0);
+	LogAction(Wrapper.GetScriptStruct(), *Wrapper.GetAction(), FRigVMBaseAction::UndoPrefix);
+#endif
+	
 	if (Wrapper.GetAction()->Undo(InController))
 	{
 		RedoActions.Add(KeyToUndo);
@@ -109,6 +124,12 @@ bool URigVMActionStack::Redo(URigVMController* InController)
 	FRigVMActionKey KeyToRedo = RedoActions.Pop();
 
 	FRigVMActionWrapper Wrapper(KeyToRedo);
+
+#if RIGVM_ACTIONSTACK_VERBOSE_LOG		
+	TGuardValue<int32> TabDepthGuard(LogActionDepth, 0);
+	LogAction(Wrapper.GetScriptStruct(), *Wrapper.GetAction(), FRigVMBaseAction::RedoPrefix);
+#endif
+	
 	if (Wrapper.GetAction()->Redo(InController))
 	{
 		UndoActions.Add(KeyToRedo);
@@ -171,6 +192,47 @@ void URigVMActionStack::PostTransacted(const FTransactionObjectEvent& Transactio
 	}
 }
 
+#if RIGVM_ACTIONSTACK_VERBOSE_LOG		
+
+void URigVMActionStack::LogAction(const UScriptStruct* InActionStruct, const FRigVMBaseAction& InAction, const FString& InPrefix)
+{
+	if(bSuspendLogActions)
+	{
+		return;
+	}
+	
+	if(URigVMController* Controller = GetTypedOuter<URigVMController>())
+	{
+		static TArray<FString> TabPrefix = {TEXT(""), TEXT("  ")};
+
+		while(TabPrefix.Num() <= LogActionDepth)
+		{
+			TabPrefix.Add(TabPrefix.Last() + TabPrefix[1]);
+		}
+		
+		const FString ActionContent = FRigVMStruct::ExportToFullyQualifiedText(InActionStruct, (const uint8*)&InAction);
+		Controller->ReportInfof(TEXT("%s%s: %s (%s) '%s"), *TabPrefix[LogActionDepth], *InPrefix, *InAction.GetTitle(), *InActionStruct->GetStructCPPName(), *ActionContent);
+	}
+
+	if(InPrefix == FRigVMBaseAction::AddActionPrefix)
+	{
+		TGuardValue<int32> ActionDepthGuard(LogActionDepth, LogActionDepth + 1);
+		for(const FRigVMActionKey& SubAction : InAction.SubActions)
+		{
+			const FRigVMActionWrapper Wrapper(SubAction);
+			LogAction(Wrapper.GetScriptStruct(), *Wrapper.GetAction(), InPrefix);
+		}
+		if(InActionStruct == FRigVMRemoveNodeAction::StaticStruct())
+		{
+			const FRigVMRemoveNodeAction* RemoveNodeAction = (const FRigVMRemoveNodeAction*)&InAction; 
+			const FRigVMActionWrapper Wrapper(RemoveNodeAction->InverseActionKey);
+			LogAction(Wrapper.GetScriptStruct(), *Wrapper.GetAction(), InPrefix);
+		}
+	}
+}
+
+#endif
+
 #endif
 
 
@@ -181,10 +243,19 @@ bool FRigVMBaseAction::Merge(const FRigVMBaseAction* Other)
 
 bool FRigVMBaseAction::Undo(URigVMController* InController)
 {
+#if RIGVM_ACTIONSTACK_VERBOSE_LOG		
+	URigVMActionStack* Stack = InController->ActionStack;
+	check(Stack);
+	TGuardValue<int32> TabDepthGuard(Stack->LogActionDepth, Stack->LogActionDepth + 1);
+#endif
+
 	bool Result = true;
 	for (int32 KeyIndex = SubActions.Num() - 1; KeyIndex >= 0; KeyIndex--)
 	{
 		FRigVMActionWrapper Wrapper(SubActions[KeyIndex]);
+#if RIGVM_ACTIONSTACK_VERBOSE_LOG		
+		Wrapper.GetAction()->LogAction(InController, UndoPrefix);
+#endif
 		if(!Wrapper.GetAction()->Undo(InController))
 		{
 			InController->ReportAndNotifyErrorf(TEXT("Error while undoing action '%s'."), *Wrapper.GetAction()->Title);
@@ -196,10 +267,19 @@ bool FRigVMBaseAction::Undo(URigVMController* InController)
 
 bool FRigVMBaseAction::Redo(URigVMController* InController)
 {
+#if RIGVM_ACTIONSTACK_VERBOSE_LOG		
+	URigVMActionStack* Stack = InController->ActionStack;
+	check(Stack);
+	TGuardValue<int32> TabDepthGuard(Stack->LogActionDepth, Stack->LogActionDepth + 1);
+#endif
+
 	bool Result = true;
 	for (int32 KeyIndex = 0; KeyIndex < SubActions.Num(); KeyIndex++)
 	{
 		FRigVMActionWrapper Wrapper(SubActions[KeyIndex]);
+#if RIGVM_ACTIONSTACK_VERBOSE_LOG		
+		Wrapper.GetAction()->LogAction(InController, RedoPrefix);
+#endif
 		if (!Wrapper.GetAction()->Redo(InController))
 		{
 			InController->ReportAndNotifyErrorf(TEXT("Error while redoing action '%s'."), *Wrapper.GetAction()->Title);
@@ -209,6 +289,20 @@ bool FRigVMBaseAction::Redo(URigVMController* InController)
 	return Result;
 }
 
+#if RIGVM_ACTIONSTACK_VERBOSE_LOG
+
+void FRigVMBaseAction::LogAction(URigVMController* InController, const FString& InPrefix) const
+{
+	check(InController);
+	URigVMActionStack* Stack = InController->ActionStack;
+	check(Stack);
+	TGuardValue<int32> TabDepthGuard(Stack->LogActionDepth, Stack->LogActionDepth + 1);
+
+	Stack->LogAction(GetScriptStruct(), *this, InPrefix);
+}
+
+#endif
+
 bool FRigVMInverseAction::Undo(URigVMController* InController)
 {
 	return FRigVMBaseAction::Redo(InController);
@@ -626,65 +720,66 @@ FRigVMRemoveNodeAction::FRigVMRemoveNodeAction(URigVMNode* InNode, URigVMControl
 
 	if (URigVMUnitNode* UnitNode = Cast<URigVMUnitNode>(InNode))
 	{
-		InverseAction.AddAction(FRigVMAddUnitNodeAction(UnitNode));
+		InverseAction.AddAction(FRigVMAddUnitNodeAction(UnitNode), InController);
 		for (URigVMPin* Pin : UnitNode->GetPins())
 		{
 			if (Pin->GetDirection() == ERigVMPinDirection::Input ||
+				Pin->GetDirection() == ERigVMPinDirection::IO ||
 				Pin->GetDirection() == ERigVMPinDirection::Visible)
 			{
-				InverseAction.AddAction(FRigVMSetPinDefaultValueAction(Pin, Pin->GetDefaultValue()));
+				InverseAction.AddAction(FRigVMSetPinDefaultValueAction(Pin, Pin->GetDefaultValue()), InController);
 			}
 		}
 	}
 	else if (URigVMVariableNode* VariableNode = Cast<URigVMVariableNode>(InNode))
 	{
-		InverseAction.AddAction(FRigVMAddVariableNodeAction(VariableNode));
+		InverseAction.AddAction(FRigVMAddVariableNodeAction(VariableNode), InController);
 		URigVMPin* ValuePin = VariableNode->FindPin(TEXT("Value"));
-		InverseAction.AddAction(FRigVMSetPinDefaultValueAction(ValuePin, ValuePin->GetDefaultValue()));
+		InverseAction.AddAction(FRigVMSetPinDefaultValueAction(ValuePin, ValuePin->GetDefaultValue()), InController);
 	}
 	else if (URigVMCommentNode* CommentNode = Cast<URigVMCommentNode>(InNode))
 	{
-		InverseAction.AddAction(FRigVMAddCommentNodeAction(CommentNode));
+		InverseAction.AddAction(FRigVMAddCommentNodeAction(CommentNode), InController);
 	}
 	else if (URigVMRerouteNode* RerouteNode = Cast<URigVMRerouteNode>(InNode))
 	{
-		InverseAction.AddAction(FRigVMAddRerouteNodeAction(RerouteNode));
+		InverseAction.AddAction(FRigVMAddRerouteNodeAction(RerouteNode), InController);
 	}
 	else if (URigVMBranchNode* BranchNode = Cast<URigVMBranchNode>(InNode))
 	{
-		InverseAction.AddAction(FRigVMAddBranchNodeAction(BranchNode));
+		InverseAction.AddAction(FRigVMAddBranchNodeAction(BranchNode), InController);
 	}
 	else if (URigVMIfNode* IfNode = Cast<URigVMIfNode>(InNode))
 	{
-		InverseAction.AddAction(FRigVMAddIfNodeAction(IfNode));
+		InverseAction.AddAction(FRigVMAddIfNodeAction(IfNode), InController);
 	}
 	else if (URigVMSelectNode* SelectNode = Cast<URigVMSelectNode>(InNode))
 	{
-		InverseAction.AddAction(FRigVMAddSelectNodeAction(SelectNode));
+		InverseAction.AddAction(FRigVMAddSelectNodeAction(SelectNode), InController);
 	}
 	else if (URigVMTemplateNode* TemplateNode = Cast<URigVMTemplateNode>(InNode))
 	{
-		InverseAction.AddAction(FRigVMAddTemplateNodeAction(TemplateNode));
+		InverseAction.AddAction(FRigVMAddTemplateNodeAction(TemplateNode), InController);
 	}
 	else if (URigVMEnumNode* EnumNode = Cast<URigVMEnumNode>(InNode))
 	{
-		InverseAction.AddAction(FRigVMAddEnumNodeAction(EnumNode));
+		InverseAction.AddAction(FRigVMAddEnumNodeAction(EnumNode), InController);
 	}
 	else if (URigVMArrayNode* ArrayNode = Cast<URigVMArrayNode>(InNode))
 	{
-		InverseAction.AddAction(FRigVMAddArrayNodeAction(ArrayNode));
+		InverseAction.AddAction(FRigVMAddArrayNodeAction(ArrayNode), InController);
 		for (URigVMPin* Pin : ArrayNode->GetPins())
 		{
 			if (Pin->GetDirection() == ERigVMPinDirection::Input ||
 				Pin->GetDirection() == ERigVMPinDirection::Visible)
 			{
-				InverseAction.AddAction(FRigVMSetPinDefaultValueAction(Pin, Pin->GetDefaultValue()));
+				InverseAction.AddAction(FRigVMSetPinDefaultValueAction(Pin, Pin->GetDefaultValue()), InController);
 			}
 		}
 	}
 	else if (URigVMLibraryNode* LibraryNode = Cast<URigVMLibraryNode>(InNode))
 	{
-		InverseAction.AddAction(FRigVMImportNodeFromTextAction(LibraryNode, InController));
+		InverseAction.AddAction(FRigVMImportNodeFromTextAction(LibraryNode, InController), InController);
 	}
 	else if (InNode->IsA<URigVMFunctionEntryNode>() || InNode->IsA<URigVMFunctionReturnNode>())
 	{
@@ -701,7 +796,7 @@ FRigVMRemoveNodeAction::FRigVMRemoveNodeAction(URigVMNode* InNode, URigVMControl
 		{
 			FRigVMSetPinExpansionAction ExpansionAction(Pin, true);
 			ExpansionAction.OldIsExpanded = false;
-			InverseAction.AddAction(ExpansionAction);
+			InverseAction.AddAction(ExpansionAction, InController);
 		}
 
 		if (Pin->HasInjectedNodes())
@@ -711,9 +806,9 @@ FRigVMRemoveNodeAction::FRigVMRemoveNodeAction(URigVMNode* InNode, URigVMControl
 				FRigVMAddVariableNodeAction AddVariableNodeAction(VariableNode);
 				FRigVMAddLinkAction AddLinkAction(VariableNode->GetValuePin(), Pin);
 				FRigVMInjectNodeIntoPinAction InjectAction(Pin->GetInjectedNodes()[0]);
-				InverseAction.AddAction(AddVariableNodeAction);
-				InverseAction.AddAction(AddLinkAction);
-				InverseAction.AddAction(InjectAction);
+				InverseAction.AddAction(AddVariableNodeAction, InController);
+				InverseAction.AddAction(AddLinkAction, InController);
+				InverseAction.AddAction(InjectAction, InController);
 			}
 		}
 	}
diff --git a/Engine/Source/Developer/RigVMDeveloper/Private/RigVMModel/RigVMPin.cpp b/Engine/Source/Developer/RigVMDeveloper/Private/RigVMModel/RigVMPin.cpp
--- a/Engine/Source/Developer/RigVMDeveloper/Private/RigVMModel/RigVMPin.cpp
+++ b/Engine/Source/Developer/RigVMDeveloper/Private/RigVMModel/RigVMPin.cpp
@@ -1305,6 +1305,27 @@ bool URigVMPin::IsLinkedTo(URigVMPin* InPin) const
 	return false;
 }
 
+bool URigVMPin::IsLinked(bool bRecursive) const
+{
+	if(!GetLinks().IsEmpty())
+	{
+		return true;
+	}
+
+	if(bRecursive)
+	{
+		for (const URigVMPin* SubPin : SubPins)
+		{
+			if(SubPin->IsLinked(true))
+			{
+				return true;
+			}
+		}
+	}
+
+	return false;
+}
+
 const TArray<URigVMLink*>& URigVMPin::GetLinks() const
 {
 	return Links;
diff --git a/Engine/Source/Developer/RigVMDeveloper/Public/RigVMModel/Nodes/RigVMUnitNode.h b/Engine/Source/Developer/RigVMDeveloper/Public/RigVMModel/Nodes/RigVMUnitNode.h
--- a/Engine/Source/Developer/RigVMDeveloper/Public/RigVMModel/Nodes/RigVMUnitNode.h
+++ b/Engine/Source/Developer/RigVMDeveloper/Public/RigVMModel/Nodes/RigVMUnitNode.h
@@ -26,6 +26,7 @@ public:
 	virtual bool IsDefinedAsVarying() const override;
 	virtual FName GetEventName() const override;
 	virtual bool IsLoopNode() const override;
+	virtual bool CanBeUpgraded() const override { return GetUpgradeInfo().IsValid(); }
 
 	bool IsDeprecated() const;
 	FString GetDeprecatedMetadata() const;
@@ -47,6 +48,8 @@ public:
 	// @param bUseDefault If set to true the default struct will be created - otherwise the struct will contains the values from the node
 	TSharedPtr<FStructOnScope> ConstructStructInstance(bool bUseDefault = false) const;
 
+	FRigVMStructUpgradeInfo GetUpgradeInfo() const;
+
 protected:
 
 	virtual FText GetToolTipTextForPin(const URigVMPin* InPin) const override;
diff --git a/Engine/Source/Developer/RigVMDeveloper/Public/RigVMModel/RigVMController.h b/Engine/Source/Developer/RigVMDeveloper/Public/RigVMModel/RigVMController.h
--- a/Engine/Source/Developer/RigVMDeveloper/Public/RigVMModel/RigVMController.h
+++ b/Engine/Source/Developer/RigVMDeveloper/Public/RigVMModel/RigVMController.h
@@ -85,6 +85,7 @@ DECLARE_DELEGATE_RetVal_ThreeParams(FName, FRigVMController_RequestNewExternalVa
 DECLARE_DELEGATE_RetVal_TwoParams(bool, FRigVMController_IsDependencyCyclicDelegate, UObject*, UObject*)
 DECLARE_DELEGATE_RetVal_TwoParams(FRigVMController_BulkEditResult, FRigVMController_RequestBulkEditDialogDelegate, URigVMLibraryNode*, ERigVMControllerBulkEditType)
 DECLARE_DELEGATE_FiveParams(FRigVMController_OnBulkEditProgressDelegate, TSoftObjectPtr<URigVMFunctionReferenceNode>, ERigVMControllerBulkEditType, ERigVMControllerBulkEditProgress, int32, int32)
+DECLARE_DELEGATE_RetVal_TwoParams(FString, FRigVMController_PinPathRemapDelegate, const FString& /* InPinPath */, bool /* bIsInput */);
 
 /**
  * The Controller is the sole authority to perform changes
@@ -217,6 +218,10 @@ public:
 	UFUNCTION(BlueprintCallable, Category = RigVMController)
 	URigVMTemplateNode* ReplaceUnitNodeWithTemplateNode(const FName& InNodeName, bool bSetupUndoRedo);
 
+	// Upgrades a set of nodes with each corresponding next known version
+	UFUNCTION(BlueprintCallable, Category = RigVMController)
+	TArray<URigVMNode*> UpgradeNodes(const TArray<FName>& InNodeNames, bool bRecursive = true, bool bSetupUndoRedo = true, bool bPrintPythonCommand = false);
+
 	// Adds a Parameter Node to the edited Graph.
 	// Parameters represent input or output arguments to the Graph / Function.
 	// Input Parameters are constant values / literals.
@@ -533,7 +538,13 @@ public:
 	// Renames a parameter in the graph.
 	// This causes a ParameterRenamed modified event.
 	UFUNCTION(BlueprintCallable, Category = RigVMController, meta=(DeprecatedFunction))
-	bool RenameParameter(const FName& InOldName, const FName& InNewName, bool bSetupUndoRedo = true);
+	bool RenameParameter(const FName& InOldName, const FName& InNewName, bool bSetupUndoRedo = true);\
+
+	// Upgrades a set of nodes with each corresponding next known version
+	TArray<URigVMNode*> UpgradeNodes(const TArray<URigVMNode*>& InNodes, bool bRecursive = true, bool bSetupUndoRedo = true);
+
+	// Upgrades a single node with its next known version
+	URigVMNode* UpgradeNode(URigVMNode* InNode, bool bSetupUndoRedo = true, FRigVMController_PinPathRemapDelegate* OutRemapPinDelegate = nullptr);
 
 	// Sets the pin to be expanded or not
 	// This causes a PinExpansionChanged modified event.
@@ -761,10 +772,17 @@ public:
 
 	const FRigVMByteCode* GetCurrentByteCode() const;
 
+	void ReportInfo(const FString& InMessage) const;
 	void ReportWarning(const FString& InMessage) const;
 	void ReportError(const FString& InMessage) const;
 	void ReportAndNotifyError(const FString& InMessage) const;
 
+	template <typename FmtType, typename... Types>
+	void ReportInfof(const FmtType& Fmt, Types... Args)
+	{
+		ReportInfo(FString::Printf(Fmt, Args...));
+	}
+
 	template <typename FmtType, typename... Types>
 	void ReportWarningf(const FmtType& Fmt, Types... Args)
 	{
@@ -984,6 +1002,7 @@ private:
 	friend class URigVMGraph;
 	friend class URigVMPin;
 	friend class URigVMActionStack;
+	friend struct FRigVMBaseAction;
 	friend class URigVMCompiler;
 	friend struct FRigVMControllerObjectFactory;
 	friend struct FRigVMAddRerouteNodeAction;
diff --git a/Engine/Source/Developer/RigVMDeveloper/Public/RigVMModel/RigVMControllerActions.h b/Engine/Source/Developer/RigVMDeveloper/Public/RigVMModel/RigVMControllerActions.h
--- a/Engine/Source/Developer/RigVMDeveloper/Public/RigVMModel/RigVMControllerActions.h
+++ b/Engine/Source/Developer/RigVMDeveloper/Public/RigVMModel/RigVMControllerActions.h
@@ -7,8 +7,15 @@
 #if WITH_EDITOR
 #include "ScopedTransaction.h"
 #endif
+#include "UObject/StructOnScope.h"
 #include "RigVMControllerActions.generated.h"
 
+// If this is set to one the action stack will log any created action
+// and invoked undo / redo on those actions and print them as an indented log.
+#ifndef RIGVM_ACTIONSTACK_VERBOSE_LOG
+#define RIGVM_ACTIONSTACK_VERBOSE_LOG 0
+#endif 
+
 /**
  * ================================================================================
  * The RigVMController doesn't rely on the transaction system for performing,
@@ -44,6 +51,9 @@ struct FRigVMActionKey
 	UPROPERTY()
 	FString ExportedText;
 
+	UScriptStruct* GetScriptStruct() const;
+	TSharedPtr<FStructOnScope> GetAction() const;
+
 	template<class ActionType>
 	void Set(const ActionType& InAction)
 	{
@@ -68,15 +78,15 @@ public:
 	FRigVMActionWrapper(const FRigVMActionKey& Key);
 	~FRigVMActionWrapper();
 
-	FRigVMBaseAction* GetAction();
-	FString ExportText();
+	const UScriptStruct* GetScriptStruct() const;
+	FRigVMBaseAction* GetAction() const;
+	FString ExportText() const;
 
 private:
 	FRigVMActionWrapper(const FRigVMActionWrapper& Other) = delete;
 	FRigVMActionWrapper& operator = (const FRigVMActionWrapper& Other) = delete;
 
-	UScriptStruct* ScriptStruct;
-	TArray<uint8> Data;
+	TSharedPtr<FStructOnScope> StructOnScope;
 };
 
 /**
@@ -89,6 +99,13 @@ struct FRigVMBaseAction
 {
 	GENERATED_BODY()
 
+	inline static const FString RedoPrefix = TEXT("Redo");
+	inline static const FString UndoPrefix = TEXT("Undo");
+	inline static const FString AddActionPrefix = TEXT("Add Action");
+	inline static const FString BeginActionPrefix = TEXT("Begin Action");
+	inline static const FString EndActionPrefix = TEXT("End Action");
+	inline static const FString CancelActionPrefix = TEXT("Cancel Action");
+
 	// Default constructor
 	FRigVMBaseAction()
 		 : Title(TEXT("Action"))
@@ -98,6 +115,9 @@ struct FRigVMBaseAction
 	// Default destructor
 	virtual ~FRigVMBaseAction() {};
 
+	// Access to the actions script struct. this needs to be overloaded
+	virtual UScriptStruct* GetScriptStruct() const { return FRigVMBaseAction::StaticStruct(); }
+
 	// Returns true if this action is empty has no effect
 	virtual bool IsEmpty() const { return SubActions.IsEmpty(); }
 
@@ -116,13 +136,20 @@ struct FRigVMBaseAction
 
 	// Adds a child / sub action to this one
 	template<class ActionType>
-	void AddAction(const ActionType& InAction)
+	void AddAction(const ActionType& InAction, URigVMController* InController)
 	{
 		FRigVMActionKey Key;
 		Key.Set<ActionType>(InAction);
 		SubActions.Add(Key);
 	}
 
+#if RIGVM_ACTIONSTACK_VERBOSE_LOG
+	
+	// Logs the action to the console / output log
+	void LogAction(URigVMController* InController, const FString& InPrefix) const;
+
+#endif
+
 	UPROPERTY()
 	FString Title;
 
@@ -151,16 +178,21 @@ class URigVMActionStack : public UObject
 public:
 
 	// Begins an action and opens a bracket / scope.
-	template<class T>
-	void BeginAction(T& InAction)
+	template<class ActionType>
+	void BeginAction(ActionType& InAction)
 	{
+#if RIGVM_ACTIONSTACK_VERBOSE_LOG		
+		TGuardValue<int32> TabDepthGuard(LogActionDepth, CurrentActions.Num());
+		LogAction<ActionType>(InAction, FRigVMBaseAction::BeginActionPrefix);
+#endif
+
 		if (CurrentActions.Num() > 0)
 		{
 			// catch erroreous duplicate calls to begin action
 			ensure(CurrentActions.Last() != &InAction);
 		}
 		CurrentActions.Add((FRigVMBaseAction*)&InAction);
-
+		
 		ModifiedEvent.Broadcast(ERigVMGraphNotifType::InteractionBracketOpened, nullptr, nullptr);
 	}
 
@@ -171,8 +203,18 @@ public:
 		ensure(CurrentActions.Num() > 0);
 		ensure((FRigVMBaseAction*)&InAction == CurrentActions.Last());
 		CurrentActions.Pop();
-		AddAction(InAction, bPerformMerge);
 
+		{
+#if RIGVM_ACTIONSTACK_VERBOSE_LOG		
+			TGuardValue<bool> LogGuard(bSuspendLogActions, true);
+#endif
+			AddAction(InAction, bPerformMerge);
+		}
+
+#if RIGVM_ACTIONSTACK_VERBOSE_LOG		
+		TGuardValue<int32> TabDepthGuard(LogActionDepth, CurrentActions.Num());
+		LogAction<ActionType>(InAction, FRigVMBaseAction::EndActionPrefix);
+#endif
 		ModifiedEvent.Broadcast(ERigVMGraphNotifType::InteractionBracketClosed, nullptr, nullptr);
 	}
 
@@ -185,6 +227,10 @@ public:
 		ensure((FRigVMBaseAction*)&InAction == CurrentActions.Last());
 		CurrentActions.Pop();
 
+#if RIGVM_ACTIONSTACK_VERBOSE_LOG		
+		TGuardValue<int32> TabDepthGuard(LogActionDepth, CurrentActions.Num());
+		LogAction<ActionType>(InAction, FRigVMBaseAction::CancelActionPrefix);
+#endif
 		ModifiedEvent.Broadcast(ERigVMGraphNotifType::InteractionBracketCanceled, nullptr, nullptr);
 	}
 
@@ -194,6 +240,10 @@ public:
 	template<class ActionType>
 	void AddAction(const ActionType& InAction, bool bPerformMerge = false)
 	{
+#if RIGVM_ACTIONSTACK_VERBOSE_LOG		
+		TGuardValue<int32> TabDepthGuard(LogActionDepth, CurrentActions.Num());
+		LogAction<ActionType>(InAction, FRigVMBaseAction::AddActionPrefix);
+#endif
 		TArray<FRigVMActionKey>* ActionList = &UndoActions;
 		if (CurrentActions.Num() > 0)
 		{
@@ -278,6 +328,16 @@ public:
 
 private:
 
+#if RIGVM_ACTIONSTACK_VERBOSE_LOG		
+	template<class ActionType>
+	FORCEINLINE void LogAction(const ActionType& InAction, const FString& InPrefix)
+	{
+		LogAction(ActionType::StaticStruct(), InAction, InPrefix);
+	}
+	
+	void LogAction(const UScriptStruct* InActionStruct, const FRigVMBaseAction& InAction, const FString& InPrefix);
+#endif
+	
 	UPROPERTY()
 	int32 ActionIndex;
 
@@ -291,6 +351,14 @@ private:
 	TArray<FRigVMBaseAction*> BracketActions;
 
 	FRigVMGraphModifiedEvent ModifiedEvent;
+
+#if RIGVM_ACTIONSTACK_VERBOSE_LOG		
+	bool bSuspendLogActions = false;
+	int32 LogActionDepth = 0;
+#endif
+	
+	friend struct FRigVMBaseAction;
+	friend class URigVMController;
 };
 
 /**
@@ -303,6 +371,8 @@ struct FRigVMInverseAction : public FRigVMBaseAction
 	GENERATED_BODY()
 
 	virtual ~FRigVMInverseAction() {};
+	virtual UScriptStruct* GetScriptStruct() const override { return FRigVMInverseAction::StaticStruct(); }
+
 	virtual bool Undo(URigVMController* InController);
 	virtual bool Redo(URigVMController* InController);
 };
@@ -320,6 +390,7 @@ public:
 	FRigVMAddUnitNodeAction();
 	FRigVMAddUnitNodeAction(URigVMUnitNode* InNode);
 	virtual ~FRigVMAddUnitNodeAction() {};
+	virtual UScriptStruct* GetScriptStruct() const override { return FRigVMAddUnitNodeAction::StaticStruct(); }
 	virtual bool Undo(URigVMController* InController) override;
 	virtual bool Redo(URigVMController* InController) override;
 
@@ -349,6 +420,7 @@ public:
 	FRigVMAddVariableNodeAction();
 	FRigVMAddVariableNodeAction(URigVMVariableNode* InNode);
 	virtual ~FRigVMAddVariableNodeAction() {};
+	virtual UScriptStruct* GetScriptStruct() const override { return FRigVMAddVariableNodeAction::StaticStruct(); }
 	virtual bool Undo(URigVMController* InController) override;
 	virtual bool Redo(URigVMController* InController) override;
 
@@ -387,6 +459,7 @@ public:
 	FRigVMAddCommentNodeAction();
 	FRigVMAddCommentNodeAction(URigVMCommentNode* InNode);
 	virtual ~FRigVMAddCommentNodeAction() {};
+	virtual UScriptStruct* GetScriptStruct() const override { return FRigVMAddCommentNodeAction::StaticStruct(); }
 	virtual bool Undo(URigVMController* InController) override;
 	virtual bool Redo(URigVMController* InController) override;
 
@@ -419,6 +492,7 @@ public:
 	FRigVMAddRerouteNodeAction();
 	FRigVMAddRerouteNodeAction(URigVMRerouteNode* InNode);
 	virtual ~FRigVMAddRerouteNodeAction() {};
+	virtual UScriptStruct* GetScriptStruct() const override { return FRigVMAddRerouteNodeAction::StaticStruct(); }
 	virtual bool Undo(URigVMController* InController) override;
 	virtual bool Redo(URigVMController* InController) override;
 
@@ -460,6 +534,7 @@ public:
 	FRigVMAddBranchNodeAction();
 	FRigVMAddBranchNodeAction(URigVMBranchNode* InNode);
 	virtual ~FRigVMAddBranchNodeAction() {};
+	virtual UScriptStruct* GetScriptStruct() const override { return FRigVMAddBranchNodeAction::StaticStruct(); }
 	virtual bool Undo(URigVMController* InController) override;
 	virtual bool Redo(URigVMController* InController) override;
 
@@ -483,6 +558,7 @@ public:
 	FRigVMAddIfNodeAction();
 	FRigVMAddIfNodeAction(URigVMIfNode* InNode);
 	virtual ~FRigVMAddIfNodeAction() {};
+	virtual UScriptStruct* GetScriptStruct() const override { return FRigVMAddIfNodeAction::StaticStruct(); }
 	virtual bool Undo(URigVMController* InController) override;
 	virtual bool Redo(URigVMController* InController) override;
 
@@ -512,6 +588,7 @@ public:
 	FRigVMAddSelectNodeAction();
 	FRigVMAddSelectNodeAction(URigVMSelectNode* InNode);
 	virtual ~FRigVMAddSelectNodeAction() {};
+	virtual UScriptStruct* GetScriptStruct() const override { return FRigVMAddSelectNodeAction::StaticStruct(); }
 	virtual bool Undo(URigVMController* InController) override;
 	virtual bool Redo(URigVMController* InController) override;
 
@@ -541,6 +618,7 @@ public:
 	FRigVMAddEnumNodeAction();
 	FRigVMAddEnumNodeAction(URigVMEnumNode* InNode);
 	virtual ~FRigVMAddEnumNodeAction() {};
+	virtual UScriptStruct* GetScriptStruct() const override { return FRigVMAddEnumNodeAction::StaticStruct(); }
 	virtual bool Undo(URigVMController* InController) override;
 	virtual bool Redo(URigVMController* InController) override;
 
@@ -570,6 +648,7 @@ public:
 	FRigVMAddTemplateNodeAction();
 	FRigVMAddTemplateNodeAction(URigVMTemplateNode* InNode);
 	virtual ~FRigVMAddTemplateNodeAction() {};
+	virtual UScriptStruct* GetScriptStruct() const override { return FRigVMAddTemplateNodeAction::StaticStruct(); }
 	virtual bool Undo(URigVMController* InController) override;
 	virtual bool Redo(URigVMController* InController) override;
 
@@ -596,6 +675,7 @@ public:
 	FRigVMInjectNodeIntoPinAction();
 	FRigVMInjectNodeIntoPinAction(URigVMInjectionInfo* InInjectionInfo);
 	virtual ~FRigVMInjectNodeIntoPinAction() {};
+	virtual UScriptStruct* GetScriptStruct() const override { return FRigVMInjectNodeIntoPinAction::StaticStruct(); }
 	virtual bool Undo(URigVMController* InController) override;
 	virtual bool Redo(URigVMController* InController) override;
 
@@ -628,6 +708,7 @@ public:
 	FRigVMRemoveNodeAction() {}
 	FRigVMRemoveNodeAction(URigVMNode* InNode, URigVMController* InController);
 	virtual ~FRigVMRemoveNodeAction() {};
+	virtual UScriptStruct* GetScriptStruct() const override { return FRigVMRemoveNodeAction::StaticStruct(); }
 	virtual bool Undo(URigVMController* InController) override;
 	virtual bool Redo(URigVMController* InController) override;
 
@@ -648,6 +729,7 @@ public:
 	FRigVMSetNodeSelectionAction();
 	FRigVMSetNodeSelectionAction(URigVMGraph* InGraph, TArray<FName> InNewSelection);
 	virtual ~FRigVMSetNodeSelectionAction() {};
+	virtual UScriptStruct* GetScriptStruct() const override { return FRigVMSetNodeSelectionAction::StaticStruct(); }
 	virtual bool Undo(URigVMController* InController) override;
 	virtual bool Redo(URigVMController* InController) override;
 
@@ -674,6 +756,7 @@ public:
 	}
 	FRigVMSetNodePositionAction(URigVMNode* InNode, const FVector2D& InNewPosition);
 	virtual ~FRigVMSetNodePositionAction() {};
+	virtual UScriptStruct* GetScriptStruct() const override { return FRigVMSetNodePositionAction::StaticStruct(); }
 	virtual bool Merge(const FRigVMBaseAction* Other);
 	virtual bool Undo(URigVMController* InController) override;
 	virtual bool Redo(URigVMController* InController) override;
@@ -704,6 +787,7 @@ public:
 	}
 	FRigVMSetNodeSizeAction(URigVMNode* InNode, const FVector2D& InNewSize);
 	virtual ~FRigVMSetNodeSizeAction() {};
+	virtual UScriptStruct* GetScriptStruct() const override { return FRigVMSetNodeSizeAction::StaticStruct(); }
 	virtual bool Merge(const FRigVMBaseAction* Other);
 	virtual bool Undo(URigVMController* InController) override;
 	virtual bool Redo(URigVMController* InController) override;
@@ -734,6 +818,7 @@ public:
 	}
 	FRigVMSetNodeColorAction(URigVMNode* InNode, const FLinearColor& InNewColor);
 	virtual ~FRigVMSetNodeColorAction() {};
+	virtual UScriptStruct* GetScriptStruct() const override { return FRigVMSetNodeColorAction::StaticStruct(); }
 	virtual bool Merge(const FRigVMBaseAction* Other);
 	virtual bool Undo(URigVMController* InController) override;
 	virtual bool Redo(URigVMController* InController) override;
@@ -764,6 +849,7 @@ public:
 	}
 	FRigVMSetNodeCategoryAction(URigVMCollapseNode* InNode, const FString& InNewCategory);
 	virtual ~FRigVMSetNodeCategoryAction() {};
+	virtual UScriptStruct* GetScriptStruct() const override { return FRigVMSetNodeCategoryAction::StaticStruct(); }
 	virtual bool Merge(const FRigVMBaseAction* Other);
 	virtual bool Undo(URigVMController* InController) override;
 	virtual bool Redo(URigVMController* InController) override;
@@ -795,6 +881,7 @@ public:
 	}
 	FRigVMSetNodeKeywordsAction(URigVMCollapseNode* InNode, const FString& InNewKeywords);
 	virtual ~FRigVMSetNodeKeywordsAction() {};
+	virtual UScriptStruct* GetScriptStruct() const override { return FRigVMSetNodeKeywordsAction::StaticStruct(); }
 	virtual bool Merge(const FRigVMBaseAction* Other);
 	virtual bool Undo(URigVMController* InController) override;
 	virtual bool Redo(URigVMController* InController) override;
@@ -825,6 +912,7 @@ public:
 	}
 	FRigVMSetNodeDescriptionAction(URigVMCollapseNode* InNode, const FString& InNewDescription);
 	virtual ~FRigVMSetNodeDescriptionAction() {};
+	virtual UScriptStruct* GetScriptStruct() const override { return FRigVMSetNodeDescriptionAction::StaticStruct(); }
 	virtual bool Merge(const FRigVMBaseAction* Other);
 	virtual bool Undo(URigVMController* InController) override;
 	virtual bool Redo(URigVMController* InController) override;
@@ -852,6 +940,7 @@ public:
 	FRigVMSetCommentTextAction();
 	FRigVMSetCommentTextAction(URigVMCommentNode* InNode, const FString& InNewText, const int32& InNewFontSize, const bool& bInNewBubbleVisible, const bool& bInNewColorBubble);
 	virtual ~FRigVMSetCommentTextAction() {};
+	virtual UScriptStruct* GetScriptStruct() const override { return FRigVMSetCommentTextAction::StaticStruct(); }
 	virtual bool Undo(URigVMController* InController) override;
 	virtual bool Redo(URigVMController* InController) override;
 
@@ -896,6 +985,7 @@ public:
 	FRigVMSetRerouteCompactnessAction();
 	FRigVMSetRerouteCompactnessAction(URigVMRerouteNode* InNode, bool InShowAsFullNode);
 	virtual ~FRigVMSetRerouteCompactnessAction() {};
+	virtual UScriptStruct* GetScriptStruct() const override { return FRigVMSetRerouteCompactnessAction::StaticStruct(); }
 	virtual bool Undo(URigVMController* InController) override;
 	virtual bool Redo(URigVMController* InController) override;
 
@@ -922,6 +1012,7 @@ public:
 	FRigVMRenameVariableAction() {}
 	FRigVMRenameVariableAction(const FName& InOldVariableName, const FName& InNewVariableName);
 	virtual ~FRigVMRenameVariableAction() {};
+	virtual UScriptStruct* GetScriptStruct() const override { return FRigVMRenameVariableAction::StaticStruct(); }
 	virtual bool Undo(URigVMController* InController) override;
 	virtual bool Redo(URigVMController* InController) override;
 
@@ -950,6 +1041,7 @@ public:
 
 	FRigVMSetPinExpansionAction(URigVMPin* InPin, bool bNewIsExpanded);
 	virtual ~FRigVMSetPinExpansionAction() {};
+	virtual UScriptStruct* GetScriptStruct() const override { return FRigVMSetPinExpansionAction::StaticStruct(); }
 	virtual bool Undo(URigVMController* InController) override;
 	virtual bool Redo(URigVMController* InController) override;
 
@@ -980,6 +1072,7 @@ public:
 
 	FRigVMSetPinWatchAction(URigVMPin* InPin, bool bNewIsWatched);
 	virtual ~FRigVMSetPinWatchAction() {};
+	virtual UScriptStruct* GetScriptStruct() const override { return FRigVMSetPinWatchAction::StaticStruct(); }
 	virtual bool Undo(URigVMController* InController) override;
 	virtual bool Redo(URigVMController* InController) override;
 
@@ -1006,6 +1099,7 @@ public:
 	FRigVMSetPinDefaultValueAction() {}
 	FRigVMSetPinDefaultValueAction(URigVMPin* InPin, const FString& InNewDefaultValue);
 	virtual ~FRigVMSetPinDefaultValueAction() {};
+	virtual UScriptStruct* GetScriptStruct() const override { return FRigVMSetPinDefaultValueAction::StaticStruct(); }
 	virtual bool Merge(const FRigVMBaseAction* Other);
 	virtual bool Undo(URigVMController* InController) override;
 	virtual bool Redo(URigVMController* InController) override;
@@ -1036,6 +1130,7 @@ public:
 	}
 	FRigVMInsertArrayPinAction(URigVMPin* InArrayPin, int32 InIndex, const FString& InNewDefaultValue);
 	virtual ~FRigVMInsertArrayPinAction() {};
+	virtual UScriptStruct* GetScriptStruct() const override { return FRigVMInsertArrayPinAction::StaticStruct(); }
 	virtual bool Undo(URigVMController* InController) override;
 	virtual bool Redo(URigVMController* InController) override;
 
@@ -1065,6 +1160,7 @@ public:
 	}
 	FRigVMRemoveArrayPinAction(URigVMPin* InArrayElementPin);
 	virtual ~FRigVMRemoveArrayPinAction() {};
+	virtual UScriptStruct* GetScriptStruct() const override { return FRigVMRemoveArrayPinAction::StaticStruct(); }
 	virtual bool Undo(URigVMController* InController) override;
 	virtual bool Redo(URigVMController* InController) override;
 
@@ -1091,6 +1187,7 @@ public:
 	FRigVMAddLinkAction() {}
 	FRigVMAddLinkAction(URigVMPin* InOutputPin, URigVMPin* InInputPin);
 	virtual ~FRigVMAddLinkAction() {};
+	virtual UScriptStruct* GetScriptStruct() const override { return FRigVMAddLinkAction::StaticStruct(); }
 	virtual bool Undo(URigVMController* InController) override;
 	virtual bool Redo(URigVMController* InController) override;
 
@@ -1114,6 +1211,7 @@ public:
 	FRigVMBreakLinkAction() {}
 	FRigVMBreakLinkAction(URigVMPin* InOutputPin, URigVMPin* InInputPin);
 	virtual ~FRigVMBreakLinkAction() {};
+	virtual UScriptStruct* GetScriptStruct() const override { return FRigVMBreakLinkAction::StaticStruct(); }
 	virtual bool Undo(URigVMController* InController) override;
 	virtual bool Redo(URigVMController* InController) override;
 
@@ -1137,6 +1235,7 @@ public:
 	FRigVMChangePinTypeAction();
 	FRigVMChangePinTypeAction(URigVMPin* InPin, const FString& InCppType, const FName& InCppTypeObjectPath, bool InSetupOrphanPins, bool InBreakLinks, bool InRemoveSubPins);
 	virtual ~FRigVMChangePinTypeAction() {};
+	virtual UScriptStruct* GetScriptStruct() const override { return FRigVMChangePinTypeAction::StaticStruct(); }
 	virtual bool Undo(URigVMController* InController) override;
 	virtual bool Redo(URigVMController* InController) override;
 
@@ -1178,6 +1277,7 @@ public:
 	FRigVMImportNodeFromTextAction();
 	FRigVMImportNodeFromTextAction(URigVMNode* InNode, URigVMController* InController);
 	virtual ~FRigVMImportNodeFromTextAction() {};
+	virtual UScriptStruct* GetScriptStruct() const override { return FRigVMImportNodeFromTextAction::StaticStruct(); }
 	virtual bool Undo(URigVMController* InController) override;
 	virtual bool Redo(URigVMController* InController) override;
 
@@ -1204,6 +1304,7 @@ public:
 	FRigVMCollapseNodesAction();
 	FRigVMCollapseNodesAction(URigVMController* InController, const TArray<URigVMNode*>& InNodes, const FString& InNodePath);
 	virtual ~FRigVMCollapseNodesAction() {};
+	virtual UScriptStruct* GetScriptStruct() const override { return FRigVMCollapseNodesAction::StaticStruct(); }
 	virtual bool Undo(URigVMController* InController) override;
 	virtual bool Redo(URigVMController* InController) override;
 
@@ -1233,6 +1334,7 @@ public:
 	FRigVMExpandNodeAction();
 	FRigVMExpandNodeAction(URigVMController* InController, URigVMLibraryNode* InLibraryNode);
 	virtual ~FRigVMExpandNodeAction() {};
+	virtual UScriptStruct* GetScriptStruct() const override { return FRigVMExpandNodeAction::StaticStruct(); }
 	virtual bool Undo(URigVMController* InController) override;
 	virtual bool Redo(URigVMController* InController) override;
 
@@ -1262,6 +1364,7 @@ public:
 	FRigVMRenameNodeAction() {}
 	FRigVMRenameNodeAction(const FName& InOldNodeName, const FName& InNewNodeName);
 	virtual ~FRigVMRenameNodeAction() {};
+	virtual UScriptStruct* GetScriptStruct() const override { return FRigVMRenameNodeAction::StaticStruct(); }
 	virtual bool Undo(URigVMController* InController) override;
 	virtual bool Redo(URigVMController* InController) override;
 
@@ -1285,6 +1388,7 @@ public:
 	FRigVMPushGraphAction() {}
 	FRigVMPushGraphAction(UObject* InGraph);
 	virtual ~FRigVMPushGraphAction() {};
+	virtual UScriptStruct* GetScriptStruct() const override { return FRigVMPushGraphAction::StaticStruct(); }
 	virtual bool Undo(URigVMController* InController) override;
 	virtual bool Redo(URigVMController* InController) override;
 
@@ -1305,6 +1409,7 @@ public:
 	FRigVMPopGraphAction() {}
 	FRigVMPopGraphAction(UObject* InGraph);
 	virtual ~FRigVMPopGraphAction() {};
+	virtual UScriptStruct* GetScriptStruct() const override { return FRigVMPopGraphAction::StaticStruct(); }
 	virtual bool Undo(URigVMController* InController) override;
 	virtual bool Redo(URigVMController* InController) override;
 
@@ -1326,6 +1431,7 @@ public:
 	FRigVMAddExposedPinAction() {}
 	FRigVMAddExposedPinAction(URigVMPin* InPin);
 	virtual ~FRigVMAddExposedPinAction() {};
+	virtual UScriptStruct* GetScriptStruct() const override { return FRigVMAddExposedPinAction::StaticStruct(); }
 	virtual bool Undo(URigVMController* InController) override;
 	virtual bool Redo(URigVMController* InController) override;
 
@@ -1358,6 +1464,7 @@ public:
 	FRigVMRemoveExposedPinAction() {}
 	FRigVMRemoveExposedPinAction(URigVMPin* InPin);
 	virtual ~FRigVMRemoveExposedPinAction() {};
+	virtual UScriptStruct* GetScriptStruct() const override { return FRigVMRemoveExposedPinAction::StaticStruct(); }
 	virtual bool Undo(URigVMController* InController) override;
 	virtual bool Redo(URigVMController* InController) override;
 
@@ -1390,6 +1497,7 @@ public:
 	FRigVMRenameExposedPinAction() {}
 	FRigVMRenameExposedPinAction(const FName& InOldPinName, const FName& InNewPinName);
 	virtual ~FRigVMRenameExposedPinAction() {};
+	virtual UScriptStruct* GetScriptStruct() const override { return FRigVMRenameExposedPinAction::StaticStruct(); }
 	virtual bool Undo(URigVMController* InController) override;
 	virtual bool Redo(URigVMController* InController) override;
 
@@ -1413,6 +1521,7 @@ public:
 	FRigVMSetPinIndexAction();
 	FRigVMSetPinIndexAction(URigVMPin* InPin, int32 InNewIndex);
 	virtual ~FRigVMSetPinIndexAction() {};
+	virtual UScriptStruct* GetScriptStruct() const override { return FRigVMSetPinIndexAction::StaticStruct(); }
 	virtual bool Undo(URigVMController* InController) override;
 	virtual bool Redo(URigVMController* InController) override;
 
@@ -1440,6 +1549,7 @@ public:
 	FRigVMSetRemappedVariableAction(URigVMFunctionReferenceNode* InFunctionRefNode, const FName& InInnerVariableName,
 		const FName& InOldOuterVariableName, const FName& InNewOuterVariableName);
 	virtual ~FRigVMSetRemappedVariableAction() {};
+	virtual UScriptStruct* GetScriptStruct() const override { return FRigVMSetRemappedVariableAction::StaticStruct(); }
 	virtual bool Undo(URigVMController* InController) override;
 	virtual bool Redo(URigVMController* InController) override;
 
@@ -1469,6 +1579,7 @@ public:
 	FRigVMAddLocalVariableAction() {}
 	FRigVMAddLocalVariableAction(const FRigVMGraphVariableDescription& InLocalVariable);
 	virtual ~FRigVMAddLocalVariableAction() {};
+	virtual UScriptStruct* GetScriptStruct() const override { return FRigVMAddLocalVariableAction::StaticStruct(); }
 	virtual bool Undo(URigVMController* InController) override;
 	virtual bool Redo(URigVMController* InController) override;
 
@@ -1489,6 +1600,7 @@ public:
 	FRigVMRemoveLocalVariableAction() {}
 	FRigVMRemoveLocalVariableAction(const FRigVMGraphVariableDescription& InLocalVariable);
 	virtual ~FRigVMRemoveLocalVariableAction() {};
+	virtual UScriptStruct* GetScriptStruct() const override { return FRigVMRemoveLocalVariableAction::StaticStruct(); }
 	virtual bool Undo(URigVMController* InController) override;
 	virtual bool Redo(URigVMController* InController) override;
 
@@ -1509,6 +1621,7 @@ public:
 	FRigVMRenameLocalVariableAction() {}
 	FRigVMRenameLocalVariableAction(const FName& InOldName, const FName& InNewName);
 	virtual ~FRigVMRenameLocalVariableAction() {};
+	virtual UScriptStruct* GetScriptStruct() const override { return FRigVMRenameLocalVariableAction::StaticStruct(); }
 	virtual bool Undo(URigVMController* InController) override;
 	virtual bool Redo(URigVMController* InController) override;
 
@@ -1531,6 +1644,7 @@ public:
 	FRigVMChangeLocalVariableTypeAction();
 	FRigVMChangeLocalVariableTypeAction(const FRigVMGraphVariableDescription& InLocalVariable, const FString& InCPPType, UObject* InCPPTypeObject);
 	virtual ~FRigVMChangeLocalVariableTypeAction() {};
+	virtual UScriptStruct* GetScriptStruct() const override { return FRigVMChangeLocalVariableTypeAction::StaticStruct(); }
 	virtual bool Undo(URigVMController* InController) override;
 	virtual bool Redo(URigVMController* InController) override;
 
@@ -1556,6 +1670,7 @@ public:
 	FRigVMChangeLocalVariableDefaultValueAction();
 	FRigVMChangeLocalVariableDefaultValueAction(const FRigVMGraphVariableDescription& InLocalVariable, const FString& InDefaultValue);
 	virtual ~FRigVMChangeLocalVariableDefaultValueAction() {};
+	virtual UScriptStruct* GetScriptStruct() const override { return FRigVMChangeLocalVariableDefaultValueAction::StaticStruct(); }
 	virtual bool Undo(URigVMController* InController) override;
 	virtual bool Redo(URigVMController* InController) override;
 
@@ -1579,6 +1694,7 @@ public:
 	FRigVMAddArrayNodeAction();
 	FRigVMAddArrayNodeAction(URigVMArrayNode* InNode);
 	virtual ~FRigVMAddArrayNodeAction() {};
+	virtual UScriptStruct* GetScriptStruct() const override { return FRigVMAddArrayNodeAction::StaticStruct(); }
 	virtual bool Undo(URigVMController* InController) override;
 	virtual bool Redo(URigVMController* InController) override;
 
@@ -1611,6 +1727,7 @@ public:
 	FRigVMPromoteNodeAction();
 	FRigVMPromoteNodeAction(const URigVMNode* InNodeToPromote, const FString& InNodePath, const FString& InFunctionDefinitionPath);
 	virtual ~FRigVMPromoteNodeAction() {};
+	virtual UScriptStruct* GetScriptStruct() const override { return FRigVMPromoteNodeAction::StaticStruct(); }
 	virtual bool Undo(URigVMController* InController) override;
 	virtual bool Redo(URigVMController* InController) override;
 
diff --git a/Engine/Source/Developer/RigVMDeveloper/Public/RigVMModel/RigVMNode.h b/Engine/Source/Developer/RigVMDeveloper/Public/RigVMModel/RigVMNode.h
--- a/Engine/Source/Developer/RigVMDeveloper/Public/RigVMModel/RigVMNode.h
+++ b/Engine/Source/Developer/RigVMDeveloper/Public/RigVMModel/RigVMNode.h
@@ -203,6 +203,9 @@ public:
 	// return true if this node is a loop node
 	virtual bool IsLoopNode() const { return false; }
 
+	// returns true if the node can be upgraded
+	virtual bool CanBeUpgraded() const { return false; }
+
 	UFUNCTION(BlueprintCallable, Category = RigVMNode)
 	bool HasBreakpoint() const { return bHasBreakpoint; }
 
diff --git a/Engine/Source/Developer/RigVMDeveloper/Public/RigVMModel/RigVMPin.h b/Engine/Source/Developer/RigVMDeveloper/Public/RigVMModel/RigVMPin.h
--- a/Engine/Source/Developer/RigVMDeveloper/Public/RigVMModel/RigVMPin.h
+++ b/Engine/Source/Developer/RigVMDeveloper/Public/RigVMModel/RigVMPin.h
@@ -298,6 +298,9 @@ public:
 	UFUNCTION(BlueprintCallable, Category = RigVMPin)
 	bool IsLinkedTo(URigVMPin* InPin) const;
 
+	// Returns true if the pin has any link
+	bool IsLinked(bool bRecursive = false) const;
+
 	// Returns all of the links linked to this Pin.
 	UFUNCTION(BlueprintCallable, Category = RigVMPin)
 	const TArray<URigVMLink*>& GetLinks() const;
diff --git a/Engine/Source/Runtime/RigVM/Private/RigVMCore/RigVMStruct.cpp b/Engine/Source/Runtime/RigVM/Private/RigVMCore/RigVMStruct.cpp
--- a/Engine/Source/Runtime/RigVM/Private/RigVMCore/RigVMStruct.cpp
+++ b/Engine/Source/Runtime/RigVM/Private/RigVMCore/RigVMStruct.cpp
@@ -120,6 +120,57 @@ void FRigVMStructUpgradeInfo::AddRemappedPin(const FString& InOldPinPath, const
 	}
 }
 
+FString FRigVMStructUpgradeInfo::RemapPin(const FString& InPinPath, bool bIsInput, bool bContainsNodeName) const
+{
+	FString NodeName;
+	FString PinPath = InPinPath;
+
+	if(bContainsNodeName)
+	{
+		if(!PinPath.Split(TEXT("."), &NodeName, &PinPath, ESearchCase::IgnoreCase, ESearchDir::FromStart))
+		{
+			return InPinPath;
+		}
+	}
+
+	const TMap<FString, FString>& LinkMap = bIsInput ? InputLinkMap : OutputLinkMap;
+
+	int32 FoundMaxLength = 0;
+	FString FoundReplacement;
+
+	for(const TPair<FString, FString>& Pair : LinkMap)
+	{
+		if(Pair.Key == PinPath)
+		{
+			PinPath = Pair.Value;
+			FoundReplacement.Reset();
+			break;
+		}
+
+		const FString KeyWithPeriod = Pair.Key + TEXT(".");
+		if(PinPath.StartsWith(KeyWithPeriod))
+		{
+			if(FoundMaxLength < KeyWithPeriod.Len())
+			{
+				FoundMaxLength = KeyWithPeriod.Len();
+				FoundReplacement = Pair.Value + PinPath.RightChop(Pair.Key.Len());
+			}
+		}
+	}
+
+	if(!FoundReplacement.IsEmpty())
+	{
+		PinPath = FoundReplacement;
+	}
+
+	if(bContainsNodeName)
+	{
+		PinPath = FString::Printf(TEXT("%s.%s"), *NodeName, *PinPath);
+	}
+
+	return PinPath;
+}
+
 void FRigVMStructUpgradeInfo::SetDefaultValues(const FRigVMStruct* InNewStructMemory)
 {
 	check(NewStruct);
@@ -475,7 +526,7 @@ FString FRigVMStruct::ExportToFullyQualifiedText(const FProperty* InMemberProper
 	return DefaultValue;
 }
 
-FString FRigVMStruct::ExportToFullyQualifiedText(UScriptStruct* InStruct, const uint8* InStructMemoryPtr)
+FString FRigVMStruct::ExportToFullyQualifiedText(const UScriptStruct* InStruct, const uint8* InStructMemoryPtr)
 {
 	check(InStruct);
 	check(InStructMemoryPtr);
@@ -502,7 +553,7 @@ FString FRigVMStruct::ExportToFullyQualifiedText(UScriptStruct* InStruct, const
 	return FString::Printf(TEXT("(%s)"), *FString::Join(FieldValues, TEXT(",")));
 }
 
-FString FRigVMStruct::ExportToFullyQualifiedText(UScriptStruct* InScriptStruct, const FName& InPropertyName, const uint8* InStructMemoryPointer) const
+FString FRigVMStruct::ExportToFullyQualifiedText(const UScriptStruct* InScriptStruct, const FName& InPropertyName, const uint8* InStructMemoryPointer) const
 {
 	check(InScriptStruct);
 	if(InStructMemoryPointer == nullptr)
diff --git a/Engine/Source/Runtime/RigVM/Public/RigVMCore/RigVMStruct.h b/Engine/Source/Runtime/RigVM/Public/RigVMCore/RigVMStruct.h
--- a/Engine/Source/Runtime/RigVM/Public/RigVMCore/RigVMStruct.h
+++ b/Engine/Source/Runtime/RigVM/Public/RigVMCore/RigVMStruct.h
@@ -168,6 +168,9 @@ public:
 	// adds a pin to be remapped
 	void AddRemappedPin(const FString& InOldPinPath, const FString& InNewPinPath, bool bAsInput = true, bool bAsOutput = true);
 
+	// remaps a pin path based on our internals
+	FString RemapPin(const FString& InPinPath, bool bIsInput, bool bContainsNodeName) const;
+
 private:
 	
 	// The complete node path including models / collapse node.
@@ -228,8 +231,8 @@ public:
 	static bool CheckFunctionExists(UScriptStruct* InStruct, const FName& FunctionName, FString* OutErrorMessage = nullptr);
 #endif
 	static FString ExportToFullyQualifiedText(const FProperty* InMemberProperty, const uint8* InMemberMemoryPtr, bool bUseQuotes = true);
-	static FString ExportToFullyQualifiedText(UScriptStruct* InStruct, const uint8* InStructMemoryPtr);
-	FString ExportToFullyQualifiedText(UScriptStruct* InScriptStruct, const FName& InPropertyName, const uint8* InStructMemoryPointer = nullptr) const;
+	static FString ExportToFullyQualifiedText(const UScriptStruct* InStruct, const uint8* InStructMemoryPtr);
+	FString ExportToFullyQualifiedText(const UScriptStruct* InScriptStruct, const FName& InPropertyName, const uint8* InStructMemoryPointer = nullptr) const;
 	virtual FRigVMStructUpgradeInfo GetUpgradeInfo() const { return FRigVMStructUpgradeInfo(); }
 
 	static const FName DeprecatedMetaName;
