Control Rig: Fix python script error when reroutes subpins connected

diff --git a/Engine/Source/Developer/RigVMDeveloper/Private/RigVMModel/RigVMController.cpp b/Engine/Source/Developer/RigVMDeveloper/Private/RigVMModel/RigVMController.cpp
--- a/Engine/Source/Developer/RigVMDeveloper/Private/RigVMModel/RigVMController.cpp
+++ b/Engine/Source/Developer/RigVMDeveloper/Private/RigVMModel/RigVMController.cpp
@@ -344,7 +344,7 @@ TArray<FString> URigVMController::GeneratePythonCommands()
 	}
 	
 	
-	// All nodes (except reroutes)
+	// All nodes
 	for (URigVMNode* Node : GetGraph()->GetNodes())
 	{
 		Commands.Append(GetAddNodePythonCommands(Node));
@@ -356,39 +356,11 @@ TArray<FString> URigVMController::GeneratePythonCommands()
 		URigVMPin* SourcePin = Link->GetSourcePin();
 		URigVMPin* TargetPin = Link->GetTargetPin();
 		
-		if (TargetPin->GetNode()->IsA<URigVMRerouteNode>())
-		{
-			continue;
-		}
-
 		if (SourcePin->GetInjectedNodes().Num() > 0 || TargetPin->GetInjectedNodes().Num() > 0)
 		{
 			continue;
 		}
 
-		// Iterate upstream until the source pin is not on a reroute node
-		while (SourcePin && SourcePin->GetNode()->IsA<URigVMRerouteNode>())
-		{
-			const URigVMNode* Node = SourcePin->GetNode();
-			if (Node->Pins.Num() > 0)
-			{
-				const TArray<URigVMLink*>& Links = Node->Pins[0]->GetSourceLinks();
-				if (Links.Num() > 0)
-				{
-					SourcePin = Links[0]->GetSourcePin();
-				}
-				else
-				{
-					break;
-				}
-			}
-		}
-
-		if (SourcePin->GetNode()->IsA<URigVMRerouteNode>())
-		{
-			continue;
-		}
-
 		const FString SourcePinPath = GetSanitizedPinPath(SourcePin->GetPinPath());
 		const FString TargetPinPath = GetSanitizedPinPath(TargetPin->GetPinPath());
 
@@ -406,121 +378,13 @@ TArray<FString> URigVMController::GeneratePythonCommands()
 		{
 			if (URigVMRerouteNode* Reroute = Cast<URigVMRerouteNode>(Node))
 			{
-				if (Reroute->Pins[0]->GetTargetLinks().Num() > 0)
-				{
-					Reroutes.Add(Reroute);
-				}
+				// SetRerouteCompactnessByName(const FName& InNodeName, bool bShowAsFullNode, bool bSetupUndoRedo = true);
+				Commands.Add(FString::Printf(TEXT("blueprint.get_controller_by_name('%s').set_reroute_compactness_by_name('%s', %s)"),
+					*GraphName,
+					*Reroute->GetName(),
+					Reroute->GetShowsAsFullNode() ? TEXT("True") : TEXT("False")));
 			}
 		}
-
-		TArray<FString> ReroutesAdded;
-		bool bNodesAdded = false;
-		do
-		{
-			bNodesAdded = false;
-			for (URigVMRerouteNode* Reroute : Reroutes)
-			{
-				const FString RerouteName = GetSanitizedNodeName(Reroute->GetName());
-
-				if (ReroutesAdded.Contains(RerouteName))
-				{
-					continue;
-				}
-
-				// If this reroute has no target links, we will ignore it (it will not be created)
-				if (Reroute->Pins.IsEmpty() || Reroute->Pins[0]->GetTargetLinks().IsEmpty())
-				{
-					continue;
-				}
-
-				TArray<URigVMPin*> TargetPins;
-				for (URigVMLink* Link : Reroute->Pins[0]->GetTargetLinks())
-				{
-					URigVMPin* TargetPin = Link->GetTargetPin();
-					if (!TargetPin->GetNode()->IsA<URigVMRerouteNode>() || ReroutesAdded.Contains(TargetPin->GetNode()->GetName()))
-					{
-						TargetPins.Add(TargetPin);
-						continue;
-					}
-				}
-
-				if (!TargetPins.IsEmpty())
-				{
-					URigVMPin* TargetPin = TargetPins[0];
-					
-					const FString TargetPinPath = GetSanitizedPinPath(TargetPin->GetPinPath());
-
-					// AddRerouteNodeOnPin(const FString& InPinPath, bool bAsInput, bool bShowAsFullNode, const FVector2D& InPosition, const FString& InNodeName, bool bSetupUndoRedo, bool bPrintPythonCommand)
-					Commands.Add(FString::Printf(TEXT("blueprint.get_controller_by_name('%s').add_reroute_node_on_pin('%s', True, %s, %s, '%s')"),
-							*GraphName,
-							*TargetPinPath,
-							Reroute->GetShowsAsFullNode() ? TEXT("True") : TEXT("False"),
-							*RigVMPythonUtils::Vector2DToPythonString(Reroute->GetPosition()),
-							*RerouteName));
-
-					ReroutesAdded.Add(RerouteName);
-					bNodesAdded = true;
-
-					// Add the rest of target links
-					for (int32 i=1; i<TargetPins.Num(); ++i)
-					{
-						URigVMPin* OtherTargetPin = Reroute->Pins[0]->GetTargetLinks()[i]->GetTargetPin();
-		
-						if (OtherTargetPin->GetNode()->IsA<URigVMRerouteNode>())
-						{
-							continue;
-						}
-
-						if (OtherTargetPin->GetInjectedNodes().Num() > 0)
-						{
-							continue;
-						}
-
-						const FString OtherTargetPinPath = GetSanitizedPinPath(OtherTargetPin->GetPinPath());
-						const FString FirstReroutePinPath = GetSanitizedPinPath(Reroute->Pins[0]->GetPinPath());
-
-						//bool AddLink(const FString& InOutputPinPath, const FString& InInputPinPath, bool bSetupUndoRedo = true);
-						Commands.Add(FString::Printf(TEXT("blueprint.get_controller_by_name('%s').add_link('%s', '%s')"),
-									*GraphName,
-									*FirstReroutePinPath,
-									*OtherTargetPinPath));
-					}
-
-					// Set default value in input if necessary
-					URigVMPin* Pin = Reroute->Pins[0];
-					const FString DefaultValue = Pin->GetDefaultValue();
-					if (!DefaultValue.IsEmpty() && DefaultValue != TEXT("()"))
-					{
-						const FString PinPath = GetSanitizedPinPath(Pin->GetPinPath());
-
-						Commands.Add(FString::Printf(TEXT("blueprint.get_controller_by_name('%s').set_pin_default_value('%s', '%s')"),
-									*GraphName,
-									*PinPath,
-									*Pin->GetDefaultValue()));
-					}
-					
-					// Add source links					
-					TArray<URigVMLink*> SourceLinks = Pin->GetSourceLinks(true);
-					for (URigVMLink* Link : SourceLinks)
-					{
-						URigVMPin* SourcePin = Link->GetSourcePin(); 
-						if (SourcePin->GetNode()->IsA<URigVMRerouteNode>())
-						{
-							continue;
-						}
-
-						const FString SourcePinPath = GetSanitizedPinPath(SourcePin->GetPinPath());
-						const FString LinkTargetPinPath = GetSanitizedPinPath(Link->GetTargetPin()->GetPinPath());
-
-						//bool AddLink(const FString& InOutputPinPath, const FString& InInputPinPath, bool bSetupUndoRedo = true);
-						Commands.Add(FString::Printf(TEXT("blueprint.get_controller_by_name('%s').add_link('%s', '%s')"),
-									*GraphName,
-									*SourcePinPath,
-									*LinkTargetPinPath));
-					}				
-				}				
-			}
-		} while (bNodesAdded);
 	}
 
 	return Commands;
@@ -691,7 +555,17 @@ TArray<FString> URigVMController::GetAddNodePythonCommands(URigVMNode* Node) con
 	}
 	else if (const URigVMRerouteNode* RerouteNode = Cast<URigVMRerouteNode>(Node))
 	{
-		// Do nothing, we need to create the links first
+		// add_free_reroute_node(bool bShowAsFullNode, const FString& InCPPType, const FName& InCPPTypeObjectPath, bool bIsConstant, const FName& InCustomWidgetName, const FString& InDefaultValue, const FVector2D& InPosition = FVector2D::ZeroVector, const FString& InNodeName = TEXT(""), bool bSetupUndoRedo = true);
+		Commands.Add(FString::Printf(TEXT("blueprint.get_controller_by_name('%s').add_free_reroute_node(%s, '%s', '%s', %s, '%s', '%s', %s, '%s')"),
+				*GraphName,
+				RerouteNode->GetShowsAsFullNode() ? TEXT("True") : TEXT("False"),
+				*RerouteNode->GetPins()[0]->GetCPPType(),
+				*RerouteNode->GetPins()[0]->GetCPPTypeObject()->GetPathName(),
+				RerouteNode->GetPins()[0]->IsDefinedAsConstant() ? TEXT("True") : TEXT("False"),
+				*RerouteNode->GetPins()[0]->GetCustomWidgetName().ToString(),
+				*RerouteNode->GetPins()[0]->GetDefaultValue(),
+				*RigVMPythonUtils::Vector2DToPythonString(RerouteNode->GetPosition()),
+				*NodeName));
 	}
 	else if (Node->IsA<URigVMFunctionEntryNode>() || Node->IsA<URigVMFunctionReturnNode>())
 	{
