Control Rig: Fix removal of BP variable provokes variable nodes referencing input arguments deleted

diff --git a/Engine/Source/Developer/RigVMDeveloper/Private/RigVMModel/RigVMController.cpp b/Engine/Source/Developer/RigVMDeveloper/Private/RigVMModel/RigVMController.cpp
--- a/Engine/Source/Developer/RigVMDeveloper/Private/RigVMModel/RigVMController.cpp
+++ b/Engine/Source/Developer/RigVMDeveloper/Private/RigVMModel/RigVMController.cpp
@@ -1109,7 +1109,7 @@ void URigVMController::OnExternalVariableRemoved(const FName& InVarName, bool bS
 	URigVMGraph* Graph = GetGraph();
 	check(Graph);
 
-	for (const FRigVMGraphVariableDescription& LocalVariable : Graph->GetLocalVariables())
+	for (const FRigVMGraphVariableDescription& LocalVariable : Graph->GetLocalVariables(true))
 	{
 		if (InVarName == LocalVariable.Name)
 		{
