Control Rig: Fix input variable node not working when connected to subpin

diff --git a/Engine/Source/Developer/RigVMDeveloper/Private/RigVMCompiler/RigVMAST.cpp b/Engine/Source/Developer/RigVMDeveloper/Private/RigVMCompiler/RigVMAST.cpp
--- a/Engine/Source/Developer/RigVMDeveloper/Private/RigVMCompiler/RigVMAST.cpp
+++ b/Engine/Source/Developer/RigVMDeveloper/Private/RigVMCompiler/RigVMAST.cpp
@@ -2957,27 +2957,44 @@ void FRigVMParserAST::Inline(URigVMGraph* InGraph, const TArray<FRigVMASTProxy>&
 				}
 
 				TArray<URigVMLink*> SourceLinks = ChildPin->GetSourceLinks(false /* recursive */);
+				
+				URigVMPin* SourcePin = nullptr;
 				if (SourceLinks.Num() > 0)
 				{
-					URigVMPin* SourcePin = nullptr;
-					if (SourceLinks.Num() > 0)
+					SourcePin = SourceLinks[0]->GetSourcePin();
+				}
+				else if(ChildPin->IsBoundToInputArgument())
+				{
+					if (URigVMGraph* Graph = ChildPin->GetGraph())
 					{
-						SourcePin = SourceLinks[0]->GetSourcePin();
+						if (URigVMFunctionEntryNode* EntryNode = Graph->GetEntryNode())
+						{
+							SourcePin = EntryNode->FindPin(ChildPin->GetBoundVariableName());
+						}
 					}
-					else
+				}
+				else if(URigVMVariableNode* VariableNode = Cast<URigVMVariableNode>(ChildPin->GetNode()))
+				{
+					if (VariableNode->IsInputArgument())
 					{
 						if (URigVMGraph* Graph = ChildPin->GetGraph())
 						{
 							if (URigVMFunctionEntryNode* EntryNode = Graph->GetEntryNode())
 							{
-								SourcePin = EntryNode->FindPin(ChildPin->GetBoundVariableName());
+								SourcePin = EntryNode->FindPin(VariableNode->GetVariableName().ToString());
+								if (ChildPin->GetParentPin())
+								{									
+									SourcePin = SourcePin->FindSubPin(ChildPin->GetSegmentPath());
+								}
 							}
 						}
 					}
-					check(SourcePin);
+				}
+				
+				if(SourcePin)
+				{
 					SourcePinProxy = InPinProxy.GetSibling(SourcePin);
 
-					// only continue the recursion on reroutes
 					if (ShouldRecursePin(SourcePinProxy))
 					{
 						FRigVMASTProxy SourceSourcePinProxy = FindSourcePin(SourcePinProxy, OutTraversalInfo);
@@ -3025,7 +3042,6 @@ void FRigVMParserAST::Inline(URigVMGraph* InGraph, const TArray<FRigVMASTProxy>&
 					{
 						SourcePinProxy = SourcePinProxy.GetSibling(SourceSubPin);
 
-						// only continue the recursion on reroutes
 						if (ShouldRecursePin(SourcePinProxy))
 						{
 							FRigVMASTProxy SourceSourceSubPinProxy = FindSourcePin(SourcePinProxy, OutTraversalInfo);
