RigVM: Ensure to create expressions in the right order for unit nodes

diff --git a/Engine/Source/Developer/RigVMDeveloper/Private/RigVMCompiler/RigVMAST.cpp b/Engine/Source/Developer/RigVMDeveloper/Private/RigVMCompiler/RigVMAST.cpp
--- a/Engine/Source/Developer/RigVMDeveloper/Private/RigVMCompiler/RigVMAST.cpp
+++ b/Engine/Source/Developer/RigVMDeveloper/Private/RigVMCompiler/RigVMAST.cpp
@@ -20,6 +20,7 @@
 #include "Stats/StatsHierarchical.h"
 #include "RigVMDeveloperModule.h"
 #include "VisualGraphUtils.h"
+#include "UObject/FieldIterator.h"
 
 FRigVMExprAST::FRigVMExprAST(EType InType, const FRigVMASTProxy& InProxy)
 	: Name(NAME_None)
@@ -978,7 +979,29 @@ TArray<FRigVMExprAST*> FRigVMParserAST::TraversePins(const FRigVMASTProxy& InNod
 	URigVMNode* Node = InNodeProxy.GetSubjectChecked<URigVMNode>();
 	TArray<FRigVMExprAST*> PinExpressions;
 
-	for (URigVMPin* Pin : Node->GetPins())
+	TArray<URigVMPin*> Pins;
+
+	// traverse the pins on a unit node in the order of the property definitions
+	if(const URigVMUnitNode* UnitNode = Cast<URigVMUnitNode>(Node))
+	{
+		if(UScriptStruct* ScriptStruct = UnitNode->GetScriptStruct())
+		{
+			for (TFieldIterator<FProperty> PropertyIt(ScriptStruct); PropertyIt; ++PropertyIt)
+			{
+				if(URigVMPin* Pin = UnitNode->FindPin(PropertyIt->GetName()))
+				{
+					Pins.Add(Pin);
+				}
+			}
+		}
+	}
+
+	if(Pins.IsEmpty())
+	{
+		Pins = Node->GetPins();
+	}
+
+	for (URigVMPin* Pin : Pins)
 	{
 		FRigVMASTProxy PinProxy = InNodeProxy.GetSibling(Pin);
 
