Control Rig: Further template node workflows

Implemented unresolve as well as re-resolve node on already typed pins
Added UI in contextual menu to re-type pins

diff --git a/Engine/Plugins/Animation/ControlRig/Source/ControlRig/Private/Units/Hierarchy/RigUnit_WorldSpace.cpp b/Engine/Plugins/Animation/ControlRig/Source/ControlRig/Private/Units/Hierarchy/RigUnit_WorldSpace.cpp
--- a/Engine/Plugins/Animation/ControlRig/Source/ControlRig/Private/Units/Hierarchy/RigUnit_WorldSpace.cpp
+++ b/Engine/Plugins/Animation/ControlRig/Source/ControlRig/Private/Units/Hierarchy/RigUnit_WorldSpace.cpp
@@ -6,35 +6,35 @@
 FRigUnit_ToWorldSpace_Transform_Execute()
 {
     DECLARE_SCOPE_HIERARCHICAL_COUNTER_RIGUNIT()
-    World = Context.ToWorldSpace(Transform);
+    World = Context.ToWorldSpace(Value);
 }
 
 FRigUnit_ToRigSpace_Transform_Execute()
 {
     DECLARE_SCOPE_HIERARCHICAL_COUNTER_RIGUNIT()
-    Global = Context.ToRigSpace(Transform);
+    Global = Context.ToRigSpace(Value);
 }
 
 FRigUnit_ToWorldSpace_Location_Execute()
 {
     DECLARE_SCOPE_HIERARCHICAL_COUNTER_RIGUNIT()
-    World = Context.ToWorldSpace(Location);
+    World = Context.ToWorldSpace(Value);
 }
 
 FRigUnit_ToRigSpace_Location_Execute()
 {
     DECLARE_SCOPE_HIERARCHICAL_COUNTER_RIGUNIT()
-    Global = Context.ToRigSpace(Location);
+    Global = Context.ToRigSpace(Value);
 }
 
 FRigUnit_ToWorldSpace_Rotation_Execute()
 {
     DECLARE_SCOPE_HIERARCHICAL_COUNTER_RIGUNIT()
-    World = Context.ToWorldSpace(Rotation);
+    World = Context.ToWorldSpace(Value);
 }
 
 FRigUnit_ToRigSpace_Rotation_Execute()
 {
     DECLARE_SCOPE_HIERARCHICAL_COUNTER_RIGUNIT()
-    Global = Context.ToRigSpace(Rotation);
+    Global = Context.ToRigSpace(Value);
 }
\ No newline at end of file
diff --git a/Engine/Plugins/Animation/ControlRig/Source/ControlRig/Private/Units/Hierarchy/RigUnit_WorldSpace.h b/Engine/Plugins/Animation/ControlRig/Source/ControlRig/Private/Units/Hierarchy/RigUnit_WorldSpace.h
--- a/Engine/Plugins/Animation/ControlRig/Source/ControlRig/Private/Units/Hierarchy/RigUnit_WorldSpace.h
+++ b/Engine/Plugins/Animation/ControlRig/Source/ControlRig/Private/Units/Hierarchy/RigUnit_WorldSpace.h
@@ -8,13 +8,13 @@
 /**
  * Converts a transform from rig (global) space to world space
  */
-USTRUCT(meta=(DisplayName="To World", MenuDescSuffix="(Transform)", Category="Hierarchy", DocumentationPolicy = "Strict", Keywords="Global,Local,World,Actor,ComponentSpace,FromRig", NodeColor="0.462745, 1,0, 0.329412", Varying))
+USTRUCT(meta=(DisplayName="To World", TemplateName="To World", Category="Hierarchy", DocumentationPolicy = "Strict", Keywords="Global,Local,World,Actor,ComponentSpace,FromRig", NodeColor="0.462745, 1,0, 0.329412", Varying))
 struct CONTROLRIG_API FRigUnit_ToWorldSpace_Transform : public FRigUnit
 {
 	GENERATED_BODY()
 
 	FRigUnit_ToWorldSpace_Transform()
-		: Transform(FTransform::Identity)
+		: Value(FTransform::Identity)
 		, World(FTransform::Identity)
 	{}
 
@@ -25,7 +25,7 @@ struct CONTROLRIG_API FRigUnit_ToWorldSpace_Transform : public FRigUnit
 	 * The input transform in global / rig space
 	 */
 	UPROPERTY(meta = (Input))
-	FTransform Transform;
+	FTransform Value;
 
 	/**
 	 * The result transform in world space
@@ -37,13 +37,13 @@ struct CONTROLRIG_API FRigUnit_ToWorldSpace_Transform : public FRigUnit
 /**
  * Converts a transform from world space to rig (global) space
  */
-USTRUCT(meta=(DisplayName="From World", MenuDescSuffix="(Transform)", Category="Hierarchy", DocumentationPolicy = "Strict", Keywords="Global,Local,World,Actor,ComponentSpace,ToRig", NodeColor="0.462745, 1,0, 0.329412", Varying))
+USTRUCT(meta=(DisplayName="From World", TemplateName="From World", Category="Hierarchy", DocumentationPolicy = "Strict", Keywords="Global,Local,World,Actor,ComponentSpace,ToRig", NodeColor="0.462745, 1,0, 0.329412", Varying))
 struct CONTROLRIG_API FRigUnit_ToRigSpace_Transform : public FRigUnit
 {
 	GENERATED_BODY()
 
 	FRigUnit_ToRigSpace_Transform()
-		: Transform(FTransform::Identity)
+		: Value(FTransform::Identity)
 		, Global(FTransform::Identity)
 	{}
 
@@ -54,7 +54,7 @@ struct CONTROLRIG_API FRigUnit_ToRigSpace_Transform : public FRigUnit
 	 * The input transform in world
 	 */
 	UPROPERTY(meta = (Input))
-	FTransform Transform;
+	FTransform Value;
 
 	/**
 	 * The result transform in global / rig space
@@ -66,13 +66,13 @@ struct CONTROLRIG_API FRigUnit_ToRigSpace_Transform : public FRigUnit
 /**
  * Converts a position / location from rig (global) space to world space
  */
-USTRUCT(meta=(DisplayName="To World", MenuDescSuffix="(Vector)", Category="Hierarchy", DocumentationPolicy = "Strict", Keywords="Global,Local,World,Actor,ComponentSpace,FromRig", NodeColor="0.462745, 1,0, 0.329412", Varying))
+USTRUCT(meta=(DisplayName="To World", TemplateName="To World", Category="Hierarchy", DocumentationPolicy = "Strict", Keywords="Global,Local,World,Actor,ComponentSpace,FromRig", NodeColor="0.462745, 1,0, 0.329412", Varying))
 struct CONTROLRIG_API FRigUnit_ToWorldSpace_Location : public FRigUnit
 {
 	GENERATED_BODY()
 
 	FRigUnit_ToWorldSpace_Location()
-		: Location(FVector::ZeroVector)
+		: Value(FVector::ZeroVector)
 		, World(FVector::ZeroVector)
 	{}
 
@@ -83,7 +83,7 @@ struct CONTROLRIG_API FRigUnit_ToWorldSpace_Location : public FRigUnit
 	 * The input position / location in global / rig space
 	 */
 	UPROPERTY(meta = (Input))
-	FVector Location;
+	FVector Value;
 
 	/**
 	 * The result position / location in world space
@@ -95,13 +95,13 @@ struct CONTROLRIG_API FRigUnit_ToWorldSpace_Location : public FRigUnit
 /**
  * Converts a position / location from world space to rig (global) space
  */
-USTRUCT(meta=(DisplayName="From World", MenuDescSuffix="(Vector)", Category="Hierarchy", DocumentationPolicy = "Strict", Keywords="Global,Local,World,Actor,ComponentSpace,ToRig", NodeColor="0.462745, 1,0, 0.329412", Varying))
+USTRUCT(meta=(DisplayName="From World", TemplateName="From World", Category="Hierarchy", DocumentationPolicy = "Strict", Keywords="Global,Local,World,Actor,ComponentSpace,ToRig", NodeColor="0.462745, 1,0, 0.329412", Varying))
 struct CONTROLRIG_API FRigUnit_ToRigSpace_Location : public FRigUnit
 {
 	GENERATED_BODY()
 
 	FRigUnit_ToRigSpace_Location()
-		: Location(FVector::ZeroVector)
+		: Value(FVector::ZeroVector)
 		, Global(FVector::ZeroVector)
 	{}
 
@@ -112,7 +112,7 @@ struct CONTROLRIG_API FRigUnit_ToRigSpace_Location : public FRigUnit
 	 * The input position / location in world
 	 */
 	UPROPERTY(meta = (Input))
-	FVector Location;
+	FVector Value;
 
 	/**
 	 * The result position / location in global / rig space
@@ -124,13 +124,13 @@ struct CONTROLRIG_API FRigUnit_ToRigSpace_Location : public FRigUnit
 /**
  * Converts a rotation from rig (global) space to world space
  */
-USTRUCT(meta=(DisplayName="To World", MenuDescSuffix="(Quat)", Category="Hierarchy", DocumentationPolicy = "Strict", Keywords="Global,Local,World,Actor,ComponentSpace,FromRig", NodeColor="0.462745, 1,0, 0.329412", Varying))
+USTRUCT(meta=(DisplayName="To World", TemplateName="To World", Category="Hierarchy", DocumentationPolicy = "Strict", Keywords="Global,Local,World,Actor,ComponentSpace,FromRig", NodeColor="0.462745, 1,0, 0.329412", Varying))
 struct CONTROLRIG_API FRigUnit_ToWorldSpace_Rotation : public FRigUnit
 {
 	GENERATED_BODY()
 
 	FRigUnit_ToWorldSpace_Rotation()
-		: Rotation(FQuat::Identity)
+		: Value(FQuat::Identity)
 		, World(FQuat::Identity)
 	{}
 
@@ -141,7 +141,7 @@ struct CONTROLRIG_API FRigUnit_ToWorldSpace_Rotation : public FRigUnit
 	 * The input rotation in global / rig space
 	 */
 	UPROPERTY(meta = (Input))
-	FQuat Rotation;
+	FQuat Value;
 
 	/**
 	 * The result rotation in world space
@@ -153,13 +153,13 @@ struct CONTROLRIG_API FRigUnit_ToWorldSpace_Rotation : public FRigUnit
 /**
  * Converts a rotation from world space to rig (global) space
  */
-USTRUCT(meta=(DisplayName="From World", MenuDescSuffix="(Quat)", Category="Hierarchy", DocumentationPolicy = "Strict", Keywords="Global,Local,World,Actor,ComponentSpace,ToRig", NodeColor="0.462745, 1,0, 0.329412", Varying))
+USTRUCT(meta=(DisplayName="From World", TemplateName="From World", Category="Hierarchy", DocumentationPolicy = "Strict", Keywords="Global,Local,World,Actor,ComponentSpace,ToRig", NodeColor="0.462745, 1,0, 0.329412", Varying))
 struct CONTROLRIG_API FRigUnit_ToRigSpace_Rotation : public FRigUnit
 {
 	GENERATED_BODY()
 
 	FRigUnit_ToRigSpace_Rotation()
-		: Rotation(FQuat::Identity)
+		: Value(FQuat::Identity)
 		, Global(FQuat::Identity)
 	{}
 
@@ -170,7 +170,7 @@ struct CONTROLRIG_API FRigUnit_ToRigSpace_Rotation : public FRigUnit
 	 * The input rotation in world
 	 */
 	UPROPERTY(meta = (Input))
-	FQuat Rotation;
+	FQuat Value;
 
 	/**
 	 * The result rotation in global / rig space
diff --git a/Engine/Plugins/Animation/ControlRig/Source/ControlRig/Private/Units/Math/RigUnit_MathBool.h b/Engine/Plugins/Animation/ControlRig/Source/ControlRig/Private/Units/Math/RigUnit_MathBool.h
--- a/Engine/Plugins/Animation/ControlRig/Source/ControlRig/Private/Units/Math/RigUnit_MathBool.h
+++ b/Engine/Plugins/Animation/ControlRig/Source/ControlRig/Private/Units/Math/RigUnit_MathBool.h
@@ -5,7 +5,7 @@
 #include "Units/Math/RigUnit_MathBase.h"
 #include "RigUnit_MathBool.generated.h"
 
-USTRUCT(meta=(Abstract, Category="Math|Boolean", MenuDescSuffix="(Bool)"))
+USTRUCT(meta=(Abstract, Category="Math|Boolean"))
 struct CONTROLRIG_API FRigUnit_MathBoolBase : public FRigUnit_MathBase
 {
 	GENERATED_BODY()
diff --git a/Engine/Plugins/Animation/ControlRig/Source/ControlRigDeveloper/Private/ControlRigBlueprint.cpp b/Engine/Plugins/Animation/ControlRig/Source/ControlRigDeveloper/Private/ControlRigBlueprint.cpp
--- a/Engine/Plugins/Animation/ControlRig/Source/ControlRigDeveloper/Private/ControlRigBlueprint.cpp
+++ b/Engine/Plugins/Animation/ControlRig/Source/ControlRigDeveloper/Private/ControlRigBlueprint.cpp
@@ -1580,6 +1580,19 @@ URigVMController* UControlRigBlueprint::GetOrCreateController(URigVMGraph* InGra
 		return FName();
 	});
 
+
+	Controller->RequestJumpToHyperlinkDelegate.BindLambda([WeakThis](const UObject* InSubject)
+	{
+		if (WeakThis.IsValid())
+		{
+			UControlRigBlueprint* StrongThis = WeakThis.Get();
+			if(StrongThis->OnRequestJumpToHyperlink().IsBound())
+			{
+				StrongThis->OnRequestJumpToHyperlink().Execute(InSubject);
+			}
+		}
+	});
+
 #endif
 
 	Controller->RemoveStaleNodes();
@@ -3825,6 +3838,12 @@ void UControlRigBlueprint::ConvertUnitNodesToTemplateNodes()
 {
 	TGuardValue<bool> GuardNotifsSelf(bSuspendModelNotificationsForSelf, true);
 
+	bool bWasDirty = false;
+	if (const UPackage* Package = GetOutermost())
+	{
+		bWasDirty = Package->IsDirty();
+	}
+
 	for (URigVMGraph* Graph : GetAllModels())
 	{
 		URigVMController* Controller = GetOrCreateController(Graph);
@@ -3834,6 +3853,11 @@ void UControlRigBlueprint::ConvertUnitNodesToTemplateNodes()
 			Controller->ReplaceUnitNodeWithTemplateNode(Node->GetFName(), false);
 		}
 	}
+
+	if (UPackage* Package = GetOutermost())
+	{
+		Package->SetDirtyFlag(bWasDirty);
+	}
 }
 
 void UControlRigBlueprint::PatchParameterNodesOnLoad()
diff --git a/Engine/Plugins/Animation/ControlRig/Source/ControlRigDeveloper/Private/Graph/ControlRigGraphNode.cpp b/Engine/Plugins/Animation/ControlRig/Source/ControlRigDeveloper/Private/Graph/ControlRigGraphNode.cpp
--- a/Engine/Plugins/Animation/ControlRig/Source/ControlRigDeveloper/Private/Graph/ControlRigGraphNode.cpp
+++ b/Engine/Plugins/Animation/ControlRig/Source/ControlRigDeveloper/Private/Graph/ControlRigGraphNode.cpp
@@ -564,7 +564,7 @@ void UControlRigGraphNode::CreateExecutionPins()
 			Pair.InputPin = CreatePin(EGPD_Input, GetPinTypeForModelPin(ModelPin), FName(*ModelPin->GetPinPath()));
 			if (Pair.InputPin != nullptr)
 		{
-				Pair.InputPin->PinFriendlyName = FText::FromName(ModelPin->GetDisplayName());
+				ConfigurePin(Pair.InputPin, ModelPin, false, true);
 		}
 	}
 		if (Pair.OutputPin == nullptr)
@@ -572,7 +572,7 @@ void UControlRigGraphNode::CreateExecutionPins()
 			Pair.OutputPin = CreatePin(EGPD_Output, GetPinTypeForModelPin(ModelPin), FName(*ModelPin->GetPinPath()));
 			if (Pair.OutputPin != nullptr)
 		{
-				Pair.OutputPin->PinFriendlyName = FText::FromName(ModelPin->GetDisplayName());
+				ConfigurePin(Pair.OutputPin, ModelPin, false, true);
 		}
 	}
 		// note: no recursion for execution pins
@@ -592,9 +592,7 @@ void UControlRigGraphNode::CreateInputPins(URigVMPin* InParentPin)
 			Pair.InputPin = CreatePin(EGPD_Input, GetPinTypeForModelPin(ModelPin), FName(*ModelPin->GetPinPath()));
 			if (Pair.InputPin != nullptr)
 			{
-				Pair.InputPin->PinFriendlyName = FText::FromName(ModelPin->GetDisplayName());
-				Pair.InputPin->bNotConnectable = ModelPin->GetDirection() != ERigVMPinDirection::Input;
-				Pair.InputPin->bOrphanedPin = ModelPin->IsOrphanPin() ? 1 : 0; 
+				ConfigurePin(Pair.InputPin, ModelPin, false, ModelPin->GetDirection() == ERigVMPinDirection::Input);
 
 				SetupPinDefaultsFromModel(Pair.InputPin);
 
@@ -629,10 +627,7 @@ void UControlRigGraphNode::CreateInputOutputPins(URigVMPin* InParentPin, bool bH
 			Pair.InputPin = CreatePin(EGPD_Input, GetPinTypeForModelPin(ModelPin), FName(*ModelPin->GetPinPath()));
 			if (Pair.InputPin != nullptr)
 		{
-				Pair.InputPin->bHidden = bHidden;
-				Pair.InputPin->PinFriendlyName = FText::FromName(ModelPin->GetDisplayName());
-				Pair.InputPin->bNotConnectable = ModelPin->GetDirection() != ERigVMPinDirection::IO;
-				Pair.InputPin->bOrphanedPin = ModelPin->IsOrphanPin() ? 1 : 0; 
+				ConfigurePin(Pair.InputPin, ModelPin, bHidden, ModelPin->GetDirection() == ERigVMPinDirection::IO);
 
 				SetupPinDefaultsFromModel(Pair.InputPin);
 
@@ -649,10 +644,7 @@ void UControlRigGraphNode::CreateInputOutputPins(URigVMPin* InParentPin, bool bH
 			Pair.OutputPin = CreatePin(EGPD_Output, GetPinTypeForModelPin(ModelPin), FName(*ModelPin->GetPinPath()));
 			if (Pair.OutputPin != nullptr)
 		{
-				Pair.OutputPin->bHidden = bHidden;
-				Pair.OutputPin->PinFriendlyName = FText::FromName(ModelPin->GetDisplayName());
-				Pair.OutputPin->bNotConnectable = ModelPin->GetDirection() != ERigVMPinDirection::IO;
-				Pair.OutputPin->bOrphanedPin = ModelPin->IsOrphanPin() ? 1 : 0; 
+				ConfigurePin(Pair.OutputPin, ModelPin, bHidden, ModelPin->GetDirection() == ERigVMPinDirection::IO);
 
 				if (InParentPin != nullptr)
 		{
@@ -696,9 +688,7 @@ void UControlRigGraphNode::CreateOutputPins(URigVMPin* InParentPin)
 			Pair.OutputPin = CreatePin(EGPD_Output, GetPinTypeForModelPin(ModelPin), FName(*ModelPin->GetPinPath()));
 			if (Pair.OutputPin != nullptr)
 			{
-				Pair.OutputPin->PinFriendlyName = FText::FromName(ModelPin->GetDisplayName());
-				Pair.OutputPin->bNotConnectable = ModelPin->GetDirection() != ERigVMPinDirection::Output;
-				Pair.OutputPin->bOrphanedPin = ModelPin->IsOrphanPin() ? 1 : 0; 
+				ConfigurePin(Pair.OutputPin, ModelPin, false,  ModelPin->GetDirection() == ERigVMPinDirection::Output);
 
 				if (InParentPin != nullptr)
 	{
@@ -1188,64 +1178,18 @@ bool UControlRigGraphNode::ShouldDrawNodeAsControlPointOnly(int32& OutInputPinIn
 
 FEdGraphPinType UControlRigGraphNode::GetPinTypeForModelPin(URigVMPin* InModelPin)
 {
-	FEdGraphPinType PinType;
-	PinType.ResetToDefaults();
-
-	FName ModelPinCPPType = InModelPin->IsArray() ? *InModelPin->GetArrayElementCppType() : *InModelPin->GetCPPType();
-	if (ModelPinCPPType == TEXT("bool"))
-	{
-		PinType.PinCategory = UEdGraphSchema_K2::PC_Boolean;
-	}
-	else if (ModelPinCPPType == TEXT("int32"))
-	{
-		PinType.PinCategory = UEdGraphSchema_K2::PC_Int;
-	}
-	else if (ModelPinCPPType == TEXT("float"))
-	{
-		PinType.PinCategory = UEdGraphSchema_K2::PC_Real;
-		PinType.PinSubCategory = UEdGraphSchema_K2::PC_Float;
-	}
-	else if (ModelPinCPPType == TEXT("double"))
-	{
-		PinType.PinCategory = UEdGraphSchema_K2::PC_Real;
-		PinType.PinSubCategory = UEdGraphSchema_K2::PC_Double;
-	}
-	else if (ModelPinCPPType == TEXT("FName"))
-	{
-		PinType.PinCategory = UEdGraphSchema_K2::PC_Name;
-	}
-	else if (ModelPinCPPType == TEXT("FString"))
-	{
-		PinType.PinCategory = UEdGraphSchema_K2::PC_String;
-	}
-	else if (InModelPin->GetScriptStruct() != nullptr)
-	{
-		PinType.PinCategory = UEdGraphSchema_K2::PC_Struct;
-		PinType.PinSubCategoryObject = InModelPin->GetScriptStruct();
-	}
-	else if (InModelPin->IsUObject())
-	{
-		PinType.PinCategory = UEdGraphSchema_K2::PC_Object;
-		PinType.PinSubCategoryObject = InModelPin->GetCPPTypeObject();
-	}
-	else if (InModelPin->GetEnum() != nullptr)
-	{
-		PinType.PinCategory = UEdGraphSchema_K2::PC_Byte;
-		PinType.PinSubCategoryObject = InModelPin->GetEnum();
-	}
-
-	if (InModelPin->IsArray())
-	{
-		PinType.ContainerType = EPinContainerType::Array;
-	}
-	else
-	{
-		PinType.ContainerType = EPinContainerType::None;
-	}
-
+	FEdGraphPinType PinType = RigVMTypeUtils::PinTypeFromCPPType(InModelPin->GetCPPType(), InModelPin->GetCPPTypeObject());
 	PinType.bIsConst = InModelPin->IsDefinedAsConstant();
-
 	return PinType;
 }
 
+void UControlRigGraphNode::ConfigurePin(UEdGraphPin* EdGraphPin, URigVMPin* ModelPin, bool bHidden, bool bConnectable)
+{
+	EdGraphPin->bHidden = bHidden;
+	EdGraphPin->PinFriendlyName = FText::FromName(ModelPin->GetDisplayName());
+	EdGraphPin->bNotConnectable = !bConnectable;
+	EdGraphPin->bOrphanedPin = ModelPin->IsOrphanPin() ? 1 : 0; 
+	EdGraphPin->bDisplayAsMutableRef = ModelPin->IsWildCard();
+}
+
 #undef LOCTEXT_NAMESPACE
diff --git a/Engine/Plugins/Animation/ControlRig/Source/ControlRigDeveloper/Private/Graph/ControlRigGraphSchema.cpp b/Engine/Plugins/Animation/ControlRig/Source/ControlRigDeveloper/Private/Graph/ControlRigGraphSchema.cpp
--- a/Engine/Plugins/Animation/ControlRig/Source/ControlRigDeveloper/Private/Graph/ControlRigGraphSchema.cpp
+++ b/Engine/Plugins/Animation/ControlRig/Source/ControlRigDeveloper/Private/Graph/ControlRigGraphSchema.cpp
@@ -562,11 +562,11 @@ bool UControlRigGraphSchema::TryCreateConnection(UEdGraphPin* PinA, UEdGraphPin*
 	{
 		if (URigVMController* Controller = RigBlueprint->GetOrCreateController(PinA->GetOwningNode()->GetGraph()))
 		{
+			ERigVMPinDirection UserLinkDirection = ERigVMPinDirection::Output;
 			if (PinA->Direction == EGPD_Input)
 			{
-				UEdGraphPin* Temp = PinA;
-				PinA = PinB;
-				PinB = Temp;
+				Swap(PinA, PinB);
+				UserLinkDirection = ERigVMPinDirection::Input;
 			}
 
 #if WITH_EDITOR
@@ -644,7 +644,7 @@ bool UControlRigGraphSchema::TryCreateConnection(UEdGraphPin* PinA, UEdGraphPin*
 			}
 #endif
 
-			return Controller->AddLink(PinA->GetName(), PinB->GetName(), true, true);
+			return Controller->AddLink(PinA->GetName(), PinB->GetName(), true, true, UserLinkDirection);
 		}
 	}
 	return false;
@@ -697,17 +697,17 @@ const FPinConnectionResponse UControlRigGraphSchema::CanCreateConnection(const U
 				PinB = PinB->GetPinForLink();
 			}
 
+			ERigVMPinDirection UserLinkDirection = ERigVMPinDirection::Output;
 			if (A->Direction == EGPD_Input)
 			{
-				URigVMPin* Temp = PinA;
-				PinA = PinB;
-				PinB = Temp;
+				Swap(PinA, PinB);
+				UserLinkDirection = ERigVMPinDirection::Input;
 			}
 
 			const FRigVMByteCode* ByteCode = RigNodeA->GetController()->GetCurrentByteCode();
 
 			FString FailureReason;
-			bool bResult = RigNodeA->GetModel()->CanLink(PinA, PinB, &FailureReason, ByteCode);
+			bool bResult = RigNodeA->GetModel()->CanLink(PinA, PinB, &FailureReason, ByteCode, UserLinkDirection);
 			if (!bResult)
 			{
 				return FPinConnectionResponse(CONNECT_RESPONSE_DISALLOW, FText::FromString(FailureReason));
@@ -1565,8 +1565,17 @@ bool UControlRigGraphSchema::ArePinsCompatible(const UEdGraphPin* PinA, const UE
 			}
 			return true;
 		};
+
+		auto IsTemplateNodePin = [](const UEdGraphPin* InPin)
+		{
+			if(const UControlRigGraphNode* RigNode = Cast<UControlRigGraphNode>(InPin->GetOwningNode()))
+			{
+				return RigNode->GetTemplate() != nullptr;
+			};
+			return false;
+		};
 		
-		if(PinA->PinType.PinSubCategoryObject == RigVMTypeUtils::GetWildCardCPPTypeObject())
+		if(PinA->PinType.PinSubCategoryObject == RigVMTypeUtils::GetWildCardCPPTypeObject() || IsTemplateNodePin(PinA))
 		{
 			if(IsPinCompatibleWithType(PinA, PinB->PinType))
 			{
@@ -1576,7 +1585,7 @@ bool UControlRigGraphSchema::ArePinsCompatible(const UEdGraphPin* PinA, const UE
 				return true;
 			}
 		}
- 		else if(PinB->PinType.PinSubCategoryObject == RigVMTypeUtils::GetWildCardCPPTypeObject())
+ 		if(PinB->PinType.PinSubCategoryObject == RigVMTypeUtils::GetWildCardCPPTypeObject() || IsTemplateNodePin(PinB))
 		{
  			if(IsPinCompatibleWithType(PinB, PinA->PinType))
  			{
diff --git a/Engine/Plugins/Animation/ControlRig/Source/ControlRigDeveloper/Public/ControlRigBlueprint.h b/Engine/Plugins/Animation/ControlRig/Source/ControlRigDeveloper/Public/ControlRigBlueprint.h
--- a/Engine/Plugins/Animation/ControlRig/Source/ControlRigDeveloper/Public/ControlRigBlueprint.h
+++ b/Engine/Plugins/Animation/ControlRig/Source/ControlRigDeveloper/Public/ControlRigBlueprint.h
@@ -634,6 +634,8 @@ public:
 
 	FControlRigOnBulkEditDialogRequestedDelegate& OnRequestBulkEditDialog() { return RequestBulkEditDialog; }
 
+	FRigVMController_RequestJumpToHyperlinkDelegate& OnRequestJumpToHyperlink() { return RequestJumpToHyperlink; };
+
 	FOnReportCompilerMessage& OnReportCompilerMessage() { return ReportCompilerMessageEvent; }
 	void BroadCastReportCompilerMessage(EMessageSeverity::Type InSeverity, UObject* InSubject, const FString& InMessage);
 
@@ -649,6 +651,7 @@ private:
 	FOnLocalizeFunctionDialogRequested RequestLocalizeFunctionDialog;
 	FOnReportCompilerMessage ReportCompilerMessageEvent;
 	FControlRigOnBulkEditDialogRequestedDelegate RequestBulkEditDialog;
+	FRigVMController_RequestJumpToHyperlinkDelegate RequestJumpToHyperlink;
 
 #endif
 
diff --git a/Engine/Plugins/Animation/ControlRig/Source/ControlRigDeveloper/Public/Graph/ControlRigGraphNode.h b/Engine/Plugins/Animation/ControlRig/Source/ControlRigDeveloper/Public/Graph/ControlRigGraphNode.h
--- a/Engine/Plugins/Animation/ControlRig/Source/ControlRigDeveloper/Public/Graph/ControlRigGraphNode.h
+++ b/Engine/Plugins/Animation/ControlRig/Source/ControlRigDeveloper/Public/Graph/ControlRigGraphNode.h
@@ -195,6 +195,8 @@ private:
 	FORCEINLINE int32 GetNodeTopologyVersion() const { return NodeTopologyVersion; }
 	int32 NodeTopologyVersion;
 
+	void ConfigurePin(UEdGraphPin* EdGraphPin, URigVMPin* ModelPin, bool bHidden, bool bConnectable);
+
 	FLinearColor CachedTitleColor;
 	FLinearColor CachedNodeColor;
 
diff --git a/Engine/Plugins/Animation/ControlRig/Source/ControlRigEditor/Private/ControlRigEditorModule.cpp b/Engine/Plugins/Animation/ControlRig/Source/ControlRigEditor/Private/ControlRigEditorModule.cpp
--- a/Engine/Plugins/Animation/ControlRig/Source/ControlRigEditor/Private/ControlRigEditorModule.cpp
+++ b/Engine/Plugins/Animation/ControlRig/Source/ControlRigEditor/Private/ControlRigEditorModule.cpp
@@ -110,6 +110,7 @@
 #include "SBakeToControlRigDialog.h"
 #include "ControlRig/Private/Units/Execution/RigUnit_InverseExecution.h"
 #include "Graph/SControlRigGraphPinVariableBinding.h"
+#include "Graph/SControlRigGraphChangePinType.h"
 #include "AssetTypeActions_ControlRigPose.h"
 #include "ControlRigBlueprintFactory.h"
 #include "ControlRigPythonLogDetails.h"
@@ -1286,6 +1287,33 @@ void FControlRigEditorModule::GetContextMenuActions(const UControlRigGraphSchema
 						));
 					}
 
+					if (URigVMTemplateNode* TemplateNode = Cast<URigVMTemplateNode>(ModelPin->GetNode()))
+					{
+						FToolMenuSection& TemplatesSection = Menu->AddSection("EdGraphSchemaTemplates", LOCTEXT("TemplatesHeader", "Templates"));
+
+						if(ModelPin->IsRootPin())
+						{
+							TSharedRef<SControlRigChangePinType> ChangePinTypeWidget =
+							SNew(SControlRigChangePinType)
+							.Blueprint(RigBlueprint)
+							.ModelPins({ModelPin});
+
+							TemplatesSection.AddEntry(FToolMenuEntry::InitWidget("ChangePinTypeWidget", ChangePinTypeWidget, FText(), true));
+						}
+						
+						TemplatesSection.AddMenuEntry(
+							"Unresolve Template Node",
+							LOCTEXT("UnresolveTemplateNode", "Unresolve Template Node"),
+							LOCTEXT("UnresolveTemplateNode_Tooltip", "Removes any type information from the template node"),
+							FSlateIcon(),
+							FUIAction(FExecuteAction::CreateLambda([Controller, ModelPin]() {
+								const TArray<FName> Nodes = ModelPin->GetGraph()->GetSelectNodes();
+								Controller->UnresolveTemplateNodes(Nodes);
+							})
+						));
+					}
+
+
 					if (ModelPin->GetDirection() == ERigVMPinDirection::Input &&
 							bIsEditablePin)
 					{
@@ -2160,6 +2188,21 @@ void FControlRigEditorModule::GetContextMenuActions(const UControlRigGraphSchema
 							));
 					}
 
+					if (URigVMTemplateNode* TemplateNode = Cast<URigVMTemplateNode>(RigNode->GetModelNode()))
+					{
+						FToolMenuSection& TemplatesSection = Menu->AddSection("EdGraphSchemaTemplates", LOCTEXT("TemplatesHeader", "Templates"));
+						TemplatesSection.AddMenuEntry(
+							"Unresolve Template Node",
+							LOCTEXT("UnresolveTemplateNode", "Unresolve Template Node"),
+							LOCTEXT("UnresolveTemplateNode_Tooltip", "Removes any type information from the template node"),
+							FSlateIcon(),
+							FUIAction(FExecuteAction::CreateLambda([Controller, Model]() {
+								const TArray<FName> Nodes = Model->GetSelectNodes();
+								Controller->UnresolveTemplateNodes(Nodes);
+							})
+						));
+					}
+
 					if (URigVMLibraryNode* LibraryNode = Cast<URigVMLibraryNode>(RigNode->GetModelNode()))
 					
 					{
diff --git a/Engine/Plugins/Animation/ControlRig/Source/ControlRigEditor/Private/Editor/ControlRigEditor.cpp b/Engine/Plugins/Animation/ControlRig/Source/ControlRigEditor/Private/Editor/ControlRigEditor.cpp
--- a/Engine/Plugins/Animation/ControlRig/Source/ControlRigEditor/Private/Editor/ControlRigEditor.cpp
+++ b/Engine/Plugins/Animation/ControlRig/Source/ControlRigEditor/Private/Editor/ControlRigEditor.cpp
@@ -246,6 +246,7 @@ FControlRigEditor::~FControlRigEditor()
 		RigBlueprint->OnGraphImported().RemoveAll(this);
 		RigBlueprint->OnRequestLocalizeFunctionDialog().RemoveAll(this);
 		RigBlueprint->OnRequestBulkEditDialog().Unbind();
+		RigBlueprint->OnRequestJumpToHyperlink().Unbind();
 		RigBlueprint->OnReportCompilerMessage().RemoveAll(this);
 
 #if WITH_EDITOR
@@ -505,6 +506,7 @@ void FControlRigEditor::InitControlRigEditor(const EToolkitMode::Type Mode, cons
 		InControlRigBlueprint->OnGraphImported().AddSP(this, &FControlRigEditor::OnGraphImported);
 		InControlRigBlueprint->OnRequestLocalizeFunctionDialog().AddSP(this, &FControlRigEditor::OnRequestLocalizeFunctionDialog);
 		InControlRigBlueprint->OnRequestBulkEditDialog().BindSP(this, &FControlRigEditor::OnRequestBulkEditDialog);
+		InControlRigBlueprint->OnRequestJumpToHyperlink().BindSP(this, &FControlRigEditor::HandleJumpToHyperlink);
 	}
 
 	UpdateStaleWatchedPins();
@@ -2764,7 +2766,10 @@ void FControlRigEditor::HandleModifiedEvent(ERigVMGraphNotifType InNotifType, UR
 									if(TargetProperty)
 									{
 										uint8* PropertyStorage = TargetProperty->ContainerPtrToValuePtr<uint8>(WrapperObject);
-										TargetProperty->ImportText_Direct(*DefaultValue, PropertyStorage, nullptr, PPF_None, nullptr);
+
+										// we are ok with not reacting to errors here
+										FRigVMPinDefaultValueImportErrorContext ErrorPipe;										
+										TargetProperty->ImportText_Direct(*DefaultValue, PropertyStorage, nullptr, PPF_None, &ErrorPipe);
 									}
 								}
 							}
@@ -2777,6 +2782,7 @@ void FControlRigEditor::HandleModifiedEvent(ERigVMGraphNotifType InNotifType, UR
 		}
 		case ERigVMGraphNotifType::PinArraySizeChanged:
 		case ERigVMGraphNotifType::PinBoundVariableChanged:
+		case ERigVMGraphNotifType::PinTypeChanged:
 		{
 			URigVMPin* Pin = Cast<URigVMPin>(InSubject);
 
@@ -4669,6 +4675,63 @@ FRigVMController_BulkEditResult FControlRigEditor::OnRequestBulkEditDialog(UCont
 	return Result;
 }
 
+void FControlRigEditor::HandleJumpToHyperlink(const UObject* InSubject)
+{
+	UControlRigBlueprint* RigBlueprint = GetControlRigBlueprint();
+	if(RigBlueprint == nullptr)
+	{
+		return;
+	}
+
+	const URigVMGraph* GraphToJumpTo = nullptr;
+	const URigVMNode* NodeToJumpTo = nullptr;
+	const URigVMPin* PinToJumpTo = nullptr;
+	if(const URigVMNode* Node = Cast<URigVMNode>(InSubject))
+	{
+		GraphToJumpTo = Node->GetGraph();
+		NodeToJumpTo = Node;
+	}
+	else if(const URigVMPin* Pin = Cast<URigVMPin>(InSubject))
+	{
+		GraphToJumpTo = Pin->GetGraph();
+		NodeToJumpTo = Pin->GetNode();
+		PinToJumpTo = Pin;
+	}
+	else if(const URigVMLink* Link = Cast<URigVMLink>(InSubject))
+	{
+		GraphToJumpTo = Link->GetGraph();
+		if(const URigVMPin* TargetPin = ((URigVMLink*)Link)->GetTargetPin())
+		{
+			NodeToJumpTo = TargetPin->GetNode();
+			PinToJumpTo = TargetPin;
+		}
+	}
+
+	if(GraphToJumpTo)
+	{
+		if(UControlRigGraph* EdGraph = Cast<UControlRigGraph>(RigBlueprint->GetEdGraph(NodeToJumpTo->GetGraph())))
+		{
+			if(NodeToJumpTo)
+			{
+				if(const UControlRigGraphNode* EdGraphNode = Cast<UControlRigGraphNode>(EdGraph->FindNodeForModelNodeName(NodeToJumpTo->GetFName())))
+				{
+					if(PinToJumpTo)
+					{
+						if(const UEdGraphPin* EdGraphPin = EdGraphNode->FindPin(PinToJumpTo->GetSegmentPath(true)))
+						{
+							JumpToPin(EdGraphPin);
+							return;
+						}
+					}
+					JumpToNode(EdGraphNode);
+					return;
+				}
+			}
+			JumpToHyperlink(EdGraph);
+		}
+	}
+}
+
 void FControlRigEditor::UpdateDefaultValueForVariable(FBPVariableDescription& InVariable, bool bUseCDO)
 {
 	if (UControlRigBlueprint* ControlRigBP = GetControlRigBlueprint())
diff --git a/Engine/Plugins/Animation/ControlRig/Source/ControlRigEditor/Private/Editor/ControlRigEditor.h b/Engine/Plugins/Animation/ControlRig/Source/ControlRigEditor/Private/Editor/ControlRigEditor.h
--- a/Engine/Plugins/Animation/ControlRig/Source/ControlRigEditor/Private/Editor/ControlRigEditor.h
+++ b/Engine/Plugins/Animation/ControlRig/Source/ControlRigEditor/Private/Editor/ControlRigEditor.h
@@ -438,6 +438,7 @@ protected:
 	void OnWrappedPropertyChangedChainEvent(UDetailsViewWrapperObject* InWrapperObject, const FString& InPropertyPath, FPropertyChangedChainEvent& InPropertyChangedChainEvent);
 	void OnRequestLocalizeFunctionDialog(URigVMLibraryNode* InFunction, UControlRigBlueprint* InTargetBlueprint, bool bForce);
 	FRigVMController_BulkEditResult OnRequestBulkEditDialog(UControlRigBlueprint* InBlueprint, URigVMController* InController, URigVMLibraryNode* InFunction, ERigVMControllerBulkEditType InEditType);
+	void HandleJumpToHyperlink(const UObject* InSubject);
 	void UpdateDefaultValueForVariable(FBPVariableDescription& InVariable, bool bUseCDO);
 
 	URigVMGraph* GetFocusedModel() const;
diff --git a/Engine/Plugins/Animation/ControlRig/Source/ControlRigEditor/Private/Graph/SControlRigGraphChangePinType.cpp b/Engine/Plugins/Animation/ControlRig/Source/ControlRigEditor/Private/Graph/SControlRigGraphChangePinType.cpp
new file mode 100644
--- /dev/null
+++ b/Engine/Plugins/Animation/ControlRig/Source/ControlRigEditor/Private/Graph/SControlRigGraphChangePinType.cpp
@@ -0,0 +1,359 @@
+// Copyright Epic Games, Inc. All Rights Reserved.
+
+
+#include "Graph/SControlRigGraphChangePinType.h"
+#include "Widgets/Layout/SSpacer.h"
+#include "Widgets/Images/SImage.h"
+#include "Framework/MultiBox/MultiBoxExtender.h"
+#include "Framework/MultiBox/MultiBoxBuilder.h"
+#include "DetailLayoutBuilder.h"
+#include "RigVMModel/RigVMController.h"
+#include "Graph/ControlRigGraphSchema.h"
+#include "Kismet2/BlueprintEditorUtils.h"
+
+#define LOCTEXT_NAMESPACE "SControlRigGraphChangePinType"
+
+static const FText ControlRigChangePinTypeMultipleValues = LOCTEXT("MultipleValues", "Multiple Values");
+
+void SControlRigChangePinType::Construct(const FArguments& InArgs)
+{
+	this->ModelPins = InArgs._ModelPins;
+	this->Blueprint = InArgs._Blueprint;
+
+	IPropertyAccessEditor& PropertyAccessEditor = IModularFeatures::Get().GetModularFeature<IPropertyAccessEditor>("PropertyAccessEditor");
+
+	BindingArgs.CurrentBindingText.BindRaw(this, &SControlRigChangePinType::GetBindingText);
+	BindingArgs.CurrentBindingImage.BindRaw(this, &SControlRigChangePinType::GetBindingImage);
+	BindingArgs.CurrentBindingColor.BindRaw(this, &SControlRigChangePinType::GetBindingColor);
+
+	BindingArgs.OnCanBindProperty.BindLambda([](const FProperty* InProperty) -> bool { return true; });
+	BindingArgs.OnCanBindToClass.BindLambda([](UClass* InClass) -> bool { return false; });
+
+	BindingArgs.bGeneratePureBindings = true;
+	BindingArgs.bAllowNewBindings = true;
+	BindingArgs.bAllowArrayElementBindings = false;
+	BindingArgs.bAllowStructMemberBindings = false;
+	BindingArgs.bAllowUObjectFunctions = false;
+
+	BindingArgs.MenuExtender = MakeShareable(new FExtender);
+	BindingArgs.MenuExtender->AddMenuExtension(
+		"Properties",
+		EExtensionHook::After,
+		nullptr,
+		FMenuExtensionDelegate::CreateSP(this, &SControlRigChangePinType::FillPinTypeMenu));
+
+	this->ChildSlot
+	[
+		PropertyAccessEditor.MakePropertyBindingWidget(Blueprint, BindingArgs)
+	];
+}
+
+FText SControlRigChangePinType::GetBindingText(const FRigVMTemplateArgument::FType& InType)
+{
+	if(UObject* CPPTypeObject = InType.CPPTypeObject)
+	{
+		if(const UScriptStruct* ScriptStruct = Cast<UScriptStruct>(CPPTypeObject))
+		{
+			return ScriptStruct->GetDisplayNameText();
+		}
+		if(const UEnum* Enum = Cast<UEnum>(CPPTypeObject))
+		{
+			return Enum->GetDisplayNameText();
+		}
+	}
+	else
+	{
+		FString CPPType = InType.CPPType;
+		if(RigVMTypeUtils::IsArrayType(CPPType))
+		{
+			CPPType = RigVMTypeUtils::BaseTypeFromArrayType(CPPType);
+		}
+
+		static const FText BoolLabel = LOCTEXT("BoolLabel", "Boolean");
+		static const FText FloatLabel = LOCTEXT("FloatLabel", "Float");
+		static const FText Int32Label = LOCTEXT("Int32Label", "Integer");
+		static const FText FNameLabel = LOCTEXT("FNameLabel", "Name");
+		static const FText FStringLabel = LOCTEXT("FStringLabel", "String");
+
+		if(CPPType == RigVMTypeUtils::BoolType)
+		{
+			return BoolLabel;
+		}
+		if(CPPType == RigVMTypeUtils::FloatType || CPPType == RigVMTypeUtils::DoubleType)
+		{
+			return FloatLabel;
+		}
+		if(CPPType == RigVMTypeUtils::Int32Type)
+		{
+			return Int32Label;
+		}
+		if(CPPType == RigVMTypeUtils::FNameType)
+		{
+			return FNameLabel;
+		}
+		if(CPPType == RigVMTypeUtils::FStringType)
+		{
+			return FStringLabel;
+		}
+
+		return FText::FromString(CPPType);
+	}
+	return FText();
+}
+
+FText SControlRigChangePinType::GetBindingText(URigVMPin* ModelPin) const
+{
+	if (ModelPin)
+	{
+		return GetBindingText(ModelPin->GetTemplateArgumentType());
+	}
+	return FText();
+}
+
+FText SControlRigChangePinType::GetBindingText() const
+{
+	if (ModelPins.Num() > 0)
+	{
+		const FText FirstText = GetBindingText(ModelPins[0]);
+		for(int32 Index = 1; Index < ModelPins.Num(); Index++)
+		{
+			if(!GetBindingText(ModelPins[Index]).EqualTo(FirstText))
+			{
+				return ControlRigChangePinTypeMultipleValues;
+			}
+		}
+		return FirstText;
+	}
+	return FText();
+}
+
+const FSlateBrush* SControlRigChangePinType::GetBindingImage() const
+{
+	static FName TypeIcon(TEXT("Kismet.VariableList.TypeIcon"));
+	static FName ArrayTypeIcon(TEXT("Kismet.VariableList.ArrayTypeIcon"));
+
+	if (ModelPins.Num() > 0)
+	{
+		if(ModelPins[0]->IsArray())
+		{
+			return FEditorStyle::GetBrush(ArrayTypeIcon);
+		}
+	}
+	return FEditorStyle::GetBrush(TypeIcon);
+}
+
+FLinearColor SControlRigChangePinType::GetBindingColor() const
+{
+	if(ModelPins.Num() > 0)
+	{
+		const FRigVMTemplateArgument::FType Type = ModelPins[0]->GetTemplateArgumentType();
+		const FEdGraphPinType PinType = RigVMTypeUtils::PinTypeFromCPPType(Type.CPPType, Type.CPPTypeObject);
+		const UControlRigGraphSchema* Schema = GetDefault<UControlRigGraphSchema>();
+		return Schema->GetPinTypeColor(PinType);
+	}
+	return FLinearColor::White;
+}
+
+void SControlRigChangePinType::FillPinTypeMenu(FMenuBuilder& MenuBuilder)
+{
+	if(ModelPins.Num() == 0)
+	{
+		return;
+	}
+
+	URigVMGraph* Model = ModelPins[0]->GetGraph();
+	if(Model == nullptr)
+	{
+		return;
+	}
+
+	struct FArgumentInfo
+	{
+		static FArgumentInfo Make(const FProperty* InProperty)
+		{
+			FArgumentInfo Info;
+			Info.Property = InProperty;
+			return Info;
+		}
+
+		const FProperty* Property;
+	};
+
+	typedef TPair<FRigVMTemplateArgument::FType, FArgumentInfo> FTypePair;
+	TMap<FRigVMTemplateArgument::FType, FArgumentInfo> Types;
+	for(URigVMPin* ModelPin : ModelPins)
+	{
+		if(!ModelPin->IsRootPin())
+		{
+			continue;
+		}
+		
+		if(const URigVMTemplateNode* TemplateNode = Cast<URigVMTemplateNode>(ModelPin->GetNode()))
+		{
+			if(const FRigVMTemplate* Template = TemplateNode->GetTemplate())
+			{
+				const FName ArgumentName = ModelPin->GetFName();
+				if(const FRigVMTemplateArgument* Argument = Template->FindArgument(ArgumentName))
+				{
+					const TArray<FRigVMTemplateArgument::FType>& AllArgumentTypes = Argument->GetTypes();
+					const TArray<FRigVMTemplateArgument::FType>& FilteredArgumentTypes = Argument->GetSupportedTypes(TemplateNode->GetResolvedPermutationIndices());
+					
+					for(int32 PermutationIndex = 0; PermutationIndex < Template->NumPermutations(); PermutationIndex++)
+					{
+						FRigVMTemplateArgument::FType ArgumentType = AllArgumentTypes[PermutationIndex];
+						if(!FilteredArgumentTypes.Contains(ArgumentType))
+						{
+							continue;
+						}
+
+						if(ArgumentType.CPPType == RigVMTypeUtils::FloatType)
+						{
+							ArgumentType.CPPType = RigVMTypeUtils::DoubleType;
+						}
+						else if(ArgumentType.CPPType == RigVMTypeUtils::FloatArrayType)
+						{
+							ArgumentType.CPPType = RigVMTypeUtils::DoubleArrayType;
+						}
+						
+						if(Types.Contains(ArgumentType))
+						{
+							continue;
+						}
+						
+						if(UScriptStruct* FunctionStruct = Template->GetPermutation(PermutationIndex)->Struct)
+						{
+							if(const FProperty* Property = FunctionStruct->FindPropertyByName(ArgumentName))
+							{
+								Types.Add(ArgumentType, FArgumentInfo::Make(Property));
+							}
+						}
+					}
+				}
+			}
+		}
+	}
+
+	// sort the types and put them into an array
+	static const TArray<FString> SortOrder = {
+		TEXT("bool"),
+		TEXT("int32"),
+		TEXT("float"),
+		TEXT("double"),
+		TEXT("FName"),
+		TEXT("FString"),
+		TEXT("FVector"),
+		TEXT("FRotator"),
+		TEXT("FQuat"),
+		TEXT("FTransform"),
+		TEXT("FEulerTransform"),
+		TEXT("FMatrix")
+	};
+	TArray<FTypePair> SortedTypes;
+	for(const FTypePair& Pair : Types)
+	{
+		SortedTypes.Add(Pair);
+	}
+	SortedTypes.Sort([](const FTypePair& A, const FTypePair& B) -> bool
+	{
+		const FString BaseTypeA = A.Key.GetBaseCPPType();
+		const FString BaseTypeB = B.Key.GetBaseCPPType();
+
+		const int32 IndexA = SortOrder.Find(BaseTypeA);  
+		const int32 IndexB = SortOrder.Find(BaseTypeB);
+		if((IndexA == INDEX_NONE) != (IndexB == INDEX_NONE))
+		{
+			return IndexA != INDEX_NONE;
+		}
+		else if((IndexA != INDEX_NONE) && (IndexB != INDEX_NONE))
+		{
+			return IndexA < IndexB;
+		}
+
+		return BaseTypeA.Compare(BaseTypeB) < 0;
+	});
+	
+	MenuBuilder.BeginSection("PinTypes", LOCTEXT("PinTypes", "Pin Types"));
+	{
+		static FName PropertyIcon(TEXT("Kismet.VariableList.TypeIcon"));
+		const UControlRigGraphSchema* Schema = GetDefault<UControlRigGraphSchema>();
+
+		for(int32 TypeIndex=0; TypeIndex < SortedTypes.Num(); TypeIndex++)
+		{
+			const FRigVMTemplateArgument::FType& Type = SortedTypes[TypeIndex].Key;
+			const FProperty* Property = SortedTypes[TypeIndex].Value.Property;
+			const FEdGraphPinType PinType = RigVMTypeUtils::PinTypeFromCPPType(Type.CPPType, Type.CPPTypeObject);
+
+			MenuBuilder.AddMenuEntry(
+				FUIAction(FExecuteAction::CreateSP(this, &SControlRigChangePinType::HandlePinTypeChanged, Type)),
+				SNew(SHorizontalBox)
+					+SHorizontalBox::Slot()
+					.AutoWidth()
+					[
+						SNew(SSpacer)
+						.Size(FVector2D(18.0f, 0.0f))
+					]
+					+SHorizontalBox::Slot()
+					.AutoWidth()
+					.VAlign(VAlign_Center)
+					.Padding(1.0f, 0.0f)
+					[
+						SNew(SImage)
+						.Image(FBlueprintEditorUtils::GetIconFromPin(PinType, true))
+						.ColorAndOpacity(Schema->GetPinTypeColor(PinType))
+					]
+					+SHorizontalBox::Slot()
+					.AutoWidth()
+					.VAlign(VAlign_Center)
+					.Padding(4.0f, 0.0f)
+					[
+						SNew(STextBlock)
+						.Text(GetBindingText(Type))
+					]);
+		}
+	}
+	MenuBuilder.EndSection(); // Local Variables
+}
+
+void SControlRigChangePinType::HandlePinTypeChanged(FRigVMTemplateArgument::FType InType)
+{
+	if(ModelPins.IsEmpty() || (Blueprint == nullptr))
+	{
+		return;
+	}
+
+	for(URigVMPin* ModelPin : ModelPins)
+	{
+		URigVMGraph* Model = ModelPin->GetGraph();
+		if(Model == nullptr)
+		{
+			continue;
+		}
+
+		URigVMController* Controller = Blueprint->GetOrCreateController(Model);
+		if(Controller == nullptr)
+		{
+			continue;
+		}
+
+		Controller->ResolveWildCardPin(ModelPin->GetPinPath(), InType.CPPType, InType.GetCPPTypeObjectPath());
+	}
+}
+
+/////////////////////////////////////////////////////////////////////////////////////////////////////////////
+
+void SControlRigGraphChangePinType::Construct(const FArguments& InArgs, UEdGraphPin* InGraphPinObj)
+{
+	this->ModelPins = InArgs._ModelPins;
+	this->Blueprint = InArgs._Blueprint;
+
+	SGraphPin::Construct(SGraphPin::FArguments(), InGraphPinObj);
+}
+
+TSharedRef<SWidget>	SControlRigGraphChangePinType::GetDefaultValueWidget()
+{
+	return SNew(SControlRigChangePinType)
+		.Blueprint(Blueprint)
+		.ModelPins(ModelPins);
+}
+
+#undef LOCTEXT_NAMESPACE
\ No newline at end of file
diff --git a/Engine/Plugins/Animation/ControlRig/Source/ControlRigEditor/Private/Graph/SControlRigGraphChangePinType.h b/Engine/Plugins/Animation/ControlRig/Source/ControlRigEditor/Private/Graph/SControlRigGraphChangePinType.h
new file mode 100644
--- /dev/null
+++ b/Engine/Plugins/Animation/ControlRig/Source/ControlRigEditor/Private/Graph/SControlRigGraphChangePinType.h
@@ -0,0 +1,68 @@
+// Copyright Epic Games, Inc. All Rights Reserved.
+
+#pragma once
+
+#include "CoreMinimal.h"
+#include "Widgets/DeclarativeSyntaxSupport.h"
+#include "Widgets/SWidget.h"
+#include "SGraphPin.h"
+#include "RigVMModel/RigVMPin.h"
+#include "IPropertyAccessEditor.h"
+#include "ControlRigBlueprint.h"
+
+class SControlRigChangePinType : public SCompoundWidget
+{
+public:
+
+	SLATE_BEGIN_ARGS(SControlRigChangePinType)
+	: _ModelPins()
+    , _Blueprint(nullptr)
+	{}
+
+		SLATE_ARGUMENT(TArray<URigVMPin*>, ModelPins)
+		SLATE_ARGUMENT(UControlRigBlueprint*, Blueprint)
+
+	SLATE_END_ARGS()
+
+	void Construct(const FArguments& InArgs);
+
+protected:
+
+	static FText GetBindingText(const FRigVMTemplateArgument::FType& InType);
+	FText GetBindingText(URigVMPin* ModelPin) const;
+	FText GetBindingText() const;
+	const FSlateBrush* GetBindingImage() const;
+	FLinearColor GetBindingColor() const;
+	bool OnCanBindProperty(FProperty* InProperty) const;
+	bool OnCanBindToClass(UClass* InClass) const;
+	void OnAddBinding(FName InPropertyName, const TArray<FBindingChainElement>& InBindingChain);
+	void FillPinTypeMenu( FMenuBuilder& MenuBuilder );
+	void HandlePinTypeChanged(FRigVMTemplateArgument::FType InType);
+
+	TArray<URigVMPin*> ModelPins;
+	UControlRigBlueprint* Blueprint;
+	FPropertyBindingWidgetArgs BindingArgs;
+};
+
+class SControlRigGraphChangePinType : public SGraphPin
+{
+public:
+
+	SLATE_BEGIN_ARGS(SControlRigGraphChangePinType){}
+
+		SLATE_ARGUMENT(TArray<URigVMPin*>, ModelPins)
+		SLATE_ARGUMENT(UControlRigBlueprint*, Blueprint)
+
+	SLATE_END_ARGS()
+
+	void Construct(const FArguments& InArgs, UEdGraphPin* InGraphPinObj);
+
+protected:
+
+	//~ Begin SGraphPin Interface
+	virtual TSharedRef<SWidget>	GetDefaultValueWidget() override;
+	//~ End SGraphPin Interface
+
+	TArray<URigVMPin*> ModelPins;
+	UControlRigBlueprint* Blueprint;
+};
diff --git a/Engine/Plugins/Animation/ControlRig/Source/ControlRigEditor/Private/Graph/SControlRigGraphPinVariableBinding.cpp b/Engine/Plugins/Animation/ControlRig/Source/ControlRigEditor/Private/Graph/SControlRigGraphPinVariableBinding.cpp
--- a/Engine/Plugins/Animation/ControlRig/Source/ControlRigEditor/Private/Graph/SControlRigGraphPinVariableBinding.cpp
+++ b/Engine/Plugins/Animation/ControlRig/Source/ControlRigEditor/Private/Graph/SControlRigGraphPinVariableBinding.cpp
@@ -95,8 +95,17 @@ FText SControlRigVariableBinding::GetBindingText() const
 
 const FSlateBrush* SControlRigVariableBinding::GetBindingImage() const
 {
-	static FName PropertyIcon(TEXT("Kismet.Tabs.Variables"));
-	return FEditorStyle::GetBrush(PropertyIcon);
+	static FName TypeIcon(TEXT("Kismet.VariableList.TypeIcon"));
+	static FName ArrayTypeIcon(TEXT("Kismet.VariableList.ArrayTypeIcon"));
+
+	if (ModelPins.Num() > 0)
+	{
+		if(ModelPins[0]->IsArray())
+		{
+			return FEditorStyle::GetBrush(ArrayTypeIcon);
+		}
+	}
+	return FEditorStyle::GetBrush(TypeIcon);
 }
 
 FLinearColor SControlRigVariableBinding::GetBindingColor() const
diff --git a/Engine/Source/Developer/RigVMDeveloper/Private/RigVMDeveloperTypeUtils.cpp b/Engine/Source/Developer/RigVMDeveloper/Private/RigVMDeveloperTypeUtils.cpp
--- a/Engine/Source/Developer/RigVMDeveloper/Private/RigVMDeveloperTypeUtils.cpp
+++ b/Engine/Source/Developer/RigVMDeveloper/Private/RigVMDeveloperTypeUtils.cpp
@@ -1,6 +1,8 @@
 ﻿// Copyright Epic Games, Inc. All Rights Reserved.
 
 #include "RigVMDeveloperTypeUtils.h"
+
+#include "Internationalization/StringTableCore.h"
 #include "RigVMModel/RigVMVariableDescription.h"
 
 #if WITH_EDITOR
@@ -63,12 +65,12 @@ FRigVMExternalVariable RigVMTypeUtils::ExternalVariableFromPinType(const FName&
 
 	if (InPinType.PinCategory == UEdGraphSchema_K2::PC_Boolean)
 	{
-		ExternalVariable.TypeName = TEXT("bool");
+		ExternalVariable.TypeName = BoolTypeName;
 		ExternalVariable.Size = sizeof(bool);
 	}
 	else if (InPinType.PinCategory == UEdGraphSchema_K2::PC_Int)
 	{
-		ExternalVariable.TypeName = TEXT("int32");
+		ExternalVariable.TypeName = Int32TypeName;
 		ExternalVariable.Size = sizeof(int32);
 	}
 	else if (InPinType.PinCategory == UEdGraphSchema_K2::PC_Enum ||
@@ -81,7 +83,7 @@ FRigVMExternalVariable RigVMTypeUtils::ExternalVariableFromPinType(const FName&
 		}
 		else
 		{
-			ExternalVariable.TypeName = TEXT("uint8");
+			ExternalVariable.TypeName = UInt8TypeName;
 		}
 		ExternalVariable.Size = sizeof(uint8);
 	}
@@ -89,12 +91,12 @@ FRigVMExternalVariable RigVMTypeUtils::ExternalVariableFromPinType(const FName&
 	{
 		if (InPinType.PinSubCategory == UEdGraphSchema_K2::PC_Float)
 		{
-			ExternalVariable.TypeName = TEXT("float");
+			ExternalVariable.TypeName = FloatTypeName;
 			ExternalVariable.Size = sizeof(float);
 		}
 		else if (InPinType.PinSubCategory == UEdGraphSchema_K2::PC_Double)
 		{
-			ExternalVariable.TypeName = TEXT("double");
+			ExternalVariable.TypeName = DoubleTypeName;
 			ExternalVariable.Size = sizeof(double);
 		}
 		else
@@ -104,12 +106,12 @@ FRigVMExternalVariable RigVMTypeUtils::ExternalVariableFromPinType(const FName&
 	}
 	else if (InPinType.PinCategory == UEdGraphSchema_K2::PC_Name)
 	{
-		ExternalVariable.TypeName = TEXT("FName");
+		ExternalVariable.TypeName = FNameTypeName;
 		ExternalVariable.Size = sizeof(FName);
 	}
 	else if (InPinType.PinCategory == UEdGraphSchema_K2::PC_String)
 	{
-		ExternalVariable.TypeName = TEXT("FString");
+		ExternalVariable.TypeName = FStringTypeName;
 		ExternalVariable.Size = sizeof(FString);
 	}
 	else if (InPinType.PinCategory == UEdGraphSchema_K2::PC_Struct)
@@ -144,32 +146,32 @@ FRigVMExternalVariable RigVMTypeUtils::ExternalVariableFromCPPTypePath(const FNa
 		CPPTypePath = BaseTypeFromArrayType(CPPTypePath);
 	}
 
-	if (CPPTypePath == TEXT("bool"))
+	if (CPPTypePath == BoolType)
 	{
 		Variable.TypeName = *CPPTypePath;
 		Variable.Size = sizeof(bool);
 	}
-	else if (CPPTypePath == TEXT("float"))
+	else if (CPPTypePath == FloatType)
 	{
 		Variable.TypeName = *CPPTypePath;
 		Variable.Size = sizeof(float);
 	}
-	else if (CPPTypePath == TEXT("double"))
+	else if (CPPTypePath == DoubleType)
 	{
 		Variable.TypeName = *CPPTypePath;
 		Variable.Size = sizeof(double);
 	}
-	else if (CPPTypePath == TEXT("int32"))
+	else if (CPPTypePath == Int32Type)
 	{
 		Variable.TypeName = *CPPTypePath;
 		Variable.Size = sizeof(int32);
 	}
-	else if (CPPTypePath == TEXT("FString"))
+	else if (CPPTypePath == FStringType)
 	{
 		Variable.TypeName = *CPPTypePath;
 		Variable.Size = sizeof(FString);
 	}
-	else if (CPPTypePath == TEXT("FName"))
+	else if (CPPTypePath == FNameType)
 	{
 		Variable.TypeName = *CPPTypePath;
 		Variable.Size = sizeof(FName);
@@ -213,32 +215,32 @@ FRigVMExternalVariable RigVMTypeUtils::ExternalVariableFromCPPType(const FName&
 		CPPType = BaseTypeFromArrayType(CPPType);
 	}
 
-	if (CPPType == TEXT("bool"))
+	if (CPPType == BoolType)
 	{
 		Variable.TypeName = *CPPType;
 		Variable.Size = sizeof(bool);
 	}
-	else if (CPPType == TEXT("float"))
+	else if (CPPType == FloatType)
 	{
 		Variable.TypeName = *CPPType;
 		Variable.Size = sizeof(float);
 	}
-	else if (CPPType == TEXT("double"))
+	else if (CPPType == DoubleType)
 	{
 		Variable.TypeName = *CPPType;
 		Variable.Size = sizeof(double);
 	}
-	else if (CPPType == TEXT("int32"))
+	else if (CPPType == Int32Type)
 	{
 		Variable.TypeName = *CPPType;
 		Variable.Size = sizeof(int32);
 	}
-	else if (CPPType == TEXT("FString"))
+	else if (CPPType == FStringType)
 	{
 		Variable.TypeName = *CPPType;
 		Variable.Size = sizeof(FString);
 	}
-	else if (CPPType == TEXT("FName"))
+	else if (CPPType == FNameType)
 	{
 		Variable.TypeName = *CPPType;
 		Variable.Size = sizeof(FName);
@@ -279,29 +281,29 @@ FEdGraphPinType RigVMTypeUtils::PinTypeFromCPPType(const FString& InCPPType, UOb
 		PinType.ContainerType = EPinContainerType::Array;
 	}
 
-	if (BaseCPPType == TEXT("bool"))
+	if (BaseCPPType == BoolType)
 	{
 		PinType.PinCategory = UEdGraphSchema_K2::PC_Boolean;
 	}
-	else if (BaseCPPType == TEXT("int32"))
+	else if (BaseCPPType == Int32Type)
 	{
 		PinType.PinCategory = UEdGraphSchema_K2::PC_Int;
 	}
-	else if (BaseCPPType == TEXT("float"))
+	else if (BaseCPPType == FloatType)
 	{
 		PinType.PinCategory = UEdGraphSchema_K2::PC_Real;
 		PinType.PinSubCategory = UEdGraphSchema_K2::PC_Float;
 	}
-	else if (BaseCPPType == TEXT("double"))
+	else if (BaseCPPType == DoubleType)
 	{
 		PinType.PinCategory = UEdGraphSchema_K2::PC_Real;
 		PinType.PinSubCategory = UEdGraphSchema_K2::PC_Double;
 	}
-	else if (BaseCPPType == TEXT("FName"))
+	else if (BaseCPPType == FNameType)
 	{
 		PinType.PinCategory = UEdGraphSchema_K2::PC_Name;
 	}
-	else if (BaseCPPType == TEXT("FString"))
+	else if (BaseCPPType == FStringType)
 	{
 		PinType.PinCategory = UEdGraphSchema_K2::PC_String;
 	}
@@ -325,29 +327,29 @@ FEdGraphPinType RigVMTypeUtils::PinTypeFromExternalVariable(const FRigVMExternal
 	PinType.ResetToDefaults();
 	PinType.PinCategory = NAME_None;
 
-	if (InExternalVariable.TypeName == TEXT("bool"))
+	if (InExternalVariable.TypeName == BoolTypeName)
 	{
 		PinType.PinCategory = UEdGraphSchema_K2::PC_Boolean;
 	}
-	else if (InExternalVariable.TypeName == TEXT("int32"))
+	else if (InExternalVariable.TypeName == Int32TypeName)
 	{
 		PinType.PinCategory = UEdGraphSchema_K2::PC_Int;
 	}
-	else if (InExternalVariable.TypeName == TEXT("float"))
+	else if (InExternalVariable.TypeName == FloatTypeName)
 	{
 		PinType.PinCategory = UEdGraphSchema_K2::PC_Real;
 		PinType.PinSubCategory = UEdGraphSchema_K2::PC_Float;
 	}
-	else if (InExternalVariable.TypeName == TEXT("double"))
+	else if (InExternalVariable.TypeName == DoubleTypeName)
 	{
 		PinType.PinCategory = UEdGraphSchema_K2::PC_Real;
 		PinType.PinSubCategory = UEdGraphSchema_K2::PC_Double;
 	}
-	else if (InExternalVariable.TypeName == TEXT("FName"))
+	else if (InExternalVariable.TypeName == FNameTypeName)
 	{
 		PinType.PinCategory = UEdGraphSchema_K2::PC_Name;
 	}
-	else if (InExternalVariable.TypeName == TEXT("FString"))
+	else if (InExternalVariable.TypeName == FStringTypeName)
 	{
 		PinType.PinCategory = UEdGraphSchema_K2::PC_String;
 	}
@@ -392,29 +394,29 @@ FEdGraphPinType RigVMTypeUtils::PinTypeFromRigVMVariableDescription(
 		PinType.ContainerType = EPinContainerType::None;
 	}
 
-	if (CurrentCPPType == TEXT("bool"))
+	if (CurrentCPPType == BoolType)
 	{
 		PinType.PinCategory = UEdGraphSchema_K2::PC_Boolean;
 	}
-	else if (CurrentCPPType == TEXT("int32"))
+	else if (CurrentCPPType == Int32Type)
 	{
 		PinType.PinCategory = UEdGraphSchema_K2::PC_Int;
 	}
-	else if (CurrentCPPType == TEXT("float"))
+	else if (CurrentCPPType == FloatType)
 	{
 		PinType.PinCategory = UEdGraphSchema_K2::PC_Real;
 		PinType.PinSubCategory = UEdGraphSchema_K2::PC_Float;
 	}
-	else if (CurrentCPPType == TEXT("double"))
+	else if (CurrentCPPType == DoubleType)
 	{
 		PinType.PinCategory = UEdGraphSchema_K2::PC_Real;
 		PinType.PinSubCategory = UEdGraphSchema_K2::PC_Double;
 	}
-	else if (CurrentCPPType == TEXT("FName"))
+	else if (CurrentCPPType == FNameType)
 	{
 		PinType.PinCategory = UEdGraphSchema_K2::PC_Name;
 	}
-	else if (CurrentCPPType == TEXT("FString"))
+	else if (CurrentCPPType == FStringType)
 	{
 		PinType.PinCategory = UEdGraphSchema_K2::PC_String;
 	}
@@ -547,21 +549,21 @@ bool RigVMTypeUtils::CPPTypeFromPinType(const FEdGraphPinType& InPinType, FStrin
 	*OutCPPTypeObject = nullptr;
 	if (InPinType.PinCategory == UEdGraphSchema_K2::PC_Boolean)
 	{
-		OutCPPType = Prefix + TEXT("bool") + Suffix;
+		OutCPPType = Prefix + BoolType + Suffix;
 	}
 	else if (InPinType.PinCategory == UEdGraphSchema_K2::PC_Int)
 	{
-		OutCPPType = Prefix + TEXT("int32") + Suffix;
+		OutCPPType = Prefix + Int32Type + Suffix;
 	}
 	else if (InPinType.PinCategory == UEdGraphSchema_K2::PC_Real)
 	{
 		if (InPinType.PinSubCategory == UEdGraphSchema_K2::PC_Float)
 		{
-			OutCPPType = Prefix + TEXT("float") + Suffix;
+			OutCPPType = Prefix + FloatType + Suffix;
 		}
 		else if (InPinType.PinSubCategory == UEdGraphSchema_K2::PC_Double)
 		{
-			OutCPPType = Prefix + TEXT("double") + Suffix;
+			OutCPPType = Prefix + DoubleType + Suffix;
 		}
 		else
 		{
@@ -570,19 +572,19 @@ bool RigVMTypeUtils::CPPTypeFromPinType(const FEdGraphPinType& InPinType, FStrin
 	}
 	else if (InPinType.PinCategory == UEdGraphSchema_K2::PC_Float)
 	{
-		OutCPPType = Prefix + TEXT("float") + Suffix;
+		OutCPPType = Prefix + FloatType + Suffix;
 	}
 	else if (InPinType.PinCategory == UEdGraphSchema_K2::PC_Double)
 	{
-		OutCPPType = Prefix + TEXT("double") + Suffix;
+		OutCPPType = Prefix + DoubleType + Suffix;
 	}
 	else if (InPinType.PinCategory == UEdGraphSchema_K2::PC_Name)
 	{
-		OutCPPType = Prefix + TEXT("FName") + Suffix;
+		OutCPPType = Prefix + FNameType + Suffix;
 	}
 	else if (InPinType.PinCategory == UEdGraphSchema_K2::PC_String)
 	{
-		OutCPPType = Prefix + TEXT("FString") + Suffix;
+		OutCPPType = Prefix + FStringType + Suffix;
 	}
 	else if (InPinType.PinCategory == UEdGraphSchema_K2::PC_Struct)
 	{
@@ -612,7 +614,7 @@ bool RigVMTypeUtils::CPPTypeFromPinType(const FEdGraphPinType& InPinType, FStrin
 		}
 		else
 		{
-			OutCPPType = Prefix + TEXT("uint8") + Suffix;
+			OutCPPType = Prefix + UInt8Type + Suffix;
 		}
 	}
 	else
@@ -679,29 +681,29 @@ bool RigVMTypeUtils::CPPTypeFromExternalVariable(const FRigVMExternalVariable& I
 	}
 
 	*OutCPPTypeObject = nullptr;
-	if (InExternalVariable.TypeName == TEXT("bool"))
+	if (InExternalVariable.TypeName == BoolTypeName)
 	{
-		OutCPPType = Prefix + TEXT("bool") + Suffix;
+		OutCPPType = Prefix + BoolType + Suffix;
 	}
-	else if (InExternalVariable.TypeName == TEXT("int32"))
+	else if (InExternalVariable.TypeName == Int32TypeName)
 	{
-		OutCPPType = Prefix + TEXT("int32") + Suffix;
+		OutCPPType = Prefix + Int32Type + Suffix;
 	}
-	else if (InExternalVariable.TypeName == TEXT("float"))
+	else if (InExternalVariable.TypeName == FloatTypeName)
 	{
-		OutCPPType = Prefix + TEXT("float") + Suffix;
+		OutCPPType = Prefix + FloatType + Suffix;
 	}
-	else if (InExternalVariable.TypeName == TEXT("double"))
+	else if (InExternalVariable.TypeName == DoubleTypeName)
 	{
-		OutCPPType = Prefix + TEXT("double") + Suffix;
+		OutCPPType = Prefix + DoubleType + Suffix;
 	}
-	else if (InExternalVariable.TypeName == TEXT("FName"))
+	else if (InExternalVariable.TypeName == FNameTypeName)
 	{
-		OutCPPType = Prefix + TEXT("FName") + Suffix;
+		OutCPPType = Prefix + FNameType + Suffix;
 	}
-	else if (InExternalVariable.TypeName == TEXT("FString"))
+	else if (InExternalVariable.TypeName == FStringTypeName)
 	{
-		OutCPPType = Prefix + TEXT("FString") + Suffix;
+		OutCPPType = Prefix + FStringType + Suffix;
 	}
 	else if (UScriptStruct* Struct = Cast<UScriptStruct>(InExternalVariable.TypeObject))
 	{
@@ -741,8 +743,8 @@ bool RigVMTypeUtils::AreCompatible(const FEdGraphPinType& InTypeA, const FEdGrap
 	// Every other case is evaluated by UEdGraphSchema_K2::ArePinTypesCompatible
 	if (SubPinTypeA.ContainerType == SubPinTypeB.ContainerType)
 	{
-		if ((SubPinTypeA.PinCategory == TEXT("float") && SubPinTypeB.PinCategory == TEXT("double")) ||
-			(SubPinTypeA.PinCategory == TEXT("double") && SubPinTypeB.PinCategory == TEXT("float")))
+		if ((SubPinTypeA.PinCategory == FloatTypeName && SubPinTypeB.PinCategory == DoubleTypeName) ||
+			(SubPinTypeA.PinCategory == DoubleTypeName && SubPinTypeB.PinCategory == FloatTypeName))
 		{
 			return true;
 		}
diff --git a/Engine/Source/Developer/RigVMDeveloper/Private/RigVMModel/Nodes/RigVMTemplateNode.cpp b/Engine/Source/Developer/RigVMDeveloper/Private/RigVMModel/Nodes/RigVMTemplateNode.cpp
--- a/Engine/Source/Developer/RigVMDeveloper/Private/RigVMModel/Nodes/RigVMTemplateNode.cpp
+++ b/Engine/Source/Developer/RigVMDeveloper/Private/RigVMModel/Nodes/RigVMTemplateNode.cpp
@@ -2,6 +2,8 @@
 
 #include "RigVMModel/Nodes/RigVMTemplateNode.h"
 
+#include "RigVMModel/RigVMController.h"
+
 URigVMTemplateNode::URigVMTemplateNode()
 	: Super()
 	, TemplateNotation(NAME_None)
@@ -96,26 +98,40 @@ bool URigVMTemplateNode::SupportsType(const URigVMPin* InPin, const FString& InC
 	static const FString WildCardCPPType = RigVMTypeUtils::GetWildCardCPPType();
 	static const FString WildCardArrayCPPType = RigVMTypeUtils::ArrayTypeFromBaseType(WildCardCPPType);
 
+	const URigVMPin* RootPin = InPin->GetRootPin();
+
 	// we always support the unknown type
 	if(((InCPPType == WildCardCPPType) && !InPin->IsArray()) ||
 		((InCPPType == WildCardArrayCPPType) && InPin->IsArray()))
 	{
-		if(OutCPPType)
+		if(const FRigVMTemplate* Template = GetTemplate())
 		{
-			*OutCPPType = InCPPType;
+			if(const FRigVMTemplateArgument* Argument = Template->FindArgument(RootPin->GetFName()))
+			{
+				// support this only on non-singleton arguments
+				if(Argument->IsSingleton())
+				{
+					return false;
+				}
+				
+				if(OutCPPType)
+				{
+					*OutCPPType = InCPPType;
+				}
+				return true;
+			}
 		}
-		return true;
+		return false;
 	}
 	
 	FString CPPType = InCPPType;
-	const URigVMPin* RootPin = InPin->GetRootPin();
 
 	if(InPin->GetParentPin() == RootPin && RootPin->IsArray())
 	{
 		CPPType = RigVMTypeUtils::ArrayTypeFromBaseType(CPPType);
 	}
-	
-	if (RootPin->IsWildCard())
+
+	if (const FRigVMTemplate* Template = GetTemplate())
 	{
 		const FString CacheKey = RootPin->GetName() + TEXT("|") + CPPType;
 		if (const TPair<bool, FRigVMTemplateArgument::FType>* CachedResult = SupportedTypesCache.Find(CacheKey))
@@ -127,37 +143,45 @@ bool URigVMTemplateNode::SupportsType(const URigVMPin* InPin, const FString& InC
 			return CachedResult->Key;
 		}
 
-		if (const FRigVMTemplate* Template = GetTemplate())
-		{
-			FRigVMTemplate::FTypeMap Types;
+		FRigVMTemplate::FTypeMap Types;
 
-			for (URigVMPin* Pin : GetPins())
+		for (URigVMPin* Pin : GetPins())
+		{
+			if (Pin == RootPin)
 			{
-				if (Pin == RootPin)
-				{
-					continue;
-				}
-
-				if (Pin->IsWildCard())
-				{
-					continue;
-				}
-
-				FRigVMTemplateArgument::FType Type(Pin->GetCPPType(), Pin->GetCPPTypeObject());
-				Types.Add(Pin->GetFName(), Type);
+				continue;
 			}
+			if (Pin->IsWildCard())
+			{
+				continue;
+			}
+			Types.Add(Pin->GetFName(), Pin->GetTemplateArgumentType());
+		}
 
-			FRigVMTemplateArgument::FType ResolvedType;
-			bool Result = Template->ArgumentSupportsType(RootPin->GetFName(), CPPType, Types, &ResolvedType);
-			SupportedTypesCache.Add(CacheKey, TPair<bool, FRigVMTemplateArgument::FType>(Result, ResolvedType));
-			if(OutCPPType)
+		FRigVMTemplateArgument::FType ResolvedType;
+		bool bSupportsType = false;
+		if(RootPin->IsWildCard())
+		{
+			bSupportsType = Template->ArgumentSupportsType(RootPin->GetFName(), CPPType, Types, &ResolvedType);
+		}
+		else
+		{
+			FRigVMTemplate::FTypeMap ResolvedTypes;
+			bSupportsType = Template->ResolveArgument(RootPin->GetFName(), FRigVMTemplateArgument::FType(CPPType), ResolvedTypes);
+			if(bSupportsType)
 			{
-				*OutCPPType = ResolvedType.CPPType;
+				ResolvedType = ResolvedTypes.FindChecked(RootPin->GetFName());
 			}
-			return Result;
 		}
+	
+		SupportedTypesCache.Add(CacheKey, TPair<bool, FRigVMTemplateArgument::FType>(bSupportsType, ResolvedType));
+		if(OutCPPType)
+		{
+			*OutCPPType = ResolvedType.CPPType;
+		}
+		return bSupportsType;
 	}
-
+	
 	if(RootPin->GetCPPType() == CPPType)
 	{
 		if(OutCPPType)
@@ -169,23 +193,29 @@ bool URigVMTemplateNode::SupportsType(const URigVMPin* InPin, const FString& InC
 	return false;
 }
 
-TArray<int32> URigVMTemplateNode::GetResolvedPermutationIndices(FRigVMTemplate::FTypeMap* OutTypes) const
+bool URigVMTemplateNode::GetTypeMapForNewPinType(const URigVMPin* InPin, const FString& InCPPType, UObject* InCPPTypeObject,
+	FRigVMTemplate::FTypeMap& OutTypes) const
 {
-	if (const FRigVMTemplate* Template = GetTemplate())
+	check(InPin);
+	check(InPin->GetNode() == this);
+	check(InPin->IsRootPin());
+	
+	if(const FRigVMTemplate* Template = GetTemplate())
 	{
-		FRigVMTemplate::FTypeMap Types;
-
-		for (URigVMPin* Pin : GetPins())
-		{
-			if (Pin->IsWildCard())
-			{
-				continue;
-			}
+		OutTypes = GetResolvedTypes();
+		const FRigVMTemplateArgument::FType ExpectedType(InCPPType, InCPPTypeObject);
+		return GetTemplate()->ResolveArgument(InPin->GetFName(), ExpectedType, OutTypes);
+	}
 
-			FRigVMTemplateArgument::FType Type(Pin->GetCPPType(), Pin->GetCPPTypeObject());
-			Types.Add(Pin->GetFName(), Type);
-		}
+	OutTypes.Reset();
+	return false;
+}
 
+TArray<int32> URigVMTemplateNode::GetResolvedPermutationIndices(FRigVMTemplate::FTypeMap* OutTypes) const
+{
+	if (const FRigVMTemplate* Template = GetTemplate())
+	{
+		FRigVMTemplate::FTypeMap Types = GetResolvedTypes();
 		TArray<int32> PermutationIndices;
 		Template->Resolve(Types, PermutationIndices, false);
 
@@ -219,6 +249,20 @@ const FRigVMTemplate* URigVMTemplateNode::GetTemplate() const
 	return CachedTemplate;
 }
 
+FRigVMTemplate::FTypeMap URigVMTemplateNode::GetResolvedTypes() const
+{
+	FRigVMTemplate::FTypeMap Types;
+	for (URigVMPin* Pin : GetPins())
+	{
+		if (Pin->IsWildCard())
+		{
+			continue;
+		}
+		Types.Add(Pin->GetFName(), Pin->GetTemplateArgumentType());
+	}
+	return Types;
+}
+
 const FRigVMFunction* URigVMTemplateNode::GetResolvedFunction() const
 {
 	if(CachedFunction == nullptr)
@@ -245,6 +289,65 @@ bool URigVMTemplateNode::IsResolved() const
 	return GetScriptStruct() != nullptr;
 }
 
+bool URigVMTemplateNode::IsFullyUnresolved() const
+{
+	check(GetTemplate());
+
+	// all permutations are available means we haven't resolved any wildcard pin
+	return GetResolvedPermutations().Num() == GetTemplate()->NumPermutations();
+}
+
+TArray<UScriptStruct*> URigVMTemplateNode::GetSupportedUnitStructs() const
+{
+	if(const FRigVMTemplate* Template = GetTemplate())
+	{
+		return URigVMController::GetUnitStructsForTemplate(Template->GetNotation());
+	}
+	return TArray<UScriptStruct*>();
+}
+
+FString URigVMTemplateNode::GetInitialDefaultValueForPin(const FName& InRootPinName, const TArray<int32>& InPermutationIndices) const
+{
+	if(GetTemplate() == nullptr)
+	{
+		return FString();
+	}
+	
+	TArray<int32> PermutationIndices = InPermutationIndices;
+	if(PermutationIndices.IsEmpty())
+	{
+		PermutationIndices = GetResolvedPermutationIndices();
+	}
+
+	FString DefaultValue;
+
+	for(const int32 PermutationIndex : PermutationIndices)
+	{
+		const FRigVMFunction* Permutation = GetTemplate()->GetPermutation(PermutationIndex);
+		check(Permutation);
+
+		const TSharedPtr<FStructOnScope> StructOnScope = MakeShareable(new FStructOnScope(Permutation->Struct));
+		const FRigVMStruct* DefaultStruct = (const FRigVMStruct*)StructOnScope->GetStructMemory();
+
+		const FString NewDefaultValue = DefaultStruct->ExportToFullyQualifiedText(
+			Cast<UScriptStruct>(StructOnScope->GetStruct()), InRootPinName);
+
+		if(!NewDefaultValue.IsEmpty())
+		{
+			if(DefaultValue.IsEmpty())
+			{
+				DefaultValue = NewDefaultValue;
+			}
+			else if(!NewDefaultValue.Equals(DefaultValue))
+			{
+				return FString();
+			}
+		}
+	}
+	
+	return DefaultValue;
+}
+
 void URigVMTemplateNode::InvalidateCache()
 {
 	SupportedTypesCache.Reset();
diff --git a/Engine/Source/Developer/RigVMDeveloper/Private/RigVMModel/RigVMController.cpp b/Engine/Source/Developer/RigVMDeveloper/Private/RigVMModel/RigVMController.cpp
--- a/Engine/Source/Developer/RigVMDeveloper/Private/RigVMModel/RigVMController.cpp
+++ b/Engine/Source/Developer/RigVMDeveloper/Private/RigVMModel/RigVMController.cpp
@@ -63,6 +63,7 @@ URigVMController::URigVMController()
 	, bSuspendNotifications(false)
 	, bReportWarningsAndErrors(true)
 	, bIgnoreRerouteCompactnessChanges(false)
+	, UserLinkDirection(ERigVMPinDirection::Invalid)
 	, bIsRunningUnitTest(false)
 	, bIsFullyResolvingTemplateNode(false)
 {
@@ -72,7 +73,12 @@ URigVMController::URigVMController()
 URigVMController::URigVMController(const FObjectInitializer& ObjectInitializer)
 	: Super(ObjectInitializer)
 	, bValidatePinDefaults(true)
+	, bSuspendNotifications(false)
 	, bReportWarningsAndErrors(true)
+	, bIgnoreRerouteCompactnessChanges(false)
+	, UserLinkDirection(ERigVMPinDirection::Invalid)
+	, bIsRunningUnitTest(false)
+	, bIsFullyResolvingTemplateNode(false)
 {
 	ActionStack = CreateDefaultSubobject<URigVMActionStack>(TEXT("ActionStack"));
 	SetExecuteContextStruct(FRigVMExecuteContext::StaticStruct());
@@ -1505,8 +1511,158 @@ URigVMTemplateNode* URigVMController::ReplaceUnitNodeWithTemplateNode(const FNam
 	return TemplateNode;
 }
 
+bool URigVMController::UnresolveTemplateNodes(const TArray<FName>& InNodeNames, bool bSetupUndoRedo, bool bPrintPythonCommand)
+{
+	if (!IsValidGraph())
+	{
+		return false;
+	}
+
+	TArray<URigVMTemplateNode*> Nodes;
+	for (const FName& NodeName : InNodeNames)
+	{
+		if (URigVMTemplateNode* Node = Cast<URigVMTemplateNode>(GetGraph()->FindNodeByName(NodeName)))
+		{
+			Nodes.Add(Node);
+		}
+	}
+
+	if(UnresolveTemplateNodes(Nodes, bSetupUndoRedo))
+	{
+		if(bPrintPythonCommand)
+		{
+			const FString GraphName = GetSanitizedGraphName(GetGraph()->GetGraphName());
+			TArray<FString> NodeNames;
+			for(const FName& NodeName : InNodeNames)
+			{
+				NodeNames.Add(GetSanitizedNodeName(NodeName.ToString()));
+			}
+			const FString NodeNamesJoined = FString::Join(NodeNames, TEXT("','"));
+
+			// UnresolveTemplateNodes(const TArray<FName>& InNodeNames)
+			RigVMPythonUtils::Print(GetGraphOuterName(),
+					FString::Printf(TEXT("blueprint.get_controller_by_name('%s').unresolve_template_nodes(['%s'])"),
+					*GraphName,
+					*NodeNamesJoined));
+		}
+
+		return true;
+	}
+
+	return false;
+}
+
+bool URigVMController::UnresolveTemplateNodes(const TArray<URigVMTemplateNode*>& InNodes, bool bSetupUndoRedo)
+{
+	if (!IsValidGraph() || InNodes.IsEmpty())
+	{
+		return false;
+	}
+
+	// check if any of the nodes needs to be unresolved
+	if(Algo::FindByPredicate(InNodes,
+		[](const URigVMTemplateNode* Node) -> bool
+		{
+			return !Node->IsFullyUnresolved();
+		}) == nullptr)
+	{
+		return false;
+	}
+
+	FRigVMBaseAction Action;
+	if(bSetupUndoRedo)
+	{
+		Action.Title = TEXT("Unresolve nodes");
+		ActionStack->BeginAction(Action);
+	}
+
+	// find all links affecting the nodes to unresolve
+	TArray<URigVMNode*> Nodes;
+	Algo::Transform(InNodes, Nodes, [](URigVMTemplateNode* InNode) -> URigVMNode* { return InNode; });
+	TArray<TPair<FString, FString>> LinkedPaths = GetLinkedPinPaths(Nodes);
+
+	bool bChangedAnyPin = false;
+	for(URigVMTemplateNode* Node : InNodes)
+	{
+		if(Node->IsFullyUnresolved())
+		{
+			continue;
+		}
+
+		int32 PermutationIndex = INDEX_NONE;
+		FRigVMTemplate::FTypeMap Types;
+		Node->GetTemplate()->FullyResolve(Types, PermutationIndex);
+
+		for(const TPair<FName, FRigVMTemplateArgument::FType>& TypePair : Types)
+		{
+			const FName& PinName = TypePair.Key;
+			const FRigVMTemplateArgument::FType& ExpectedType = TypePair.Value;
+			
+			URigVMPin* Pin = Node->FindPin(PinName.ToString());
+			if(Pin == nullptr)
+			{
+				ReportErrorf(TEXT("Cannot find expected pin '%s' on Template Node '%s'"), *PinName.ToString(), *Node->GetNodePath());
+				if(bSetupUndoRedo)
+				{
+					ActionStack->CancelAction(Action);
+				}
+				return false;
+			}
+
+			if(ChangePinType(Pin, ExpectedType.CPPType, ExpectedType.CPPTypeObject, bSetupUndoRedo, false, true, true))
+			{
+				bChangedAnyPin = true;
+			}
+		}
+
+		ResolveTemplateNodeMetaData(Node, bSetupUndoRedo);
+	}
+
+	for(const TPair<FString, FString>& LinkedPath : LinkedPaths)
+	{
+		const FString& SourcePath = LinkedPath.Key;
+		const FString& TargetPath = LinkedPath.Value;
+
+		URigVMPin* SourcePin = GetGraph()->FindPin(SourcePath);
+		URigVMPin* TargetPin = GetGraph()->FindPin(TargetPath);
+
+		if(SourcePin == nullptr || TargetPin == nullptr)
+		{
+			ReportRemovedLink(SourcePath, TargetPath);
+			continue;
+		}
+
+		if(SourcePin->IsWildCard() != TargetPin->IsWildCard())
+		{
+			ReportRemovedLink(SourcePath, TargetPath);
+			continue;
+		}
+
+		// it's ok if this fails - we want to maintain the minimum set of links
+		if(!AddLink(SourcePin, TargetPin, bSetupUndoRedo))
+		{
+			ReportRemovedLink(SourcePath, TargetPath);
+		}
+	}
+	
+	if(bSetupUndoRedo)
+	{
+		if(bChangedAnyPin)
+		{
+			ActionStack->EndAction(Action);
+		}
+		else
+		{
+			ActionStack->CancelAction(Action);
+			return false;
+		}
+	}
+
+	return true;
+}
+
 TArray<URigVMNode*> URigVMController::UpgradeNodes(const TArray<FName>& InNodeNames, bool bRecursive, bool bSetupUndoRedo,
-	bool bPrintPythonCommand)
+                                                   bool bPrintPythonCommand)
 {
 	TArray<URigVMNode*> Nodes;
 	if (!IsValidGraph())
@@ -1587,17 +1743,7 @@ TArray<URigVMNode*> URigVMController::UpgradeNodes(const TArray<URigVMNode*>& In
 	}
 
 	// find all links affecting the nodes to upgrade
-	TArray<TPair<FString, FString>> LinkedPaths;
-	for(URigVMNode* Node : InNodes)
-	{
-		TArray<URigVMLink*> Links = Node->GetLinks();
-		for(URigVMLink* Link : Links)
-		{
-			TPair<FString, FString> LinkedPath(Link->GetSourcePin()->GetPinPath(), Link->GetTargetPin()->GetPinPath());
-			LinkedPaths.AddUnique(LinkedPath);
-		}
-	}
-
+	TArray<TPair<FString, FString>> LinkedPaths = GetLinkedPinPaths(InNodes);
 	for(const TPair<FString, FString>& LinkedPath : LinkedPaths)
 	{
 		if(!BreakLink(LinkedPath.Key, LinkedPath.Value, bSetupUndoRedo))
@@ -1652,7 +1798,7 @@ TArray<URigVMNode*> URigVMController::UpgradeNodes(const TArray<URigVMNode*>& In
 		
 		if(!AddLink(SourcePath, TargetPath, bSetupUndoRedo))
 		{
-			ReportWarningf(TEXT("Couldn't recreate link '%s' -> '%s'"), *SourcePath, *TargetPath);
+			ReportRemovedLink(SourcePath, TargetPath);
 		}
 	}
 
@@ -5051,6 +5197,31 @@ void URigVMController::RefreshFunctionPins(URigVMNode* InNode, bool bNotify)
 	}
 }
 
+void URigVMController::ReportRemovedLink(const FString& InSourcePinPath, const FString& InTargetPinPath)
+{
+	if(bSuspendNotifications)
+	{
+		return;
+	}
+	if(!IsValidGraph())
+	{
+		return;
+	}
+	
+	const URigVMPin* TargetPin = GetGraph()->FindPin(InTargetPinPath);
+	FString TargetNodeName, TargetSegmentPath;
+	if(!URigVMPin::SplitPinPathAtStart(InTargetPinPath, TargetNodeName, TargetSegmentPath))
+	{
+		TargetSegmentPath = InTargetPinPath;
+	}
+	
+	ReportWarningf(TEXT("Link '%s' -> '%s' was removed."), *InSourcePinPath, *InTargetPinPath);
+	SendUserFacingNotification(
+		FString::Printf(TEXT("Link to target pin '%s' was removed."), *TargetSegmentPath),
+		0.f, TargetPin, TEXT("MessageLog.Note")
+	);
+}
+
 bool URigVMController::RemoveNode(URigVMNode* InNode, bool bSetupUndoRedo, bool bRecursive, bool bPrintPythonCommand, bool bRelinkPins)
 {
 	if (!IsValidNodeForGraph(InNode))
@@ -7510,7 +7681,8 @@ bool URigVMController::PromotePinToVariable(URigVMPin* InPin, bool bCreateVariab
 	return false;
 }
 
-bool URigVMController::AddLink(const FString& InOutputPinPath, const FString& InInputPinPath, bool bSetupUndoRedo, bool bPrintPythonCommand)
+bool URigVMController::AddLink(const FString& InOutputPinPath, const FString& InInputPinPath, bool bSetupUndoRedo,
+	bool bPrintPythonCommand, ERigVMPinDirection InUserDirection)
 {
 	if(!IsValidGraph())
 	{
@@ -7520,6 +7692,7 @@ bool URigVMController::AddLink(const FString& InOutputPinPath, const FString& In
 	URigVMGraph* Graph = GetGraph();
 	check(Graph);
 
+
 	FString OutputPinPath = InOutputPinPath;
 	FString InputPinPath = InInputPinPath;
 
@@ -7548,7 +7721,7 @@ bool URigVMController::AddLink(const FString& InOutputPinPath, const FString& In
 	}
 	InputPin = InputPin->GetPinForLink();
 
-	const bool bSuccess = AddLink(OutputPin, InputPin, bSetupUndoRedo);
+	const bool bSuccess = AddLink(OutputPin, InputPin, bSetupUndoRedo, InUserDirection);
 	if (bSuccess && bPrintPythonCommand)
 	{
 		const FString GraphName = GetSanitizedGraphName(GetGraph()->GetGraphName());
@@ -7566,7 +7739,8 @@ bool URigVMController::AddLink(const FString& InOutputPinPath, const FString& In
 	return bSuccess;
 }
 
-bool URigVMController::AddLink(URigVMPin* OutputPin, URigVMPin* InputPin, bool bSetupUndoRedo)
+bool URigVMController::AddLink(URigVMPin* OutputPin, URigVMPin* InputPin, bool bSetupUndoRedo,
+	ERigVMPinDirection InUserDirection)
 {
 	if(OutputPin == nullptr)
 	{
@@ -7594,9 +7768,12 @@ bool URigVMController::AddLink(URigVMPin* OutputPin, URigVMPin* InputPin, bool b
 		return false;
 	}
 
+	TGuardValue<ERigVMPinDirection> UserLinkDirectionGuard(UserLinkDirection,
+	InUserDirection == ERigVMPinDirection::Invalid ? UserLinkDirection : InUserDirection);
+
 	{
 		FString FailureReason;
-		if (!Graph->CanLink(OutputPin, InputPin, &FailureReason, GetCurrentByteCode()))
+		if (!Graph->CanLink(OutputPin, InputPin, &FailureReason, GetCurrentByteCode(), UserLinkDirection))
 		{
 			if(OutputPin->IsExecuteContext() && InputPin->IsExecuteContext())
 			{
@@ -7622,19 +7799,35 @@ bool URigVMController::AddLink(URigVMPin* OutputPin, URigVMPin* InputPin, bool b
 		ActionStack->BeginAction(Action);
 	}
 
+	if (OutputPin->IsExecuteContext())
+	{
+		BreakAllLinks(OutputPin, false, bSetupUndoRedo);
+	}
+
+	BreakAllLinks(InputPin, true, bSetupUndoRedo);
+	if (bSetupUndoRedo)
+	{
+		BreakAllLinksRecursive(InputPin, true, true, bSetupUndoRedo);
+		BreakAllLinksRecursive(InputPin, true, false, bSetupUndoRedo);
+	}
+
 	// resolve types on the pins if needed
-	if(InputPin->GetCPPTypeObject() != OutputPin->GetCPPTypeObject())
+	if(InputPin->GetCPPTypeObject() != OutputPin->GetCPPTypeObject() ||
+		OutputPin->GetCPPType() != InputPin->GetCPPType())
 	{
-		if(InputPin->IsWildCard())
+		bool bOutputPinCanChangeType = OutputPin->IsWildCard();
+		bool bInputPinCanChangeType = InputPin->IsWildCard();
+
+		if(!bOutputPinCanChangeType && !bInputPinCanChangeType)
 		{
-			Notify(ERigVMGraphNotifType::InteractionBracketOpened, nullptr);
-			ResolveWildCardPin(InputPin, OutputPin, bSetupUndoRedo);
-			Notify(ERigVMGraphNotifType::InteractionBracketClosed, nullptr);
+			bInputPinCanChangeType = UserLinkDirection == ERigVMPinDirection::Output && InputPin->GetNode()->IsA<URigVMTemplateNode>(); 
+			bOutputPinCanChangeType = UserLinkDirection == ERigVMPinDirection::Input && OutputPin->GetNode()->IsA<URigVMTemplateNode>(); 
 		}
-		else if(OutputPin->IsWildCard())
+		
+		if(bOutputPinCanChangeType)
 		{
 			Notify(ERigVMGraphNotifType::InteractionBracketOpened, nullptr);
-			ResolveWildCardPin(OutputPin, InputPin, bSetupUndoRedo);
+			ResolveWildCardPin(OutputPin, InputPin, bSetupUndoRedo, true, true, true, true);
 			if(OutputPin->GetNode()->IsA<URigVMRerouteNode>())
 			{
 				SetPinDefaultValue(OutputPin, InputPin->GetDefaultValue(), true, bSetupUndoRedo, false, true);
@@ -7645,18 +7838,13 @@ bool URigVMController::AddLink(URigVMPin* OutputPin, URigVMPin* InputPin, bool b
 			}
 			Notify(ERigVMGraphNotifType::InteractionBracketClosed, nullptr);
 		}
-	}
+		else if(bInputPinCanChangeType)
+		{
+			Notify(ERigVMGraphNotifType::InteractionBracketOpened, nullptr);
+			ResolveWildCardPin(InputPin, OutputPin, bSetupUndoRedo, true, true, true, true);
+			Notify(ERigVMGraphNotifType::InteractionBracketClosed, nullptr);
+		}
 
-	if (OutputPin->IsExecuteContext())
-	{
-		BreakAllLinks(OutputPin, false, bSetupUndoRedo);
-	}
-
-	BreakAllLinks(InputPin, true, bSetupUndoRedo);
-	if (bSetupUndoRedo)
-	{
-		BreakAllLinksRecursive(InputPin, true, true, bSetupUndoRedo);
-		BreakAllLinksRecursive(InputPin, true, false, bSetupUndoRedo);
 	}
 
 	if (bSetupUndoRedo)
@@ -10139,17 +10327,40 @@ URigVMTemplateNode* URigVMController::AddTemplateNode(const FName& InNotation, c
 			Pin->CPPTypeObjectPath = *Pin->CPPTypeObject->GetPathName();
 		}
 		Pin->Direction = Arg->GetDirection();
+
 		AddNodePin(Node, Pin);
+
+		if(!Pin->IsWildCard())
+		{
+			const FString DefaultValue = Node->GetInitialDefaultValueForPin(Pin->GetFName());
+			if(UScriptStruct* ScriptStruct = Cast<UScriptStruct>(Pin->CPPTypeObject))
+			{
+				AddPinsForStruct(ScriptStruct, Pin->GetNode(), Pin, Pin->Direction, DefaultValue, false, false);
+			}
+			else
+			{
+				SetPinDefaultValue(Pin, DefaultValue, true, false, false, false);
+			}
+		}
 	}
 
 	Graph->Nodes.Add(Node);
 
 	Notify(ERigVMGraphNotifType::NodeAdded, Node);
 
+	FRigVMAddTemplateNodeAction Action;
 	FRigVMControllerCompileBracketScope CompileScope(this);
 	if (bSetupUndoRedo)
 	{
-		ActionStack->AddAction(FRigVMAddTemplateNodeAction(Node));
+		Action = FRigVMAddTemplateNodeAction(Node);
+		ActionStack->BeginAction(Action);
+	}
+	
+	ResolveTemplateNodeMetaData(Node, bSetupUndoRedo);
+
+	if (bSetupUndoRedo)
+	{
+		ActionStack->EndAction(Action);
 	}
 
 	if (bPrintPythonCommand)
@@ -10162,11 +10373,73 @@ URigVMTemplateNode* URigVMController::AddTemplateNode(const FName& InNotation, c
 		}
 	}
 
-	ResolveTemplateNodeMetaData(Node);
-
 	return Node;
 }
 
+TArray<UScriptStruct*> URigVMController::GetRegisteredUnitStructs()
+{
+	TArray<UScriptStruct*> UnitStructs;
+
+	for(const FRigVMFunction& Function : FRigVMRegistry::Get().GetFunctions())
+	{
+		if(!Function.IsValid())
+		{
+			continue;
+		}
+		UScriptStruct* Struct = Function.Struct;
+		if (!Struct->IsChildOf(FRigVMStruct::StaticStruct()))
+		{
+			continue;
+		}
+		UnitStructs.Add(Struct);
+	}
+	
+	return UnitStructs; 
+}
+
+TArray<FString> URigVMController::GetRegisteredTemplates()
+{
+	TArray<FString> Templates;
+
+	for(const FRigVMTemplate& Template : FRigVMRegistry::Get().GetTemplates())
+	{
+		if(!Template.IsValid() || Template.NumPermutations() < 2)
+		{
+			continue;
+		}
+		Templates.Add(Template.GetNotation().ToString());
+	}
+	
+	return Templates; 
+}
+
+TArray<UScriptStruct*> URigVMController::GetUnitStructsForTemplate(const FName& InNotation)
+{
+	TArray<UScriptStruct*> UnitStructs;
+
+	const FRigVMTemplate* Template = FRigVMRegistry::Get().FindTemplate(InNotation);
+	if(Template)
+	{
+		for(int32 PermutationIndex = 0; PermutationIndex < Template->NumPermutations(); PermutationIndex++)
+		{
+			UnitStructs.Add(Template->GetPermutation(PermutationIndex)->Struct);
+		}
+	}
+	
+	return UnitStructs;
+}
+
+FString URigVMController::GetTemplateForUnitStruct(UScriptStruct* InFunction, const FString& InMethodName)
+{
+	if(const FRigVMFunction* Function = FRigVMRegistry::Get().FindFunction(InFunction, *InMethodName))
+	{
+		if(const FRigVMTemplate* Template = Function->GetTemplate())
+		{
+			return Template->GetNotation().ToString();
+		}
+	}
+	return FString();
+}
 
 URigVMEnumNode* URigVMController::AddEnumNode(const FName& InCPPTypeObjectPath, const FVector2D& InPosition, const FString& InNodeName, bool bSetupUndoRedo, bool bPrintPythonCommand)
 {
@@ -12498,6 +12771,23 @@ void URigVMController::ReportError(const FString& InMessage) const
 	FScriptExceptionHandler::Get().HandleException(ELogVerbosity::Error, *Message, *FString());
 }
 
+void URigVMController::ReportAndNotifyInfo(const FString& InMessage) const
+{
+	ReportWarning(InMessage);
+	SendUserFacingNotification(InMessage, 0.f, nullptr, TEXT("MessageLog.Note"));
+}
+
+void URigVMController::ReportAndNotifyWarning(const FString& InMessage) const
+{
+	if (!bReportWarningsAndErrors)
+	{
+		return;
+	}
+
+	ReportWarning(InMessage);
+	SendUserFacingNotification(InMessage, 0.f, nullptr, TEXT("MessageLog.Warning"));
+}
+
 void URigVMController::ReportAndNotifyError(const FString& InMessage) const
 {
 	if (!bReportWarningsAndErrors)
@@ -12506,16 +12796,55 @@ void URigVMController::ReportAndNotifyError(const FString& InMessage) const
 	}
 
 	ReportError(InMessage);
+	SendUserFacingNotification(InMessage, 0.f, nullptr, TEXT("MessageLog.Error"));
+}
 
+void URigVMController::SendUserFacingNotification(const FString& InMessage, float InDuration, const UObject* InSubject, const FName& InBrushName) const
+{
 #if WITH_EDITOR
+
+	if(InDuration < SMALL_NUMBER)
+	{
+		InDuration = FMath::Clamp(0.1f * InMessage.Len(), 5.0f, 20.0f);
+	}
+	
 	FNotificationInfo Info(FText::FromString(InMessage));
 	Info.bUseSuccessFailIcons = true;
-	Info.Image = FEditorStyle::GetBrush(TEXT("MessageLog.Warning"));
+	Info.Image = FEditorStyle::GetBrush(InBrushName);
 	Info.bFireAndForget = true;
 	Info.bUseThrobber = true;
 	// longer message needs more time to read
-	Info.FadeOutDuration = FMath::Clamp(0.1f * InMessage.Len(), 5.0f, 20.0f);
-	Info.ExpireDuration = Info.FadeOutDuration;
+	Info.FadeOutDuration = FMath::Min(InDuration, 1.f);
+	Info.ExpireDuration = InDuration;
+
+	if(InSubject)
+	{
+		if(const URigVMNode* Node = Cast<URigVMNode>(InSubject))
+		{
+			Info.HyperlinkText = FText::FromString(Node->GetNodePath());
+		}
+		else if(const URigVMPin* Pin = Cast<URigVMPin>(InSubject))
+		{
+			Info.HyperlinkText = FText::FromString(Pin->GetPinPath());
+		}
+		else if(const URigVMLink* Link = Cast<URigVMLink>(InSubject))
+		{
+			Info.HyperlinkText = FText::FromString(((URigVMLink*)Link)->GetPinPathRepresentation());
+		}
+		else
+		{
+			Info.HyperlinkText = FText::FromName(InSubject->GetFName());
+		}
+
+		Info.Hyperlink = FSimpleDelegate::CreateLambda([InSubject, this]()
+		{
+			if(RequestJumpToHyperlinkDelegate.IsBound())
+			{
+				RequestJumpToHyperlinkDelegate.Execute(InSubject);
+			}
+		});
+	}
+	
 	TSharedPtr<SNotificationItem> NotificationPtr = FSlateNotificationManager::Get().AddNotification(Info);
 	if (NotificationPtr)
 	{
@@ -12524,7 +12853,6 @@ void URigVMController::ReportAndNotifyError(const FString& InMessage) const
 #endif
 }
 
-
 void URigVMController::CreateDefaultValueForStructIfRequired(UScriptStruct* InStruct, FString& InOutDefaultValue)
 {
 	if (InStruct != nullptr)
@@ -12551,7 +12879,8 @@ void URigVMController::CreateDefaultValueForStructIfRequired(UScriptStruct* InSt
 		
 		if (!InOutDefaultValue.IsEmpty() && InOutDefaultValue != TEXT("()"))
 		{ 
-			InStruct->ImportText(*InOutDefaultValue, TempBuffer.GetData(), nullptr, PPF_None, nullptr, FString());
+			FRigVMPinDefaultValueImportErrorContext ErrorPipe;
+			InStruct->ImportText(*InOutDefaultValue, TempBuffer.GetData(), nullptr, PPF_None, &ErrorPipe, FString());
 		}
 
 		// in case InOutDefaultValue is not empty, it needs to be cleared
@@ -12654,14 +12983,18 @@ void URigVMController::ResolveTemplateNodePins(URigVMTemplateNode* InNode, bool
 					if(Pin->IsArray())
 					{
 						const FString BaseCPPType = RigVMTypeUtils::BaseTypeFromArrayType(ResolvedType.CPPType);
-						for(URigVMPin* SubPin : Pin->GetSubPins())
+						for(int32 SubPinIndex = 0; SubPinIndex < Pin->GetSubPins().Num(); SubPinIndex++)
 						{
+							URigVMPin* SubPin = Pin->GetSubPins()[SubPinIndex];
 							if(SubPin->IsWildCard())
 							{
 								ResolveWildCardPinImpl(SubPin, BaseCPPType, ResolvedType.CPPTypeObject, bSetupUndoRedo, false, false, true);
 							}
 						}
 					}
+
+					const FString NewDefaultValue = InNode->GetInitialDefaultValueForPin(Pin->GetFName(), ResolvedPermutationIndices);
+					SetPinDefaultValue(Pin, NewDefaultValue, true, bSetupUndoRedo, false, true);
 				}
 			}
 		}
@@ -12690,10 +13023,10 @@ void URigVMController::ResolveTemplateNodePins(URigVMTemplateNode* InNode, bool
 		}
 	}
 
-	ResolveTemplateNodeMetaData(InNode);
+	ResolveTemplateNodeMetaData(InNode, bSetupUndoRedo);
 }
 
-void URigVMController::ResolveTemplateNodeMetaData(URigVMTemplateNode* InNode)
+void URigVMController::ResolveTemplateNodeMetaData(URigVMTemplateNode* InNode, bool bSetupUndoRedo)
 {
 #if WITH_EDITOR
 	check(InNode);
@@ -12708,6 +13041,22 @@ void URigVMController::ResolveTemplateNodeMetaData(URigVMTemplateNode* InNode)
 		Notify(ERigVMGraphNotifType::NodeColorChanged, InNode);
 	}
 #endif
+
+	if(InNode->IsResolved())
+	{
+		for(URigVMPin* Pin : InNode->GetPins())
+		{
+			if(Pin->IsWildCard() || Pin->ContainsWildCardSubPin())
+			{
+				continue;
+			}
+			if(!Pin->IsValidDefaultValue(Pin->GetDefaultValue()))
+			{
+				const FString NewDefaultValue = InNode->GetInitialDefaultValueForPin(Pin->GetFName(), ResolvedPermutationIndices);
+				SetPinDefaultValue(Pin, NewDefaultValue, true, bSetupUndoRedo, false, true);
+			}
+		}
+	}
 }
 
 bool URigVMController::FullyResolveTemplateNode(URigVMTemplateNode* InNode, const FRigVMFunction* InFunctionToResolve, bool bSetupUndoRedo)
@@ -12795,7 +13144,7 @@ bool URigVMController::FullyResolveTemplateNode(URigVMTemplateNode* InNode, cons
 	// exit out early if there's nothing to do
 	if(PinTypesToChange.IsEmpty() && MissingPins.IsEmpty() && PinsToRemove.IsEmpty())
 	{
-		ResolveTemplateNodeMetaData(InNode);
+		ResolveTemplateNodeMetaData(InNode, bSetupUndoRedo);
 		return true;
 	}
 
@@ -12906,7 +13255,7 @@ bool URigVMController::ResolveWildCardPin(const FString& InPinPath, const FStrin
 
 	if (URigVMPin* Pin = Graph->FindPin(InPinPath))
 	{
-		if(Pin->IsWildCard())
+		if(Pin->IsWildCard() || Pin->GetNode()->IsA<URigVMTemplateNode>())
 		{
 			FRigVMBaseAction Action;
 			if (bSetupUndoRedo)
@@ -12949,7 +13298,8 @@ bool URigVMController::ResolveWildCardPin(const FString& InPinPath, const FStrin
 }
 
 bool URigVMController::ResolveWildCardPin(URigVMPin* InPinToResolve, URigVMPin* InTemplatePin,
-										  bool bSetupUndoRedo, bool bTraverseNode, bool bTraverseParentPins, bool bTraverseLinks)
+										  bool bSetupUndoRedo, bool bTraverseNode, bool bTraverseParentPins,
+										  bool bTraverseLinks, bool bAllowFloatingPointCasts)
 {
 	check(InTemplatePin);
 	check(!InTemplatePin->IsWildCard());
@@ -12963,7 +13313,7 @@ bool URigVMController::ResolveWildCardPin(URigVMPin* InPinToResolve, URigVMPin*
 		ActionStack->BeginAction(Action);
 	}
 
-	const bool bResult = ResolveWildCardPinImpl(InPinToResolve, InTemplatePin->GetCPPType(), InTemplatePin->GetCPPTypeObject(), bSetupUndoRedo, bTraverseNode, bTraverseParentPins, bTraverseLinks);
+	const bool bResult = ResolveWildCardPinImpl(InPinToResolve, InTemplatePin->GetCPPType(), InTemplatePin->GetCPPTypeObject(), bSetupUndoRedo, bTraverseNode, bTraverseParentPins, bTraverseLinks, bAllowFloatingPointCasts);
 	if(!bResult)
 	{
 		if(InPinToResolve->IsLinkedTo(InTemplatePin))
@@ -12982,10 +13332,18 @@ bool URigVMController::ResolveWildCardPin(URigVMPin* InPinToResolve, URigVMPin*
 }
 
 bool URigVMController::ResolveWildCardPinImpl(URigVMPin* InPinToResolve, const FString& InCPPType, UObject* InCPPTypeObject,
-                                             bool bSetupUndoRedo, bool bTraverseNode, bool bTraverseParentPins, bool bTraverseLinks)
+                                             bool bSetupUndoRedo, bool bTraverseNode, bool bTraverseParentPins,
+                                             bool bTraverseLinks, bool bAllowFloatingPointCasts)
 {
 	check(InPinToResolve);
-	check(InPinToResolve->IsWildCard());
+	check(InPinToResolve->IsWildCard() || InPinToResolve->GetNode()->IsA<URigVMTemplateNode>());
+
+	// exit early if we already have the right type
+	const FRigVMTemplateArgument::FType ArgumentType(InCPPType, InCPPTypeObject);
+	if(ArgumentType.Matches(InPinToResolve->GetCPPType(), bAllowFloatingPointCasts))
+	{
+		return false;
+	}
 
 	// if we are trying to resolve a pin on the template node...
 	FString CPPType = InCPPType;
@@ -12995,13 +13353,90 @@ bool URigVMController::ResolveWildCardPinImpl(URigVMPin* InPinToResolve, const F
 	{
 		const FString BaseCPPType = RigVMTypeUtils::IsArrayType(CPPType) ? RigVMTypeUtils::BaseTypeFromArrayType(CPPType) : CPPType;
 		const FString TypeToResolveAgainst = InPinToResolve->GetRootPin()->IsArray() ? RigVMTypeUtils::ArrayTypeFromBaseType(BaseCPPType) : BaseCPPType; 
-
 		if(!TemplateNode->SupportsType(InPinToResolve->GetRootPin(), TypeToResolveAgainst, &CPPType))
 		{
 			return false;
 		}
 	}
 
+	// if the root pin is not a wildcard - we will need to unresolve first but maintain all links.
+	if(!InPinToResolve->IsWildCard())
+	{
+		URigVMTemplateNode* TemplateNode = CastChecked<URigVMTemplateNode>(InPinToResolve->GetNode());
+		TArray<TPair<FString, FString>> LinkedPaths = GetLinkedPinPaths(TemplateNode);
+		const FRigVMTemplate::FTypeMap OriginalTypes = TemplateNode->GetResolvedTypes();
+		const TMap<FString, FPinState> PinStates = GetPinStates(TemplateNode);
+
+		bool bUnresolveResult = false;
+		{
+			TGuardValue<bool> NotifGuard(bSuspendNotifications, true);
+			bUnresolveResult = UnresolveTemplateNodes({TemplateNode}, bSetupUndoRedo);
+		}
+
+		if(bUnresolveResult)
+		{
+			URigVMPin* RootPin = InPinToResolve->GetRootPin();
+			const FName ArgumentName = RootPin->GetFName();
+			const FString BaseCPPType = RigVMTypeUtils::IsArrayType(CPPType) ? RigVMTypeUtils::BaseTypeFromArrayType(CPPType) : CPPType;
+			const FString TypeToResolveAgainst = RootPin->IsArray() ? RigVMTypeUtils::ArrayTypeFromBaseType(BaseCPPType) : BaseCPPType;
+			const FRigVMTemplateArgument::FType ExpectedType(TypeToResolveAgainst, InCPPTypeObject);
+
+			// get all of the new types that we want to resolve towards
+			FRigVMTemplate::FTypeMap ResolvedTypes = OriginalTypes;
+			ResolvedTypes.Remove(ArgumentName);
+			if(TemplateNode->GetTemplate()->ResolveArgument(ArgumentName, ExpectedType, ResolvedTypes))
+			{
+				for(const FRigVMTemplate::FTypePair& TypePair : ResolvedTypes)
+				{
+					const FName& RootPinName = TypePair.Key;
+					const FRigVMTemplateArgument::FType& ResolvedType = TypePair.Value;
+
+					URigVMPin* Pin = TemplateNode->FindPin(RootPinName.ToString());
+					check(Pin);
+
+					if(Pin->GetTemplateArgumentType().Matches(ResolvedType.CPPType, false))
+					{
+						continue;
+					}
+
+					ResolveWildCardPinImpl(Pin, ResolvedType.CPPType, ResolvedType.CPPTypeObject, bSetupUndoRedo, false, false, false);
+				}
+
+				// re-attach the links. this is going to recursively change types
+				// on consecutive template nodes
+				for(const TPair<FString, FString>& LinkedPath : LinkedPaths)
+				{
+					const FString& SourcePath = LinkedPath.Key;
+					const FString& TargetPath = LinkedPath.Value;
+					URigVMPin* SourcePin = GetGraph()->FindPin(SourcePath);
+					URigVMPin* TargetPin = GetGraph()->FindPin(TargetPath);
+
+					if(SourcePin == nullptr || TargetPin == nullptr)
+					{
+						ReportRemovedLink(SourcePath, TargetPath);
+						continue;
+					}
+
+					// we are ok with this failing
+					if(!AddLink(SourcePin, TargetPin, bSetupUndoRedo,
+						SourcePin->GetNode() == TemplateNode ? ERigVMPinDirection::Output : ERigVMPinDirection::Input))
+					{
+						ReportRemovedLink(SourcePath, TargetPath);
+					}
+				}
+
+				// reapply default values
+				ApplyPinStates(TemplateNode, PinStates);
+				
+				return true;
+			}
+			else
+			{
+				return false;
+			}
+		}
+	}
+
 	FRigVMControllerCompileBracketScope CompileScope(this);
 	if(bTraverseParentPins)
 	{
@@ -13227,6 +13662,14 @@ bool URigVMController::ChangePinType(URigVMPin* InPin, const FString& InCPPType,
 	// we might want to use GetPinInitialDefaultValue here for a better default value
 	InPin->DefaultValue = FString();
 
+	if(InPin->IsRootPin() && !InPin->IsWildCard())
+	{
+		if(URigVMTemplateNode* TemplateNode = Cast<URigVMTemplateNode>(InPin->GetNode()))
+		{
+			InPin->DefaultValue = TemplateNode->GetInitialDefaultValueForPin(InPin->GetFName());
+		}
+	}
+
 	if (InPin->IsExecuteContext() && !InPin->GetNode()->IsA<URigVMFunctionEntryNode>() && !InPin->GetNode()->IsA<URigVMFunctionReturnNode>())
 	{
 		InPin->Direction = ERigVMPinDirection::IO;
@@ -13235,10 +13678,7 @@ bool URigVMController::ChangePinType(URigVMPin* InPin, const FString& InCPPType,
 	if (InPin->IsStruct() && !InPin->IsArray())
 	{
 		FString DefaultValue = InPin->DefaultValue;
-		if(!bPinWasWildCard)
-		{
-			CreateDefaultValueForStructIfRequired(InPin->GetScriptStruct(), DefaultValue);
-		}
+		CreateDefaultValueForStructIfRequired(InPin->GetScriptStruct(), DefaultValue);
 		AddPinsForStruct(InPin->GetScriptStruct(), InPin->GetNode(), InPin, InPin->Direction, DefaultValue, false, true);
 	}
 
@@ -13374,8 +13814,6 @@ bool URigVMController::ChangePinType(URigVMPin* InPin, const FString& InCPPType,
 					Local::ApplyResolvedDefaultValue(this, Pin, FString(), SimpleTypeDefaultValue, bSetupUndoRedo);
 				}
 			}
-			
-			// todo: update the default value for pins that don't have a default value
 		}
 	}
 
@@ -13723,3 +14161,24 @@ void URigVMController::SanitizeName(FString& InOutName, bool bAllowPeriod, bool
 		InOutName.LeftChopInline(InOutName.Len() - GetMaxNameLength());
 	}
 }
+
+TArray<TPair<FString, FString>> URigVMController::GetLinkedPinPaths(URigVMNode* InNode)
+{
+	const TArray<URigVMNode*> Nodes = {InNode};
+	return GetLinkedPinPaths(Nodes);
+}
+
+TArray<TPair<FString, FString>> URigVMController::GetLinkedPinPaths(const TArray<URigVMNode*>& InNodes)
+{
+	TArray<TPair<FString, FString>> LinkedPaths;
+	for(URigVMNode* Node : InNodes)
+	{
+		TArray<URigVMLink*> Links = Node->GetLinks();
+		for(URigVMLink* Link : Links)
+		{
+			TPair<FString, FString> LinkedPath(Link->GetSourcePin()->GetPinPath(), Link->GetTargetPin()->GetPinPath());
+			LinkedPaths.AddUnique(LinkedPath);
+		}
+	}
+	return LinkedPaths;
+}
diff --git a/Engine/Source/Developer/RigVMDeveloper/Private/RigVMModel/RigVMControllerActions.cpp b/Engine/Source/Developer/RigVMDeveloper/Private/RigVMModel/RigVMControllerActions.cpp
--- a/Engine/Source/Developer/RigVMDeveloper/Private/RigVMModel/RigVMControllerActions.cpp
+++ b/Engine/Source/Developer/RigVMDeveloper/Private/RigVMModel/RigVMControllerActions.cpp
@@ -1274,6 +1274,14 @@ FRigVMSetPinDefaultValueAction::FRigVMSetPinDefaultValueAction(URigVMPin* InPin,
 , OldDefaultValue(InPin->GetDefaultValue())
 , NewDefaultValue(InNewDefaultValue)
 {
+	if(!OldDefaultValue.IsEmpty())
+	{
+		check(InPin->IsValidDefaultValue(OldDefaultValue));
+	}
+	if(!NewDefaultValue.IsEmpty())
+	{
+		check(InPin->IsValidDefaultValue(NewDefaultValue));
+	}
 }
 
 bool FRigVMSetPinDefaultValueAction::Merge(const FRigVMBaseAction* Other)
diff --git a/Engine/Source/Developer/RigVMDeveloper/Private/RigVMModel/RigVMGraph.cpp b/Engine/Source/Developer/RigVMDeveloper/Private/RigVMModel/RigVMGraph.cpp
--- a/Engine/Source/Developer/RigVMDeveloper/Private/RigVMModel/RigVMGraph.cpp
+++ b/Engine/Source/Developer/RigVMDeveloper/Private/RigVMModel/RigVMGraph.cpp
@@ -413,9 +413,9 @@ void URigVMGraph::PrepareCycleChecking(URigVMPin* InPin, bool bAsInput)
 	GetDiagnosticsAST(false, LinksToSkip)->PrepareCycleChecking(InPin);
 }
 
-bool URigVMGraph::CanLink(URigVMPin* InSourcePin, URigVMPin* InTargetPin, FString* OutFailureReason, const FRigVMByteCode* InByteCode)
+bool URigVMGraph::CanLink(URigVMPin* InSourcePin, URigVMPin* InTargetPin, FString* OutFailureReason, const FRigVMByteCode* InByteCode, ERigVMPinDirection InUserLinkDirection)
 {
-	if (!URigVMPin::CanLink(InSourcePin, InTargetPin, OutFailureReason, InByteCode))
+	if (!URigVMPin::CanLink(InSourcePin, InTargetPin, OutFailureReason, InByteCode, InUserLinkDirection))
 	{
 		return false;
 	}
diff --git a/Engine/Source/Developer/RigVMDeveloper/Private/RigVMModel/RigVMNode.cpp b/Engine/Source/Developer/RigVMDeveloper/Private/RigVMModel/RigVMNode.cpp
--- a/Engine/Source/Developer/RigVMDeveloper/Private/RigVMModel/RigVMNode.cpp
+++ b/Engine/Source/Developer/RigVMDeveloper/Private/RigVMModel/RigVMNode.cpp
@@ -147,7 +147,7 @@ URigVMPin* URigVMNode::FindPin(const FString& InPinPath) const
 
 	for (URigVMPin* Pin : GetPins())
 	{
-		if (Pin->GetName() == Left)
+		if (Pin->NameEquals(Left, true))
 		{
 			if (Right.IsEmpty())
 			{
diff --git a/Engine/Source/Developer/RigVMDeveloper/Private/RigVMModel/RigVMPin.cpp b/Engine/Source/Developer/RigVMDeveloper/Private/RigVMModel/RigVMPin.cpp
--- a/Engine/Source/Developer/RigVMDeveloper/Private/RigVMModel/RigVMPin.cpp
+++ b/Engine/Source/Developer/RigVMDeveloper/Private/RigVMModel/RigVMPin.cpp
@@ -21,23 +21,9 @@
 #include "RigVMModel/Nodes/RigVMFunctionEntryNode.h"
 #include "RigVMModel/Nodes/RigVMFunctionReturnNode.h"
 
-class FRigVMPinDefaultValueImportErrorContext : public FOutputDevice
-{
-public:
-
-	int32 NumErrors;
-
-	FRigVMPinDefaultValueImportErrorContext()
-		: FOutputDevice()
-		, NumErrors(0)
-	{
-	}
-
-	virtual void Serialize(const TCHAR* V, ELogVerbosity::Type Verbosity, const class FName& Category) override
-	{
-		NumErrors++;
-	}
-};
+#if WITH_EDITOR
+#include "UObject/CoreRedirects.h"
+#endif
 
 URigVMGraph* URigVMInjectionInfo::GetGraph() const
 {
@@ -200,6 +186,57 @@ URigVMPin::URigVMPin()
 #endif
 }
 
+bool URigVMPin::NameEquals(const FString& InName, bool bFollowCoreRedirectors) const
+{
+	if(InName.Equals(GetName()))
+	{
+		return true;
+	}
+#if WITH_EDITOR
+	if(bFollowCoreRedirectors)
+	{
+		UScriptStruct* Struct = nullptr;
+		if(const URigVMPin* ParentPin = GetParentPin())
+		{
+			Struct = ParentPin->GetScriptStruct();
+		}
+		else if(const URigVMUnitNode* UnitNode = Cast<URigVMUnitNode>(GetNode()))
+		{
+			Struct = UnitNode->GetScriptStruct();
+		}
+
+		if(Struct)
+		{
+			typedef TPair<FName, FString> FRedirectPinPair;
+			const FRedirectPinPair Key(Struct->GetFName(), InName);
+			static TMap<FRedirectPinPair, FName> RedirectedPinNames;
+
+			if(const FName* RedirectedNamePtr = RedirectedPinNames.Find(Key))
+			{
+				if(RedirectedNamePtr->IsNone())
+				{
+					return false;
+				}
+				return NameEquals(*RedirectedNamePtr->ToString(), false);
+			}
+
+			const FCoreRedirectObjectName OldObjectName(*InName, Struct->GetFName(), *Struct->GetOutermost()->GetPathName());
+			const FCoreRedirectObjectName NewObjectName = FCoreRedirects::GetRedirectedName(ECoreRedirectFlags::Type_Property, OldObjectName);
+			if (OldObjectName != NewObjectName)
+			{
+				RedirectedPinNames.Add(Key, NewObjectName.ObjectName);
+				
+				const FString RedirectedName = NewObjectName.ObjectName.ToString();
+				return NameEquals(RedirectedName, false);
+			}
+
+			RedirectedPinNames.Add(Key, NAME_None);
+		}
+	}
+#endif
+	return false;
+}
+
 FString URigVMPin::GetPinPath(bool bUseNodePath) const
 {
 	FString PinPath;
@@ -521,6 +558,11 @@ FString URigVMPin::GetArrayElementCppType() const
 	return RigVMTypeUtils::BaseTypeFromArrayType(ResolvedType);
 }
 
+FRigVMTemplateArgument::FType URigVMPin::GetTemplateArgumentType() const
+{
+	return FRigVMTemplateArgument::FType(GetCPPType(), GetCPPTypeObject());
+}
+
 bool URigVMPin::IsStringType() const
 {
 	const FString ResolvedType = GetCPPType();
@@ -551,6 +593,19 @@ bool URigVMPin::IsWildCard() const
 	return false;
 }
 
+bool URigVMPin::ContainsWildCardSubPin() const
+{
+	for(const URigVMPin* SubPin : SubPins)
+	{
+		if(SubPin->IsWildCard() || SubPin->ContainsWildCardSubPin())
+		{
+			return true;
+		}
+	}
+	return false;
+}
+
+
 FString URigVMPin::GetDefaultValue() const
 {
 	return GetDefaultValue(EmptyPinOverride);
@@ -1232,6 +1287,11 @@ URigVMPin* URigVMPin::GetPinForLink() const
 
 URigVMPin* URigVMPin::GetOriginalPinFromInjectedNode() const
 {
+	if(GetNode() == nullptr)
+	{
+		return nullptr;
+	}
+	
 	if (URigVMInjectionInfo* Injection = GetNode()->GetInjectionInfo())
 	{
 		URigVMPin* RootPin = GetRootPin();
@@ -1281,7 +1341,7 @@ URigVMPin* URigVMPin::FindSubPin(const FString& InPinPath) const
 
 	for (URigVMPin* Pin : SubPins)
 	{
-		if (Pin->GetName() == Left)
+		if (Pin->NameEquals(Left, true))
 		{
 			if (Right.IsEmpty())
 			{
@@ -1447,7 +1507,7 @@ URigVMGraph* URigVMPin::GetGraph() const
 	return nullptr;
 }
 
-bool URigVMPin::CanLink(URigVMPin* InSourcePin, URigVMPin* InTargetPin, FString* OutFailureReason, const FRigVMByteCode* InByteCode)
+bool URigVMPin::CanLink(URigVMPin* InSourcePin, URigVMPin* InTargetPin, FString* OutFailureReason, const FRigVMByteCode* InByteCode, ERigVMPinDirection InUserLinkDirection)
 {
 	if (InSourcePin == nullptr || InTargetPin == nullptr)
 	{
@@ -1551,12 +1611,37 @@ bool URigVMPin::CanLink(URigVMPin* InSourcePin, URigVMPin* InTargetPin, FString*
 					return true;
 
 				};
+
+				auto IsPinValidForTypeChange = [](URigVMPin* InPin, bool bIsInput, ERigVMPinDirection InDirection) -> bool
+				{
+					if(InPin->IsWildCard())
+					{
+						return true;
+					}
+
+					if(!InPin->GetNode()->IsA<URigVMTemplateNode>() || InDirection == ERigVMPinDirection::Invalid)
+					{
+						return false;
+					}
+
+					if((bIsInput && (InDirection != ERigVMPinDirection::Input)) ||
+						(!bIsInput && (InDirection != ERigVMPinDirection::Output)))
+					{
+						if(InPin->IsRootPin() ||
+							(InPin->GetParentPin()->IsRootPin() && InPin->IsArrayElement()))
+						{
+							return true;
+						}
+					}
+					
+					return false;
+				};
 				
-				if(InSourcePin->IsWildCard() && !InTargetPin->IsExecuteContext())
+				if(IsPinValidForTypeChange(InSourcePin, false, InUserLinkDirection) && !InTargetPin->IsExecuteContext())
 				{
 					return TemplateNodeSupportsType(InSourcePin, InTargetPin->GetCPPType(), OutFailureReason);
 				}
-				else if(InTargetPin->IsWildCard() && !InSourcePin->IsExecuteContext())
+				else if(IsPinValidForTypeChange(InTargetPin, true, InUserLinkDirection) && !InSourcePin->IsExecuteContext())
 				{
 					return TemplateNodeSupportsType(InTargetPin, InSourcePin->GetCPPType(), OutFailureReason);
 				}
diff --git a/Engine/Source/Developer/RigVMDeveloper/Public/RigVMModel/Nodes/RigVMTemplateNode.h b/Engine/Source/Developer/RigVMDeveloper/Public/RigVMModel/Nodes/RigVMTemplateNode.h
--- a/Engine/Source/Developer/RigVMDeveloper/Public/RigVMModel/Nodes/RigVMTemplateNode.h
+++ b/Engine/Source/Developer/RigVMDeveloper/Public/RigVMModel/Nodes/RigVMTemplateNode.h
@@ -31,12 +31,15 @@ public:
 	virtual FText GetToolTipTextForPin(const URigVMPin* InPin) const override;
 
 	// Returns the notation of the node
-	UFUNCTION(BlueprintCallable, Category = Template)
+	UFUNCTION(BlueprintPure, Category = Template)
 	FName GetNotation() const;
 
 	// returns true if a pin supports a given type
 	bool SupportsType(const URigVMPin* InPin, const FString& InCPPType, FString* OutCPPType = nullptr);
 
+	// resolves a pin to a new type outputs the new type map
+	bool GetTypeMapForNewPinType(const URigVMPin* InPin, const FString& InCPPType, UObject* InCPPTypeObject, FRigVMTemplate::FTypeMap& OutTypes) const;
+
 	// returns the index of the resolved RigVM function (or INDEX_NONE)
 	TArray<int32> GetResolvedPermutationIndices(FRigVMTemplate::FTypeMap* OutTypes = nullptr) const;
 
@@ -46,12 +49,31 @@ public:
 	// returns the template used for this node
 	const FRigVMTemplate* GetTemplate() const;
 
+	// returns the type map for the currently resolved pins
+	FRigVMTemplate::FTypeMap GetResolvedTypes() const;
+
 	// returns the resolved function or nullptr if there are still unresolved pins left
 	const FRigVMFunction* GetResolvedFunction() const;
 
 	// returns true if the template node is resolved
+	UFUNCTION(BlueprintPure, Category = Template)
 	bool IsResolved() const;
 
+	// returns true if the template is fully unresolved
+	UFUNCTION(BlueprintPure, Category = Template)
+	bool IsFullyUnresolved() const;
+
+	// returns the sole resolved unit struct
+	UFUNCTION(BlueprintPure, Category = Template)
+	UScriptStruct* GetResolvedUnitStruct() const { return GetScriptStruct(); }
+
+	// returns the list of all supported unit structs
+	UFUNCTION(BlueprintPure, Category = Template)
+	TArray<UScriptStruct*> GetSupportedUnitStructs() const;
+
+	// returns a default value for pin if it is known
+	FString GetInitialDefaultValueForPin(const FName& InRootPinName, const TArray<int32>& InPermutationIndices = TArray<int32>()) const;
+
 private:
 
 	void InvalidateCache();
diff --git a/Engine/Source/Developer/RigVMDeveloper/Public/RigVMModel/RigVMController.h b/Engine/Source/Developer/RigVMDeveloper/Public/RigVMModel/RigVMController.h
--- a/Engine/Source/Developer/RigVMDeveloper/Public/RigVMModel/RigVMController.h
+++ b/Engine/Source/Developer/RigVMDeveloper/Public/RigVMModel/RigVMController.h
@@ -86,6 +86,8 @@ DECLARE_DELEGATE_RetVal_TwoParams(bool, FRigVMController_IsDependencyCyclicDeleg
 DECLARE_DELEGATE_RetVal_TwoParams(FRigVMController_BulkEditResult, FRigVMController_RequestBulkEditDialogDelegate, URigVMLibraryNode*, ERigVMControllerBulkEditType)
 DECLARE_DELEGATE_FiveParams(FRigVMController_OnBulkEditProgressDelegate, TSoftObjectPtr<URigVMFunctionReferenceNode>, ERigVMControllerBulkEditType, ERigVMControllerBulkEditProgress, int32, int32)
 DECLARE_DELEGATE_RetVal_TwoParams(FString, FRigVMController_PinPathRemapDelegate, const FString& /* InPinPath */, bool /* bIsInput */);
+DECLARE_DELEGATE_OneParam(FRigVMController_RequestJumpToHyperlinkDelegate, const UObject* InSubject);
+
 
 /**
  * The Controller is the sole authority to perform changes
@@ -218,6 +220,10 @@ public:
 	UFUNCTION(BlueprintCallable, Category = RigVMController)
 	URigVMTemplateNode* ReplaceUnitNodeWithTemplateNode(const FName& InNodeName, bool bSetupUndoRedo);
 
+	// Turns a resolved templated node(s) back into its template.
+	UFUNCTION(BlueprintCallable, Category = RigVMController)
+	bool UnresolveTemplateNodes(const TArray<FName>& InNodeNames, bool bSetupUndoRedo = true, bool bPrintPythonCommand = false);
+
 	// Upgrades a set of nodes with each corresponding next known version
 	UFUNCTION(BlueprintCallable, Category = RigVMController)
 	TArray<URigVMNode*> UpgradeNodes(const TArray<FName>& InNodeNames, bool bRecursive = true, bool bSetupUndoRedo = true, bool bPrintPythonCommand = false);
@@ -291,6 +297,22 @@ public:
 	UFUNCTION(BlueprintCallable, Category = RigVMController)
 	URigVMTemplateNode* AddTemplateNode(const FName& InNotation, const FVector2D& InPosition = FVector2D::ZeroVector, const FString& InNodeName = TEXT(""), bool bSetupUndoRedo = true, bool bPrintPythonCommand = false);
 
+	// Returns all registered unit structs
+	UFUNCTION(BlueprintCallable, Category = RigVMController)
+	static TArray<UScriptStruct*> GetRegisteredUnitStructs();
+
+	// Returns all registered template notations
+	UFUNCTION(BlueprintCallable, Category = RigVMController)
+	static TArray<FString> GetRegisteredTemplates();
+
+	// Returns all supported unit structs for a given template notation
+	UFUNCTION(BlueprintCallable, Category = RigVMController)
+	static TArray<UScriptStruct*> GetUnitStructsForTemplate(const FName& InNotation);
+
+	// Returns the template for a given function (or an empty string)
+	UFUNCTION(BlueprintCallable, Category = RigVMController)
+	static FString GetTemplateForUnitStruct(UScriptStruct* InFunction, const FString& InMethodName = TEXT("Execute"));
+
 	// Resolves a wildcard pin on any node
 	UFUNCTION(BlueprintCallable, Category = RigVMController)
 	bool ResolveWildCardPin(const FString& InPinPath, const FString& InCPPType, const FName& InCPPTypeObjectPath, bool bSetupUndoRedo = true, bool bPrintPythonCommand = false);
@@ -540,6 +562,8 @@ public:
 	UFUNCTION(BlueprintCallable, Category = RigVMController, meta=(DeprecatedFunction))
 	bool RenameParameter(const FName& InOldName, const FName& InNewName, bool bSetupUndoRedo = true);\
 
+	bool UnresolveTemplateNodes(const TArray<URigVMTemplateNode*>& InNodes, bool bSetupUndoRedo);
+
 	// Upgrades a set of nodes with each corresponding next known version
 	TArray<URigVMNode*> UpgradeNodes(const TArray<URigVMNode*>& InNodes, bool bRecursive = true, bool bSetupUndoRedo = true);
 
@@ -625,7 +649,7 @@ public:
 	// Adds a link to the graph.
 	// This causes a LinkAdded modified event.
 	UFUNCTION(BlueprintCallable, Category = RigVMController)
-	bool AddLink(const FString& InOutputPinPath, const FString& InInputPinPath, bool bSetupUndoRedo = true, bool bPrintPythonCommand = false);
+	bool AddLink(const FString& InOutputPinPath, const FString& InInputPinPath, bool bSetupUndoRedo = true, bool bPrintPythonCommand = false, ERigVMPinDirection InUserDirection = ERigVMPinDirection::Invalid);
 
 	// Removes a link from the graph.
 	// This causes a LinkRemoved modified event.
@@ -738,6 +762,9 @@ public:
 	// A delegate to inform the host / client about the progress during a bulk edit
 	FRigVMController_OnBulkEditProgressDelegate OnBulkEditProgressDelegate;
 
+	// A delegate to request the client to follow a hyper link
+	FRigVMController_RequestJumpToHyperlinkDelegate RequestJumpToHyperlinkDelegate; 
+
 	// Returns the build data of the host
 	static URigVMBuildData* GetBuildData(bool bCreateIfNeeded = true);
 
@@ -775,7 +802,10 @@ public:
 	void ReportInfo(const FString& InMessage) const;
 	void ReportWarning(const FString& InMessage) const;
 	void ReportError(const FString& InMessage) const;
+	void ReportAndNotifyInfo(const FString& InMessage) const;
+	void ReportAndNotifyWarning(const FString& InMessage) const;
 	void ReportAndNotifyError(const FString& InMessage) const;
+	void SendUserFacingNotification(const FString& InMessage, float InDuration = 0.f, const UObject* InSubject = nullptr, const FName& InBrushName = TEXT("MessageLog.Warning")) const;
 
 	template <typename FmtType, typename... Types>
 	void ReportInfof(const FmtType& Fmt, Types... Args)
@@ -795,6 +825,18 @@ public:
 		ReportError(FString::Printf(Fmt, Args...));
 	}
 
+	template <typename FmtType, typename... Types>
+	void ReportAndNotifyInfof(const FmtType& Fmt, Types... Args)
+	{
+		ReportAndNotifyInfo(FString::Printf(Fmt, Args...));
+	}
+
+	template <typename FmtType, typename... Types>
+	void ReportAndNotifyWarningf(const FmtType& Fmt, Types... Args)
+	{
+		ReportAndNotifyWarning(FString::Printf(Fmt, Args...));
+	}
+
 	template <typename FmtType, typename... Types>
 	void ReportAndNotifyErrorf(const FmtType& Fmt, Types... Args)
 	{
@@ -856,7 +898,7 @@ private:
 	void RelinkSourceAndTargetPins(URigVMNode* RigNode, bool bSetupUndoRedo = true);
 
 public:
-	bool AddLink(URigVMPin* OutputPin, URigVMPin* InputPin, bool bSetupUndoRedo = true);
+	bool AddLink(URigVMPin* OutputPin, URigVMPin* InputPin, bool bSetupUndoRedo = true, ERigVMPinDirection InUserDirection = ERigVMPinDirection::Invalid);
 	bool BreakLink(URigVMPin* OutputPin, URigVMPin* InputPin, bool bSetupUndoRedo = true);
 	bool BreakAllLinks(URigVMPin* Pin, bool bAsInput, bool bSetupUndoRedo = true);
 
@@ -877,6 +919,8 @@ private:
 
 	void RefreshFunctionPins(URigVMNode* InNode, bool bNotify = true);
 
+	void ReportRemovedLink(const FString& InSourcePinPath, const FString& InTargetPinPath);
+
 	struct FPinState
 	{
 		ERigVMPinDirection Direction;
@@ -899,10 +943,10 @@ private:
 	static void PostProcessDefaultValue(URigVMPin* Pin, FString& OutDefaultValue);
 	static FString PostProcessCPPType(const FString& InCPPType, UObject* InCPPTypeObject);
 
-	bool ResolveWildCardPin(URigVMPin* InPinToResolve, URigVMPin* InTemplatePin, bool bSetupUndoRedo, bool bTraverseNode = true, bool bTraverseParentPins = true, bool bTraverseLinks = true);
-	bool ResolveWildCardPinImpl(URigVMPin* InPinToResolve, const FString& InCPPType, UObject* InCPPTypeObject, bool bSetupUndoRedo, bool bTraverseNode = true, bool bTraverseParentPins = true, bool bTraverseLinks = true);
+	bool ResolveWildCardPin(URigVMPin* InPinToResolve, URigVMPin* InTemplatePin, bool bSetupUndoRedo, bool bTraverseNode = true, bool bTraverseParentPins = true, bool bTraverseLinks = true, bool bAllowFloatingPointCasts = false);
+	bool ResolveWildCardPinImpl(URigVMPin* InPinToResolve, const FString& InCPPType, UObject* InCPPTypeObject, bool bSetupUndoRedo, bool bTraverseNode = true, bool bTraverseParentPins = true, bool bTraverseLinks = true, bool bAllowFloatingPointCasts = false);
 	void ResolveTemplateNodePins(URigVMTemplateNode* InNode, bool bSetupUndoRedo);
-	void ResolveTemplateNodeMetaData(URigVMTemplateNode* InNode);
+	void ResolveTemplateNodeMetaData(URigVMTemplateNode* InNode, bool bSetupUndoRedo);
 	bool FullyResolveTemplateNode(URigVMTemplateNode* InNode, const FRigVMFunction* InFunctionToResolve, bool bSetupUndoRedo);
 	bool ChangePinType(const FString& InPinPath, const FString& InCPPType, const FName& InCPPTypeObjectPath, bool bSetupUndoRedo, bool bSetupOrphanPins = true, bool bBreakLinks = true, bool bRemoveSubPins = true);
 	bool ChangePinType(URigVMPin* InPin, const FString& InCPPType, const FName& InCPPTypeObjectPath, bool bSetupUndoRedo, bool bSetupOrphanPins = true, bool bBreakLinks = true, bool bRemoveSubPins = true);
@@ -940,6 +984,8 @@ public:
 	static FString GetSanitizedPinName(const FString& InName);
 	static FString GetSanitizedPinPath(const FString& InName);
 	static void SanitizeName(FString& InOutName, bool bAllowPeriod, bool bAllowSpace);
+	static TArray<TPair<FString, FString>> GetLinkedPinPaths(URigVMNode* InNode);
+	static TArray<TPair<FString, FString>> GetLinkedPinPaths(const TArray<URigVMNode*>& InNodes);
 
 private: 
 	UPROPERTY(transient)
@@ -954,6 +1000,7 @@ private:
 	bool bSuspendNotifications;
 	bool bReportWarningsAndErrors;
 	bool bIgnoreRerouteCompactnessChanges;
+	ERigVMPinDirection UserLinkDirection;
 
 	// temporary maps used for pin redirection
 	// only valid between Detach & ReattachLinksToPinObjects
diff --git a/Engine/Source/Developer/RigVMDeveloper/Public/RigVMModel/RigVMGraph.h b/Engine/Source/Developer/RigVMDeveloper/Public/RigVMModel/RigVMGraph.h
--- a/Engine/Source/Developer/RigVMDeveloper/Public/RigVMModel/RigVMGraph.h
+++ b/Engine/Source/Developer/RigVMDeveloper/Public/RigVMModel/RigVMGraph.h
@@ -133,7 +133,7 @@ public:
 	FRigVMGraphModifiedEvent& OnModified();
 
 	void PrepareCycleChecking(URigVMPin* InPin, bool bAsInput);
-	virtual bool CanLink(URigVMPin* InSourcePin, URigVMPin* InTargetPin, FString* OutFailureReason, const FRigVMByteCode* InByteCode);
+	virtual bool CanLink(URigVMPin* InSourcePin, URigVMPin* InTargetPin, FString* OutFailureReason, const FRigVMByteCode* InByteCode, ERigVMPinDirection InUserLinkDirection = ERigVMPinDirection::IO);
 	
 	TSharedPtr<FRigVMParserAST> GetDiagnosticsAST(bool bForceRefresh = false, TArray<URigVMLink*> InLinksToSkip = TArray<URigVMLink*>());
 	TSharedPtr<FRigVMParserAST> GetRuntimeAST(const FRigVMParserASTSettings& InSettings = FRigVMParserASTSettings::Optimized(), bool bForceRefresh = false);
diff --git a/Engine/Source/Developer/RigVMDeveloper/Public/RigVMModel/RigVMPin.h b/Engine/Source/Developer/RigVMDeveloper/Public/RigVMModel/RigVMPin.h
--- a/Engine/Source/Developer/RigVMDeveloper/Public/RigVMModel/RigVMPin.h
+++ b/Engine/Source/Developer/RigVMDeveloper/Public/RigVMModel/RigVMPin.h
@@ -7,6 +7,7 @@
 #include "UObject/ObjectMacros.h"
 #include "RigVMCore/RigVMRegistry.h"
 #include "RigVMCore/RigVMByteCode.h"
+#include "RigVMCore/RigVMTemplate.h"
 #include "RigVMCompiler/RigVMASTProxy.h"
 #include "RigVMPin.generated.h"
 
@@ -123,6 +124,9 @@ public:
 	// Default constructor
 	URigVMPin();
 
+	// returns true if the name of this pin matches a given name
+	bool NameEquals(const FString& InName, bool bFollowCoreRedirectors = false) const;
+
 	// Returns a . separated path containing all names of the pin and its owners,
 	// this includes the node name, for example "Node.Color.R"
 	UFUNCTION(BlueprintCallable, Category = RigVMPin)
@@ -206,6 +210,9 @@ public:
 	UFUNCTION(BlueprintCallable, Category = RigVMPin)
 	FString GetArrayElementCppType() const;
 
+	// Returns the argument type this pin would represent within a template
+	FRigVMTemplateArgument::FType GetTemplateArgumentType() const;
+
 	// Returns true if the C++ data type is FString or FName
 	UFUNCTION(BlueprintCallable, Category = RigVMPin)
 	bool IsStringType() const;
@@ -218,6 +225,10 @@ public:
 	UFUNCTION(BlueprintCallable, Category = RigVMPin)
 	bool IsWildCard() const;
 
+	// Returns true if any of the subpins is a wildcard
+	UFUNCTION(BlueprintCallable, Category = RigVMPin)
+	bool ContainsWildCardSubPin() const;
+
 	// Returns the default value of the Pin as a string.
 	// Note that this value is computed based on the Pin's
 	// SubPins - so for example for a FVector typed Pin
@@ -335,7 +346,7 @@ public:
 
 	// Returns true is the two provided source and target Pins
 	// can be linked to one another.
-	static bool CanLink(URigVMPin* InSourcePin, URigVMPin* InTargetPin, FString* OutFailureReason, const FRigVMByteCode* InByteCode);
+	static bool CanLink(URigVMPin* InSourcePin, URigVMPin* InTargetPin, FString* OutFailureReason, const FRigVMByteCode* InByteCode, ERigVMPinDirection InUserLinkDirection = ERigVMPinDirection::IO);
 
 	// Returns true if this pin has injected nodes
 	bool HasInjectedNodes() const { return InjectionInfos.Num() > 0; }
@@ -470,3 +481,21 @@ private:
 	friend class URigVMNode;
 	friend class FRigVMParserAST;
 };
+
+class RIGVMDEVELOPER_API FRigVMPinDefaultValueImportErrorContext : public FOutputDevice
+{
+public:
+
+	int32 NumErrors;
+
+	FRigVMPinDefaultValueImportErrorContext()
+		: FOutputDevice()
+		, NumErrors(0)
+	{
+	}
+
+	virtual void Serialize(const TCHAR* V, ELogVerbosity::Type Verbosity, const class FName& Category) override
+	{
+		NumErrors++;
+	}
+};
diff --git a/Engine/Source/Runtime/RigVM/Private/RigVMCore/RigVMTemplate.cpp b/Engine/Source/Runtime/RigVM/Private/RigVMCore/RigVMTemplate.cpp
--- a/Engine/Source/Runtime/RigVM/Private/RigVMCore/RigVMTemplate.cpp
+++ b/Engine/Source/Runtime/RigVM/Private/RigVMCore/RigVMTemplate.cpp
@@ -43,6 +43,10 @@ FRigVMTemplateArgument::FRigVMTemplateArgument(FProperty* InProperty)
 	{
 		Type.CPPTypeObject = EnumProperty->GetEnum();
 	}
+	else if (FByteProperty* ByteProperty = CastField<FByteProperty>(InProperty))
+	{
+		Type.CPPTypeObject = ByteProperty->Enum;
+	}
 
 	Types.Add(Type);
 }
@@ -600,6 +604,15 @@ bool FRigVMTemplate::Resolve(FTypeMap& InOutTypes, TArray<int32>& OutPermutation
 			if(bFoundMatch)
 			{
 				InOutTypes.Add(Argument.Name,MatchedType);
+
+				// if we found a perfect match - remove all permutations which don't match this one
+				if(bFoundPerfectMatch)
+				{
+					OutPermutationIndices.RemoveAll([Argument, MatchedType](int32 PermutationIndex) -> bool
+					{
+						return Argument.Types[PermutationIndex] != MatchedType;
+					});
+				}
 				continue;
 			}
 		}
@@ -636,6 +649,92 @@ bool FRigVMTemplate::Resolve(FTypeMap& InOutTypes, TArray<int32>& OutPermutation
 	return !OutPermutationIndices.IsEmpty();
 }
 
+bool FRigVMTemplate::ResolveArgument(const FName& InArgumentName, const FRigVMTemplateArgument::FType& InType,
+	FTypeMap& InOutTypes) const
+{
+	auto RemoveWildCardTypes = [](const FTypeMap& InTypes)
+	{
+		FTypeMap FilteredTypes;
+		for(const FTypePair& Pair: InTypes)
+		{
+			if(!Pair.Value.IsWildCard())
+			{
+				FilteredTypes.Add(Pair);
+			}
+		}
+		return FilteredTypes;
+	};
+
+	// remove all wildcards from map
+	InOutTypes = RemoveWildCardTypes(InOutTypes);
+
+	// first resolve with no types given except for the new argument type
+	FTypeMap ResolvedTypes;
+	ResolvedTypes.Add(InArgumentName, InType);
+	TArray<int32> PermutationIndices;
+	FTypeMap RemainingTypesToResolve;
+	
+	if(Resolve(ResolvedTypes, PermutationIndices, true))
+	{
+		// let's see if the input argument resolved into the expected type
+		const FRigVMTemplateArgument::FType ResolvedInputType = ResolvedTypes.FindChecked(InArgumentName);
+		if(!ResolvedInputType.Matches(InType.CPPType, true))
+		{
+			return false;
+		}
+		
+		ResolvedTypes = RemoveWildCardTypes(ResolvedTypes);
+		
+		// remove all argument types from the reference list
+		// provided from the outside. we cannot resolve these further
+		auto RemoveResolvedTypesFromRemainingList = [](
+			FTypeMap& InOutTypes, const FTypeMap& InResolvedTypes, FTypeMap& InOutRemainingTypesToResolve)
+		{
+			InOutRemainingTypesToResolve = InOutTypes;
+			for(const FTypePair& Pair: InOutTypes)
+			{
+				if(InResolvedTypes.Contains(Pair.Key))
+				{
+					InOutRemainingTypesToResolve.Remove(Pair.Key);
+				}
+			}
+			InOutTypes = InResolvedTypes;
+		};
+
+		RemoveResolvedTypesFromRemainingList(InOutTypes, ResolvedTypes, RemainingTypesToResolve);
+
+		// if the type hasn't been specified we need to slowly resolve the template
+		// arguments until we hit a match. for this we'll create a list of arguments
+		// to resolve and reduce the list slowly.
+		bool bSuccessFullyResolvedRemainingTypes = true;
+		while(!RemainingTypesToResolve.IsEmpty())
+		{
+			PermutationIndices.Reset();
+
+			const FTypePair TypeToResolve = *RemainingTypesToResolve.begin();
+			FTypeMap NewResolvedTypes = RemoveWildCardTypes(ResolvedTypes);
+			NewResolvedTypes.FindOrAdd(TypeToResolve.Key) = TypeToResolve.Value;
+
+			if(Resolve(NewResolvedTypes, PermutationIndices, true))
+			{
+				ResolvedTypes = NewResolvedTypes;
+				RemoveResolvedTypesFromRemainingList(InOutTypes, ResolvedTypes, RemainingTypesToResolve);
+			}
+			else
+			{
+				// we were not able to resolve this argument, remove it from the resolved types list.
+				RemainingTypesToResolve.Remove(TypeToResolve.Key);
+				bSuccessFullyResolvedRemainingTypes = false;
+			}
+		}
+
+		// if there is nothing left to resolve we were successful
+		return RemainingTypesToResolve.IsEmpty() && bSuccessFullyResolvedRemainingTypes;
+	}
+
+	return false;
+}
+
 #if WITH_EDITOR
 
 FString FRigVMTemplate::GetCategory() const
diff --git a/Engine/Source/Runtime/RigVM/Public/RigVMCore/RigVMTemplate.h b/Engine/Source/Runtime/RigVM/Public/RigVMCore/RigVMTemplate.h
--- a/Engine/Source/Runtime/RigVM/Public/RigVMCore/RigVMTemplate.h
+++ b/Engine/Source/Runtime/RigVM/Public/RigVMCore/RigVMTemplate.h
@@ -48,6 +48,11 @@ struct RIGVM_API FRigVMTemplateArgument
 			return CPPType != InOther.CPPType;
 		}
 
+		friend FORCEINLINE uint32 GetTypeHash(const FType& InType)
+		{
+			return GetTypeHash(InType.CPPType);
+		}
+
 		FORCEINLINE bool Matches(const FString& InCPPType, bool bAllowFloatingPointCasts = true) const
 		{
 			if(CPPType == InCPPType)
@@ -87,13 +92,24 @@ struct RIGVM_API FRigVMTemplateArgument
 
 		bool IsWildCard() const
 		{
-			return CPPTypeObject == RigVMTypeUtils::GetWildCardCPPTypeObject();
+			return CPPTypeObject == RigVMTypeUtils::GetWildCardCPPTypeObject() ||
+				CPPType == RigVMTypeUtils::GetWildCardCPPType() ||
+				CPPType == RigVMTypeUtils::GetWildCardArrayCPPType();
 		}
 
 		bool IsArray() const
 		{
 			return RigVMTypeUtils::IsArrayType(CPPType);
 		}
+
+		FString GetBaseCPPType() const
+		{
+			if(IsArray())
+			{
+				return RigVMTypeUtils::BaseTypeFromArrayType(CPPType);
+			}
+			return CPPType;
+		}
 	};
 
 	// returns the name of the argument
@@ -149,6 +165,7 @@ struct RIGVM_API FRigVMTemplate
 public:
 
 	typedef TMap<FName, FRigVMTemplateArgument::FType> FTypeMap;
+	typedef TPair<FName, FRigVMTemplateArgument::FType> FTypePair;
 
 	// returns true if this is a valid template
 	bool IsValid() const;
@@ -195,6 +212,9 @@ public:
 	// returns true if the template was able to resolve to at least one permutation
 	bool Resolve(FTypeMap& InOutTypes, TArray<int32> & OutPermutationIndices, bool bAllowFloatingPointCasts) const;
 
+	// returns true if the template can resolve an argument to a new type
+	bool ResolveArgument(const FName& InArgumentName, const FRigVMTemplateArgument::FType& InType, FTypeMap& InOutTypes) const;
+
 	// returns true if a given argument is valid for a template
 	static bool IsValidArgumentForTemplate(const FRigVMTemplateArgument& InArgument);
 
diff --git a/Engine/Source/Runtime/RigVM/Public/RigVMTypeUtils.h b/Engine/Source/Runtime/RigVM/Public/RigVMTypeUtils.h
--- a/Engine/Source/Runtime/RigVM/Public/RigVMTypeUtils.h
+++ b/Engine/Source/Runtime/RigVM/Public/RigVMTypeUtils.h
@@ -16,6 +16,7 @@ namespace RigVMTypeUtils
 	const FString FloatType = TEXT("float");
 	const FString DoubleType = TEXT("double");
 	const FString Int32Type = TEXT("int32");
+	const FString UInt8Type = TEXT("uint8");
 	const FString FNameType = TEXT("FName");
 	const FString FStringType = TEXT("FString");
 	const FString BoolArrayType = TEXT("TArray<bool>");
@@ -25,6 +26,20 @@ namespace RigVMTypeUtils
 	const FString FNameArrayType = TEXT("TArray<FName>");
 	const FString FStringArrayType = TEXT("TArray<FString>");
 
+	const FName BoolTypeName = *BoolType;
+	const FName FloatTypeName = *FloatType;
+	const FName DoubleTypeName = *DoubleType;
+	const FName Int32TypeName = *Int32Type;
+	const FName UInt8TypeName = *UInt8Type;
+	const FName FNameTypeName = *FNameType;
+	const FName FStringTypeName = *FStringType;
+	const FName BoolArrayTypeName = *BoolArrayType;
+	const FName FloatArrayTypeName = *FloatArrayType;
+	const FName DoubleArrayTypeName = *DoubleArrayType;
+	const FName Int32ArrayTypeName = *Int32ArrayType;
+	const FName FNameArrayTypeName = *FNameArrayType;
+	const FName FStringArrayTypeName = *FStringArrayType;
+
 	// Returns true if the type specified is an array
 	FORCEINLINE bool IsArrayType(const FString& InCPPType)
 	{
