Control Rig: Enable template nodes

diff --git a/Engine/Source/Developer/RigVMDeveloper/Public/RigVMModel/RigVMController.h b/Engine/Source/Developer/RigVMDeveloper/Public/RigVMModel/RigVMController.h
--- a/Engine/Source/Developer/RigVMDeveloper/Public/RigVMModel/RigVMController.h
+++ b/Engine/Source/Developer/RigVMDeveloper/Public/RigVMModel/RigVMController.h
@@ -21,7 +21,7 @@
 #include "RigVMController.generated.h"
 
 #ifndef UE_RIGVM_ENABLE_TEMPLATE_NODES
-#define UE_RIGVM_ENABLE_TEMPLATE_NODES 0
+#define UE_RIGVM_ENABLE_TEMPLATE_NODES 1
 #endif
 
 class URigVMActionStack;
