RigVM: Fix non-unity build fallout from 19384087

Merging 19384449

diff --git a/Engine/Source/Developer/RigVMDeveloper/Private/RigVMCompiler/RigVMCompiler.cpp b/Engine/Source/Developer/RigVMDeveloper/Private/RigVMCompiler/RigVMCompiler.cpp
--- a/Engine/Source/Developer/RigVMDeveloper/Private/RigVMCompiler/RigVMCompiler.cpp
+++ b/Engine/Source/Developer/RigVMDeveloper/Private/RigVMCompiler/RigVMCompiler.cpp
@@ -5,6 +5,7 @@
 #include "RigVMCore/RigVMExecuteContext.h"
 #include "RigVMDeveloperModule.h"
 #include "UObject/PropertyPortFlags.h"
+#include "UObject/Interface.h"
 #include "Stats/StatsHierarchical.h"
 #include "RigVMTypeUtils.h"
 
diff --git a/Engine/Source/Runtime/RigVM/Private/RigVMCore/RigVMMemoryStorage.cpp b/Engine/Source/Runtime/RigVM/Private/RigVMCore/RigVMMemoryStorage.cpp
--- a/Engine/Source/Runtime/RigVM/Private/RigVMCore/RigVMMemoryStorage.cpp
+++ b/Engine/Source/Runtime/RigVM/Private/RigVMCore/RigVMMemoryStorage.cpp
@@ -2,6 +2,7 @@
 
 #include "RigVMCore/RigVMMemoryStorage.h"
 #include "RigVMModule.h"
+#include "UObject/Interface.h"
 
 ////////////////////////////////////////////////////////////////////////////////////////////////////////////////////////
 
