use ContainsByPredicate instead of FindByPredicate

diff --git a/Engine/Source/Developer/RigVMDeveloper/Private/RigVMModel/RigVMController.cpp b/Engine/Source/Developer/RigVMDeveloper/Private/RigVMModel/RigVMController.cpp
--- a/Engine/Source/Developer/RigVMDeveloper/Private/RigVMModel/RigVMController.cpp
+++ b/Engine/Source/Developer/RigVMDeveloper/Private/RigVMModel/RigVMController.cpp
@@ -1560,15 +1560,15 @@ bool URigVMController::UnresolveTemplateNodes(const TArray<URigVMTemplateNode*>&
 	}
 
 	// check if any of the nodes needs to be unresolved
-	if(Algo::FindByPredicate(InNodes,
-		[](const URigVMTemplateNode* Node) -> bool
-		{
-			return !Node->IsFullyUnresolved();
-		}) == nullptr)
+	const bool bHasNodeToResolve = InNodes.ContainsByPredicate( [](const URigVMTemplateNode* Node) -> bool
+	{
+		return !Node->IsFullyUnresolved();
+	});
+	if (!bHasNodeToResolve)
 	{
 		return false;
 	}
-
+	
 	FRigVMBaseAction Action;
 	if(bSetupUndoRedo)
 	{
