Control Rig: Allow display name for pins on template nodes

diff --git a/Engine/Source/Developer/RigVMDeveloper/Private/RigVMModel/Nodes/RigVMTemplateNode.cpp b/Engine/Source/Developer/RigVMDeveloper/Private/RigVMModel/Nodes/RigVMTemplateNode.cpp
--- a/Engine/Source/Developer/RigVMDeveloper/Private/RigVMModel/Nodes/RigVMTemplateNode.cpp
+++ b/Engine/Source/Developer/RigVMDeveloper/Private/RigVMModel/Nodes/RigVMTemplateNode.cpp
@@ -211,22 +211,23 @@ bool URigVMTemplateNode::GetTypeMapForNewPinType(const URigVMPin* InPin, const F
 	return false;
 }
 
-TArray<int32> URigVMTemplateNode::GetResolvedPermutationIndices(FRigVMTemplate::FTypeMap* OutTypes) const
+const TArray<int32>& URigVMTemplateNode::GetResolvedPermutationIndices(FRigVMTemplate::FTypeMap* OutTypes) const
 {
+	if(!ResolvedPermutations.IsEmpty())
+	{
+		return ResolvedPermutations;
+	}
 	if (const FRigVMTemplate* Template = GetTemplate())
 	{
 		FRigVMTemplate::FTypeMap Types = GetResolvedTypes();
-		TArray<int32> PermutationIndices;
-		Template->Resolve(Types, PermutationIndices, false);
+		Template->Resolve(Types, ResolvedPermutations, false);
 
 		if (OutTypes)
 		{
 			*OutTypes = Types;
 		}
-
-		return PermutationIndices;
 	}
-	return TArray<int32>();
+	return ResolvedPermutations;
 }
 
 TArray<const FRigVMFunction*> URigVMTemplateNode::GetResolvedPermutations(FRigVMTemplate::FTypeMap* OutTypes) const
@@ -348,10 +349,40 @@ FString URigVMTemplateNode::GetInitialDefaultValueForPin(const FName& InRootPinN
 	return DefaultValue;
 }
 
+FName URigVMTemplateNode::GetDisplayNameForPin(const FName& InRootPinName,
+	const TArray<int32>& InPermutationIndices) const
+{
+#if WITH_EDITOR
+	if(const FRigVMTemplate* Template = GetTemplate())
+	{
+		const TArray<int32>* PermutationIndicesPtr = &InPermutationIndices;
+		if(PermutationIndicesPtr->IsEmpty())
+		{
+			PermutationIndicesPtr = &GetResolvedPermutationIndices();
+		}
+
+		const FText DisplayNameText = Template->GetDisplayNameForArgument(InRootPinName, *PermutationIndicesPtr);
+		if(DisplayNameText.IsEmpty())
+		{
+			return InRootPinName;
+		}
+
+		const FName DisplayName = *DisplayNameText.ToString();
+		if(DisplayName.IsEqual(InRootPinName))
+		{
+			return NAME_None;
+		}
+		return DisplayName;
+	}
+#endif
+	return NAME_None;
+}
+
 void URigVMTemplateNode::InvalidateCache()
 {
 	SupportedTypesCache.Reset();
 	CachedFunction = nullptr;
+	ResolvedPermutations.Reset();
 
 	for(URigVMPin* Pin : GetPins())
 	{
diff --git a/Engine/Source/Developer/RigVMDeveloper/Private/RigVMModel/RigVMController.cpp b/Engine/Source/Developer/RigVMDeveloper/Private/RigVMModel/RigVMController.cpp
--- a/Engine/Source/Developer/RigVMDeveloper/Private/RigVMModel/RigVMController.cpp
+++ b/Engine/Source/Developer/RigVMDeveloper/Private/RigVMModel/RigVMController.cpp
@@ -13051,6 +13051,16 @@ void URigVMController::ResolveTemplateNodeMetaData(URigVMTemplateNode* InNode, b
 	}
 #endif
 
+	for(URigVMPin* Pin : InNode->GetPins())
+	{
+		const FName DisplayName = InNode->GetDisplayNameForPin(Pin->GetFName());
+		if(Pin->DisplayName != DisplayName)
+		{
+			Pin->DisplayName = DisplayName;
+			Notify(ERigVMGraphNotifType::PinRenamed, Pin);
+		}
+	}
+
 	if(InNode->IsResolved())
 	{
 		for(URigVMPin* Pin : InNode->GetPins())
diff --git a/Engine/Source/Developer/RigVMDeveloper/Public/RigVMModel/Nodes/RigVMTemplateNode.h b/Engine/Source/Developer/RigVMDeveloper/Public/RigVMModel/Nodes/RigVMTemplateNode.h
--- a/Engine/Source/Developer/RigVMDeveloper/Public/RigVMModel/Nodes/RigVMTemplateNode.h
+++ b/Engine/Source/Developer/RigVMDeveloper/Public/RigVMModel/Nodes/RigVMTemplateNode.h
@@ -41,7 +41,7 @@ public:
 	bool GetTypeMapForNewPinType(const URigVMPin* InPin, const FString& InCPPType, UObject* InCPPTypeObject, FRigVMTemplate::FTypeMap& OutTypes) const;
 
 	// returns the index of the resolved RigVM function (or INDEX_NONE)
-	TArray<int32> GetResolvedPermutationIndices(FRigVMTemplate::FTypeMap* OutTypes = nullptr) const;
+	const TArray<int32>& GetResolvedPermutationIndices(FRigVMTemplate::FTypeMap* OutTypes = nullptr) const;
 
 	// returns the resolved functions for the template
 	TArray<const FRigVMFunction*> GetResolvedPermutations(FRigVMTemplate::FTypeMap* OutTypes = nullptr) const;
@@ -74,6 +74,9 @@ public:
 	// returns a default value for pin if it is known
 	FString GetInitialDefaultValueForPin(const FName& InRootPinName, const TArray<int32>& InPermutationIndices = TArray<int32>()) const;
 
+	// returns the display name for a pin
+	FName GetDisplayNameForPin(const FName& InRootPinName, const TArray<int32>& InPermutationIndices = TArray<int32>()) const;
+
 private:
 
 	void InvalidateCache();
@@ -88,6 +91,7 @@ private:
 
 	mutable const FRigVMTemplate* CachedTemplate;
 	mutable const FRigVMFunction* CachedFunction;
+	mutable TArray<int32> ResolvedPermutations;
 
 	friend class URigVMController;
 };
diff --git a/Engine/Source/Runtime/RigVM/Private/RigVMCore/RigVMTemplate.cpp b/Engine/Source/Runtime/RigVM/Private/RigVMCore/RigVMTemplate.cpp
--- a/Engine/Source/Runtime/RigVM/Private/RigVMCore/RigVMTemplate.cpp
+++ b/Engine/Source/Runtime/RigVM/Private/RigVMCore/RigVMTemplate.cpp
@@ -417,6 +417,60 @@ FText FRigVMTemplate::GetTooltipText(const TArray<int32>& InPermutationIndices)
 	return ResolvedTooltipText;
 }
 
+FText FRigVMTemplate::GetDisplayNameForArgument(const FName& InArgumentName, const TArray<int32>& InPermutationIndices) const
+{
+	if(const FRigVMTemplateArgument* Argument = FindArgument(InArgumentName))
+	{
+		FText ResolvedDisplayName;
+
+		auto VisitPermutation = [InArgumentName, &ResolvedDisplayName, this](int32 InPermutationIndex) -> bool
+		{
+			const FRigVMFunction* ResolvedFunction = GetPermutation(InPermutationIndex);
+			if(const FProperty* Property = ResolvedFunction->Struct->FindPropertyByName(InArgumentName))
+			{
+				const FText DisplayName = Property->GetDisplayNameText();
+				if (!ResolvedDisplayName.IsEmpty())
+				{
+					if(!ResolvedDisplayName.EqualTo(DisplayName))
+					{
+						ResolvedDisplayName = FText::FromName(InArgumentName);
+						return false;
+					}
+				}
+				else
+				{
+					ResolvedDisplayName = DisplayName;
+				}
+			}
+			return true;
+		};
+
+		if(InPermutationIndices.IsEmpty())
+		{
+			for(int32 PermutationIndex = 0; PermutationIndex < Permutations.Num(); PermutationIndex++)
+			{
+				if(!VisitPermutation(PermutationIndex))
+				{
+					break;
+				}
+			}
+		}
+		else
+		{
+			for(const int32 PermutationIndex : InPermutationIndices)
+			{
+				if(!VisitPermutation(PermutationIndex))
+				{
+					break;
+				}
+			}
+		}
+
+		return ResolvedDisplayName;
+	}
+	return FText();
+}
+
 #endif
 
 bool FRigVMTemplate::IsCompatible(const FRigVMTemplate& InOther) const
diff --git a/Engine/Source/Runtime/RigVM/Public/RigVMCore/RigVMTemplate.h b/Engine/Source/Runtime/RigVM/Public/RigVMCore/RigVMTemplate.h
--- a/Engine/Source/Runtime/RigVM/Public/RigVMCore/RigVMTemplate.h
+++ b/Engine/Source/Runtime/RigVM/Public/RigVMCore/RigVMTemplate.h
@@ -235,6 +235,9 @@ public:
 	// Returns the tooltip based on the permutation's metadata
 	FText GetTooltipText(const TArray<int32>& InPermutationIndices = TArray<int32>()) const;
 
+	// Returns the display name text for an argument 
+	FText GetDisplayNameForArgument(const FName& InArgumentName, const TArray<int32>& InPermutationIndices = TArray<int32>()) const;
+
 	FString GetCategory() const;
 	FString GetKeywords() const;
 
