RigVM: Performance: Cache interim results for BubbleUpExpressions

diff --git a/Engine/Source/Developer/RigVMDeveloper/Private/RigVMCompiler/RigVMAST.cpp b/Engine/Source/Developer/RigVMDeveloper/Private/RigVMCompiler/RigVMAST.cpp
--- a/Engine/Source/Developer/RigVMDeveloper/Private/RigVMCompiler/RigVMAST.cpp
+++ b/Engine/Source/Developer/RigVMDeveloper/Private/RigVMCompiler/RigVMAST.cpp
@@ -197,6 +197,17 @@ int32 FRigVMExprAST::GetMinChildIndexWithinParent(const FRigVMExprAST* InParentE
 {
 	int32 MinIndex = INDEX_NONE;
 
+	const TTuple<const FRigVMExprAST*, const FRigVMExprAST*> MapKey(InParentExpr, this);
+
+	const FRigVMParserAST* Parser = GetParser();
+	if(Parser)
+	{
+		if(const int32* IndexPtr = Parser->MinIndexOfChildWithinParent.Find(MapKey))
+		{
+			return *IndexPtr;
+		}
+	}
+
 	for (const FRigVMExprAST* Parent : Parents)
 	{
 		int32 ChildIndex = INDEX_NONE;
@@ -216,6 +227,7 @@ int32 FRigVMExprAST::GetMinChildIndexWithinParent(const FRigVMExprAST* InParentE
 		}
 	}
 
+	Parser->MinIndexOfChildWithinParent.Add(MapKey, MinIndex);
 	return MinIndex;
 }
 
@@ -1413,6 +1425,7 @@ void FRigVMParserAST::BubbleUpExpressions()
 			// of the cached value, so that the traverser sees it earlier
 			if (OuterBlock)
 			{
+				MinIndexOfChildWithinParent.Reset();
 				int32 ChildIndex = Expression->GetMinChildIndexWithinParent(OuterBlock);
 				if (ChildIndex != INDEX_NONE)
 				{
diff --git a/Engine/Source/Developer/RigVMDeveloper/Public/RigVMCompiler/RigVMAST.h b/Engine/Source/Developer/RigVMDeveloper/Public/RigVMCompiler/RigVMAST.h
--- a/Engine/Source/Developer/RigVMDeveloper/Public/RigVMCompiler/RigVMAST.h
+++ b/Engine/Source/Developer/RigVMDeveloper/Public/RigVMCompiler/RigVMAST.h
@@ -1427,6 +1427,7 @@ private:
 	TArray<URigVMLink*> LinksToSkip;
 
 	FRigVMParserASTSettings Settings;
+	mutable TMap<TTuple<const FRigVMExprAST*,const FRigVMExprAST*>, int32> MinIndexOfChildWithinParent;
 
 	friend class FRigVMExprAST;
 	friend class URigVMCompiler;
