RigVM: Allow multiple connections to the same source variable node

diff --git a/Engine/Source/Developer/RigVMDeveloper/Private/RigVMModel/RigVMPin.cpp b/Engine/Source/Developer/RigVMDeveloper/Private/RigVMModel/RigVMPin.cpp
--- a/Engine/Source/Developer/RigVMDeveloper/Private/RigVMModel/RigVMPin.cpp
+++ b/Engine/Source/Developer/RigVMDeveloper/Private/RigVMModel/RigVMPin.cpp
@@ -1721,7 +1721,9 @@ bool URigVMPin::CanLink(URigVMPin* InSourcePin, URigVMPin* InTargetPin, FString*
 		{
 			for (int32 SourceNodeIndex = 0; SourceNodeIndex < SourceNodes.Num(); SourceNodeIndex++)
 			{
-				bool bNodeCanLinkAnywhere = SourceNodes[SourceNodeIndex]->IsA<URigVMRerouteNode>();
+				bool bNodeCanLinkAnywhere =
+					SourceNodes[SourceNodeIndex]->IsA<URigVMRerouteNode>() ||
+					SourceNodes[SourceNodeIndex]->IsA<URigVMVariableNode>();
 				if (!bNodeCanLinkAnywhere)
 				{
 					if (URigVMUnitNode* UnitNode = Cast<URigVMUnitNode>(SourceNodes[SourceNodeIndex]))
@@ -1737,10 +1739,21 @@ bool URigVMPin::CanLink(URigVMPin* InSourcePin, URigVMPin* InTargetPin, FString*
 
 				if (!bNodeCanLinkAnywhere)
 				{
-					int32 SourceNodeInstructionIndex = InByteCode->GetFirstInstructionIndexForSubject(SourceNodes[SourceNodeIndex]);
+					const int32 SourceNodeInstructionIndex = InByteCode->GetFirstInstructionIndexForSubject(SourceNodes[SourceNodeIndex]);
 					if (SourceNodeInstructionIndex != INDEX_NONE &&
 						SourceNodeInstructionIndex > TargetNodeInstructionIndex)
 					{
+						if (OutFailureReason)
+						{
+							static constexpr TCHAR IncorrectNodeOrderMessage[] = TEXT("Source node %s (%s) and target node %s (%s) are in the incorrect order.");
+							*OutFailureReason = FString::Printf(
+								IncorrectNodeOrderMessage,
+								*SourceNodes[SourceNodeIndex]->GetName(),
+								*SourceNodes[SourceNodeIndex]->GetNodeTitle(),
+								*TargetNode->GetName(),
+								*TargetNode->GetNodeTitle());
+						}
+
 						return false;
 					}
 					SourceNodes.Append(SourceNodes[SourceNodeIndex]->GetLinkedSourceNodes());
