fixed missing include

diff --git a/Engine/Source/Developer/RigVMDeveloper/Private/RigVMModel/RigVMController.cpp b/Engine/Source/Developer/RigVMDeveloper/Private/RigVMModel/RigVMController.cpp
--- a/Engine/Source/Developer/RigVMDeveloper/Private/RigVMModel/RigVMController.cpp
+++ b/Engine/Source/Developer/RigVMDeveloper/Private/RigVMModel/RigVMController.cpp
@@ -17,6 +17,7 @@
 #include "Misc/CoreMisc.h"
 #include "Algo/Sort.h"
 #include "Algo/Count.h"
+#include "Algo/Transform.h"
 #include "RigVMPythonUtils.h"
 #include "RigVMTypeUtils.h"
 #include "Engine/UserDefinedStruct.h"
