RigVM: Fix sorting of children within blocks

diff --git a/Engine/Source/Developer/RigVMDeveloper/Private/RigVMCompiler/RigVMAST.cpp b/Engine/Source/Developer/RigVMDeveloper/Private/RigVMCompiler/RigVMAST.cpp
--- a/Engine/Source/Developer/RigVMDeveloper/Private/RigVMCompiler/RigVMAST.cpp
+++ b/Engine/Source/Developer/RigVMDeveloper/Private/RigVMCompiler/RigVMAST.cpp
@@ -221,9 +221,12 @@ int32 FRigVMExprAST::GetMinChildIndexWithinParent(const FRigVMExprAST* InParentE
 			ChildIndex = Parent->GetMinChildIndexWithinParent(InParentExpr);
 		}
 
-		if (ChildIndex < MinIndex || MinIndex == INDEX_NONE)
+		if (ChildIndex != INDEX_NONE)
 		{
-			MinIndex = ChildIndex;
+			if (ChildIndex < MinIndex || MinIndex == INDEX_NONE)
+			{
+				MinIndex = ChildIndex;
+			}
 		}
 	}
 
