Control Rig: Deprecate URigVMGraph::GetVariableDescriptions

diff --git a/Engine/Plugins/Animation/ControlRig/Source/ControlRigDeveloper/Private/ControlRigBlueprint.cpp b/Engine/Plugins/Animation/ControlRig/Source/ControlRigDeveloper/Private/ControlRigBlueprint.cpp
--- a/Engine/Plugins/Animation/ControlRig/Source/ControlRigDeveloper/Private/ControlRigBlueprint.cpp
+++ b/Engine/Plugins/Animation/ControlRig/Source/ControlRigDeveloper/Private/ControlRigBlueprint.cpp
@@ -2214,6 +2214,25 @@ TArray<UStruct*> UControlRigBlueprint::GetAvailableRigUnits()
 
 #if WITH_EDITOR
 
+TArray<FRigVMGraphVariableDescription> UControlRigBlueprint::GetMemberVariables() const
+{
+	TArray<FRigVMGraphVariableDescription> Variables;
+	for (const FBPVariableDescription& BPVariable : NewVariables)
+	{
+		FRigVMGraphVariableDescription NewVariable;
+		NewVariable.Name = BPVariable.VarName;
+		NewVariable.DefaultValue = BPVariable.DefaultValue;
+		FString CPPType;
+		UObject* CPPTypeObject;
+		RigVMTypeUtils::CPPTypeFromPinType(BPVariable.VarType, CPPType, &CPPTypeObject);
+		NewVariable.CPPType = CPPType;
+		NewVariable.CPPTypeObject = CPPTypeObject;
+		Variables.Add(NewVariable);
+	}
+
+	return Variables;
+}
+
 FName UControlRigBlueprint::AddMemberVariable(const FName& InName, const FString& InCPPType, bool bIsPublic, bool bIsReadOnly, FString InDefaultValue)
 {
 	FRigVMExternalVariable Variable = RigVMTypeUtils::ExternalVariableFromCPPTypePath(InName, InCPPType, bIsPublic, bIsReadOnly);
diff --git a/Engine/Plugins/Animation/ControlRig/Source/ControlRigDeveloper/Public/ControlRigBlueprint.h b/Engine/Plugins/Animation/ControlRig/Source/ControlRigDeveloper/Public/ControlRigBlueprint.h
--- a/Engine/Plugins/Animation/ControlRig/Source/ControlRigDeveloper/Public/ControlRigBlueprint.h
+++ b/Engine/Plugins/Animation/ControlRig/Source/ControlRigDeveloper/Public/ControlRigBlueprint.h
@@ -410,6 +410,9 @@ public:
 	static TArray<UStruct*> GetAvailableRigUnits();
 
 #if WITH_EDITOR
+	UFUNCTION(BlueprintCallable, Category = "Variables")
+	TArray<FRigVMGraphVariableDescription> GetMemberVariables() const;
+	
 	UFUNCTION(BlueprintCallable, Category = "Variables")
 	FName AddMemberVariable(const FName& InName, const FString& InCPPType, bool bIsPublic = false, bool bIsReadOnly = false, FString InDefaultValue = TEXT(""));
 
diff --git a/Engine/Plugins/Animation/ControlRig/Source/ControlRigEditor/Private/Graph/SControlRigGraphPinVariableName.cpp b/Engine/Plugins/Animation/ControlRig/Source/ControlRigEditor/Private/Graph/SControlRigGraphPinVariableName.cpp
--- a/Engine/Plugins/Animation/ControlRig/Source/ControlRigEditor/Private/Graph/SControlRigGraphPinVariableName.cpp
+++ b/Engine/Plugins/Animation/ControlRig/Source/ControlRigEditor/Private/Graph/SControlRigGraphPinVariableName.cpp
@@ -113,6 +113,25 @@ TArray<TSharedPtr<FString>>& SControlRigGraphPinVariableName::GetVariableNames()
 					VariableNames.Add(MakeShared<FString>(VariableDescription.Name.ToString()));
 				}
 			}
+
+			if (UControlRigBlueprint* Blueprint = Cast<UControlRigBlueprint>(ModelNode->GetGraph()->GetOuter()))
+			{
+				for (FBPVariableDescription& BPVariable : Blueprint->NewVariables)
+				{
+					FRigVMGraphVariableDescription* Found = VariableDescriptions.FindByPredicate([&BPVariable](const FRigVMGraphVariableDescription& Description) {
+						if (Description.Name == BPVariable.VarName)
+						{
+							return true;
+						}
+						return false;
+					});
+
+					if (!Found)
+					{
+						VariableNames.Add(MakeShared<FString>(BPVariable.VarName.ToString()));
+					}
+				}
+			}
 		}
 	}
 
diff --git a/Engine/Source/Developer/RigVMDeveloper/Private/RigVMModel/RigVMController.cpp b/Engine/Source/Developer/RigVMDeveloper/Private/RigVMModel/RigVMController.cpp
--- a/Engine/Source/Developer/RigVMDeveloper/Private/RigVMModel/RigVMController.cpp
+++ b/Engine/Source/Developer/RigVMDeveloper/Private/RigVMModel/RigVMController.cpp
@@ -1188,16 +1188,16 @@ void URigVMController::OnExternalVariableRemoved(const FName& InVarName, bool bS
 	}
 }
 
-void URigVMController::OnExternalVariableRenamed(const FName& InOldVarName, const FName& InNewVarName, bool bSetupUndoRedo)
+bool URigVMController::OnExternalVariableRenamed(const FName& InOldVarName, const FName& InNewVarName, bool bSetupUndoRedo)
 {
 	if (!IsValidGraph())
 	{
-		return;
+		return false;
 	}
 
 	if (!InOldVarName.IsValid() || !InNewVarName.IsValid())
 	{
-		return;
+		return false;
 	}
 
 	URigVMGraph* Graph = GetGraph();
@@ -1207,7 +1207,7 @@ void URigVMController::OnExternalVariableRenamed(const FName& InOldVarName, cons
 	{
 		if (InOldVarName == LocalVariable.Name)
 		{
-			return;
+			return false;
 		}
 	}
 
@@ -1274,6 +1274,8 @@ void URigVMController::OnExternalVariableRenamed(const FName& InOldVarName, cons
 	{
 		CloseUndoBracket();
 	}
+
+	return true;
 }
 
 void URigVMController::OnExternalVariableTypeChanged(const FName& InVarName, const FString& InCPPType, UObject* InCPPTypeObject, bool bSetupUndoRedo)
@@ -9901,20 +9903,23 @@ bool URigVMController::SetVariableName(URigVMVariableNode* InVariableNode, const
 	URigVMGraph* Graph = GetGraph();
 	check(Graph);
 
-	TArray<FRigVMGraphVariableDescription> Descriptions = Graph->GetVariableDescriptions();
+	TArray<FRigVMExternalVariable> Descriptions = GetAllVariables();
 	TMap<FName, int32> NameToIndex;
 	for (int32 VariableIndex = 0; VariableIndex < Descriptions.Num(); VariableIndex++)
 	{
 		NameToIndex.Add(Descriptions[VariableIndex].Name, VariableIndex);
 	}
 
-	FName VariableName = GetUniqueName(InVariableName, [Descriptions, NameToIndex, InVariableNode](const FName& InName) {
+	const FRigVMExternalVariable VariableType = RigVMTypeUtils::ExternalVariableFromCPPType(InVariableName, InVariableNode->GetCPPType(), InVariableNode->GetCPPTypeObject());
+	FName VariableName = GetUniqueName(InVariableName, [Descriptions, NameToIndex, VariableType](const FName& InName) {
 		const int32* FoundIndex = NameToIndex.Find(InName);
 		if (FoundIndex == nullptr)
 		{
 			return true;
 		}
-		return InVariableNode->GetCPPType() == Descriptions[*FoundIndex].CPPType;
+		return VariableType.TypeName == Descriptions[*FoundIndex].TypeName &&
+				VariableType.TypeObject == Descriptions[*FoundIndex].TypeObject &&
+				VariableType.bIsArray == Descriptions[*FoundIndex].bIsArray;
 	}, false, true);
 
 	int32 NodesSharingName = 0;
diff --git a/Engine/Source/Developer/RigVMDeveloper/Private/RigVMModel/RigVMControllerActions.cpp b/Engine/Source/Developer/RigVMDeveloper/Private/RigVMModel/RigVMControllerActions.cpp
--- a/Engine/Source/Developer/RigVMDeveloper/Private/RigVMModel/RigVMControllerActions.cpp
+++ b/Engine/Source/Developer/RigVMDeveloper/Private/RigVMModel/RigVMControllerActions.cpp
@@ -1207,12 +1207,12 @@ bool FRigVMRenameVariableAction::Undo(URigVMController* InController)
 	{
 		return false;
 	}
-	return InController->RenameVariable(*NewVariableName, *OldVariableName, false);
+	return InController->OnExternalVariableRenamed(*NewVariableName, *OldVariableName, false);
 }
 
 bool FRigVMRenameVariableAction::Redo(URigVMController* InController)
 {
-	if(!InController->RenameVariable(*OldVariableName, *NewVariableName, false))
+	if(!InController->OnExternalVariableRenamed(*OldVariableName, *NewVariableName, false))
 	{
 		return false;
 	}
diff --git a/Engine/Source/Developer/RigVMDeveloper/Public/RigVMModel/RigVMController.h b/Engine/Source/Developer/RigVMDeveloper/Public/RigVMModel/RigVMController.h
--- a/Engine/Source/Developer/RigVMDeveloper/Public/RigVMModel/RigVMController.h
+++ b/Engine/Source/Developer/RigVMDeveloper/Public/RigVMModel/RigVMController.h
@@ -204,7 +204,7 @@ public:
 
 	// Renames the variable name in all relevant nodes
 	UFUNCTION(BlueprintCallable, Category = RigVMController)
-	void OnExternalVariableRenamed(const FName& InOldVarName, const FName& InNewVarName, bool bSetupUndoRedo);
+	bool OnExternalVariableRenamed(const FName& InOldVarName, const FName& InNewVarName, bool bSetupUndoRedo);
 
 	// Changes the data type of all nodes matching a given variable name
 	UFUNCTION(BlueprintCallable, Category = RigVMController)
@@ -554,7 +554,7 @@ public:
 
 	// Renames a variable in the graph.
 	// This causes a VariableRenamed modified event.
-	UFUNCTION(BlueprintCallable, Category = RigVMController)
+	UFUNCTION(BlueprintCallable, Category = RigVMController, meta=(DeprecatedFunction))
 	bool RenameVariable(const FName& InOldName, const FName& InNewName, bool bSetupUndoRedo = true);
 
 	// Renames a parameter in the graph.
diff --git a/Engine/Source/Developer/RigVMDeveloper/Public/RigVMModel/RigVMControllerActions.h b/Engine/Source/Developer/RigVMDeveloper/Public/RigVMModel/RigVMControllerActions.h
--- a/Engine/Source/Developer/RigVMDeveloper/Public/RigVMModel/RigVMControllerActions.h
+++ b/Engine/Source/Developer/RigVMDeveloper/Public/RigVMModel/RigVMControllerActions.h
@@ -1002,7 +1002,7 @@ public:
 /**
  * An action renaming a variable in the graph.
  */
-USTRUCT()
+USTRUCT(meta=(Deprecated = "5.1"))
 struct FRigVMRenameVariableAction : public FRigVMBaseAction
 {
 	GENERATED_BODY()
