RigVM: Follow links on non-existing subpins

In the case of arrays we can have links on elements that don't exist on the source pin - but may exist on intermediate reroutes or function boundaries (entry or return nodes).

#b sara.schvartzman

diff --git a/Engine/Source/Developer/RigVMDeveloper/Private/RigVMCompiler/RigVMAST.cpp b/Engine/Source/Developer/RigVMDeveloper/Private/RigVMCompiler/RigVMAST.cpp
--- a/Engine/Source/Developer/RigVMDeveloper/Private/RigVMCompiler/RigVMAST.cpp
+++ b/Engine/Source/Developer/RigVMDeveloper/Private/RigVMCompiler/RigVMAST.cpp
@@ -881,16 +881,17 @@ FRigVMExprAST* FRigVMParserAST::TraverseMutableNode(const FRigVMASTProxy& InNode
 						}
 					}
 
-					const TArray<FRigVMASTProxy>& TargetPins = GetTargetPins(SourcePinProxy);
-					for(const FRigVMASTProxy& TargetPinProxy : TargetPins)
+					const TArray<int32>& LinkIndices = GetTargetLinkIndices(SourcePinProxy);
+					for(const int32 LinkIndex : LinkIndices)
 					{
-						if(ShouldLinkBeSkipped(FRigVMPinProxyPair(SourcePinProxy, TargetPinProxy)))
+						const FRigVMASTLinkDescription& Link = Links[LinkIndex];
+						if(ShouldLinkBeSkipped(Link))
 						{
 							continue;
 						}
 
-						URigVMNode* TargetNode = TargetPinProxy.GetSubjectChecked<URigVMPin>()->GetNode();
-						FRigVMASTProxy TargetNodeProxy = TargetPinProxy.GetSibling(TargetNode);
+						URigVMNode* TargetNode = Link.TargetProxy.GetSubjectChecked<URigVMPin>()->GetNode();
+						FRigVMASTProxy TargetNodeProxy = Link.TargetProxy.GetSibling(TargetNode);
 						TraverseMutableNode(TargetNodeProxy, ParentExpr);
 					}
 				}
@@ -1048,13 +1049,13 @@ FRigVMExprAST* FRigVMParserAST::TraversePin(const FRigVMASTProxy& InPinProxy, FR
 	URigVMPin* Pin = InPinProxy.GetSubjectChecked<URigVMPin>();
 	URigVMPin::FPinOverride PinOverride(InPinProxy, PinOverrides);
 
-	TArray<FRigVMPinProxyPair> Links = GetSourceLinks(InPinProxy, true);
+	TArray<int32> LinkIndices = GetSourceLinkIndices(InPinProxy, true);
 	
 	if (LinksToSkip.Num() > 0)
 	{
-		Links.RemoveAll([this](const FRigVMPinProxyPair& LinkToCheck)
+		LinkIndices.RemoveAll([this](int32 LinkIndex)
 			{
-				return this->ShouldLinkBeSkipped(LinkToCheck);
+				return this->ShouldLinkBeSkipped(Links[LinkIndex]);
 			}
 		);
 	}
@@ -1078,14 +1079,15 @@ FRigVMExprAST* FRigVMParserAST::TraversePin(const FRigVMASTProxy& InPinProxy, FR
 
 	if ((Pin->GetDirection() == ERigVMPinDirection::Input ||
 		Pin->GetDirection() == ERigVMPinDirection::Visible) &&
-		Links.Num() == 0)
+		LinkIndices.Num() == 0)
 	{
 		if (Cast<URigVMVariableNode>(Pin->GetNode()))
 		{
 			PinExpr = MakeExpr<FRigVMVarExprAST>(FRigVMExprAST::EType::Var, InPinProxy);
 			FRigVMExprAST* PinLiteralExpr = MakeExpr<FRigVMLiteralExprAST>(InPinProxy);
 			PinLiteralExpr->Name = PinExpr->Name;
-			FRigVMExprAST* PinCopyExpr = MakeExpr<FRigVMCopyExprAST>(InPinProxy, InPinProxy);
+			const FRigVMASTLinkDescription LiteralLink(InPinProxy, InPinProxy, FString());
+			FRigVMExprAST* PinCopyExpr = MakeExpr<FRigVMCopyExprAST>(LiteralLink);
 			PinCopyExpr->AddParent(PinExpr);
 			PinLiteralExpr->AddParent(PinCopyExpr);
 		}
@@ -1123,9 +1125,9 @@ FRigVMExprAST* FRigVMParserAST::TraversePin(const FRigVMASTProxy& InPinProxy, FR
 	{
 		bool bHasSourceLinkToRoot = false;
 		URigVMPin* RootPin = Pin->GetRootPin();
-		for (const FRigVMPinProxyPair& SourceLink : Links)
+		for (const int32 LinkIndex : LinkIndices)
 		{
-			if (SourceLink.Value.GetSubject() == RootPin)
+			if (Links[LinkIndex].TargetProxy.GetSubject() == RootPin)
 			{
 				bHasSourceLinkToRoot = true;
 				break;
@@ -1133,12 +1135,13 @@ FRigVMExprAST* FRigVMParserAST::TraversePin(const FRigVMASTProxy& InPinProxy, FR
 		}
 
 		if (!bHasSourceLinkToRoot && 
-			GetSourcePins(InPinProxy).Num() == 0 &&
-			(Pin->GetDirection() == ERigVMPinDirection::IO || Links.Num() > 0))
+			GetSourceLinkIndices(InPinProxy, false).Num() == 0 &&
+			(Pin->GetDirection() == ERigVMPinDirection::IO || LinkIndices.Num() > 0))
 		{
 			FRigVMLiteralExprAST* LiteralExpr = MakeExpr<FRigVMLiteralExprAST>(InPinProxy);
-			FRigVMCopyExprAST* LiteralCopyExpr = MakeExpr<FRigVMCopyExprAST>(InPinProxy, InPinProxy);
-			LiteralCopyExpr->Name = *FString::Printf(TEXT("%s -> %s"), *Pin->GetPinPath(), *Pin->GetPinPath());
+			const FRigVMASTLinkDescription LiteralLink(InPinProxy, InPinProxy, FString());
+			FRigVMCopyExprAST* LiteralCopyExpr = MakeExpr<FRigVMCopyExprAST>(LiteralLink);
+			LiteralCopyExpr->Name = *GetLinkAsString(LiteralCopyExpr->GetLink());
 			LiteralCopyExpr->AddParent(PinExpr);
 			LiteralExpr->AddParent(LiteralCopyExpr);
 			LiteralExpr->Name = *Pin->GetPinPath();
@@ -1151,7 +1154,7 @@ FRigVMExprAST* FRigVMParserAST::TraversePin(const FRigVMASTProxy& InPinProxy, FR
 
 	if ((Pin->GetDirection() == ERigVMPinDirection::IO || Pin->GetDirection() == ERigVMPinDirection::Input) &&
 		(InParentExpr->IsA(FRigVMExprAST::If) || InParentExpr->IsA(FRigVMExprAST::Select)) &&
-		Links.Num() > 0)
+		LinkIndices.Num() > 0)
 	{
 		FRigVMBlockExprAST* BlockExpr = MakeExpr<FRigVMBlockExprAST>(FRigVMExprAST::EType::Block, FRigVMASTProxy());
 		BlockExpr->AddParent(PinExpr);
@@ -1159,25 +1162,26 @@ FRigVMExprAST* FRigVMParserAST::TraversePin(const FRigVMASTProxy& InPinProxy, FR
 		ParentExprForLinks = BlockExpr;
 	}
 
-	for (const FRigVMPinProxyPair& SourceLink : Links)
+	for (const int32 LinkIndex : LinkIndices)
 	{
-		TraverseLink(SourceLink, ParentExprForLinks);
+		TraverseLink(LinkIndex, ParentExprForLinks);
 	}
 
 	return PinExpr;
 }
 
-FRigVMExprAST* FRigVMParserAST::TraverseLink(const FRigVMPinProxyPair& InLink, FRigVMExprAST* InParentExpr)
+FRigVMExprAST* FRigVMParserAST::TraverseLink(int32 InLinkIndex, FRigVMExprAST* InParentExpr)
 {
-	const FRigVMASTProxy& SourceProxy = InLink.Key;
-	const FRigVMASTProxy& TargetProxy = InLink.Value;
+	const FRigVMASTLinkDescription& Link = Links[InLinkIndex];
+	const FRigVMASTProxy& SourceProxy = Link.SourceProxy;
+	const FRigVMASTProxy& TargetProxy = Link.TargetProxy;
 	URigVMPin* SourcePin = SourceProxy.GetSubjectChecked<URigVMPin>();
 	URigVMPin* TargetPin = TargetProxy.GetSubjectChecked<URigVMPin>();
 	URigVMPin* SourceRootPin = SourcePin->GetRootPin();
 	URigVMPin* TargetRootPin = TargetPin->GetRootPin();
 	FRigVMASTProxy SourceNodeProxy = SourceProxy.GetSibling(SourcePin->GetNode());
 
-	bool bRequiresCopy = SourceRootPin != SourcePin || TargetRootPin != TargetPin;
+	bool bRequiresCopy = SourceRootPin != SourcePin || TargetRootPin != TargetPin || !Link.SegmentPath.IsEmpty();
 	if (!bRequiresCopy)
 	{
 		if(Cast<URigVMVariableNode>(TargetRootPin->GetNode()))
@@ -1189,14 +1193,14 @@ FRigVMExprAST* FRigVMParserAST::TraverseLink(const FRigVMPinProxyPair& InLink, F
 	FRigVMAssignExprAST* AssignExpr = nullptr;
 	if (bRequiresCopy)
 	{
-		AssignExpr = MakeExpr<FRigVMCopyExprAST>(SourceProxy, TargetProxy);
+		AssignExpr = MakeExpr<FRigVMCopyExprAST>(Link);
 	}
 	else
 	{
-		AssignExpr = MakeExpr<FRigVMAssignExprAST>(FRigVMExprAST::EType::Assign, SourceProxy, TargetProxy);
+		AssignExpr = MakeExpr<FRigVMAssignExprAST>(FRigVMExprAST::EType::Assign, Link);
 	}
 
-	AssignExpr->Name = *GetLinkAsString(InLink);
+	AssignExpr->Name = *GetLinkAsString(Link);
 	AssignExpr->AddParent(InParentExpr);
 
 	FRigVMExprAST* NodeExpr = TraverseNode(SourceNodeProxy, AssignExpr);
@@ -1515,6 +1519,14 @@ void FRigVMParserAST::FoldAssignments()
 		URigVMPin* SourcePin = AssignExpr->GetSourcePin();
 		URigVMPin* TargetPin = AssignExpr->GetTargetPin();
 
+		// it's possible that we'll see copies from an element onto an array
+		// here. we'll need to ignore these links and leave them since they
+		// represent copies.
+		if(SourcePin->IsArray() != TargetPin->IsArray())
+		{
+			continue;
+		}
+
 		// in case the assign has different types for left and right - we need to avoid folding
 		// since this assign represents a cast operation
 		if(SourcePin->GetCPPTypeObject() != TargetPin->GetCPPTypeObject())
@@ -1605,8 +1617,9 @@ void FRigVMParserAST::FoldAssignments()
 			if (GrandParent->IsA(FRigVMExprAST::EType::Assign))
 			{
 				FRigVMAssignExprAST* GrandParentAssign = GrandParent->To<FRigVMAssignExprAST>();
-				GrandParentAssign->SourceProxy = AssignExpr->SourceProxy;
-				GrandParentAssign->Name = *FString::Printf(TEXT("%s -> %s"), *GrandParentAssign->GetSourcePin()->GetPinPath(), *GrandParentAssign->GetTargetPin()->GetPinPath());
+				GrandParentAssign->Link = FRigVMASTLinkDescription(AssignExpr->GetSourceProxy(),
+					GrandParentAssign->GetTargetProxy(), FString());
+				GrandParentAssign->Name = *GetLinkAsString(GrandParentAssign->GetLink());
 			}
 		}
 
@@ -1699,8 +1712,8 @@ bool FRigVMParserAST::FoldConstantValuesToLiterals(URigVMGraph* InGraph, URigVMC
 				}
 			}
 
-			TArray<FRigVMPinProxyPair> SourcePins = GetSourceLinks(PinProxy, true);
-			if (SourcePins.Num() == 0)
+			TArray<int32> LinkIndices = GetSourceLinkIndices(PinProxy, true);
+			if (LinkIndices.Num() == 0)
 			{
 				continue;
 			}
@@ -1711,11 +1724,11 @@ bool FRigVMParserAST::FoldConstantValuesToLiterals(URigVMGraph* InGraph, URigVMC
 			}
 
 			bool bFoundValidSourcePin = false;
-			for (const FRigVMPinProxyPair& SourcePin : SourcePins)
+			for (const int32 LinkIndex : LinkIndices)
 			{
-				const FRigVMASTProxy& SourcePinProxy = SourcePin.Key;
+				const FRigVMASTProxy& SourcePinProxy = Links[LinkIndex].SourceProxy;
 				URigVMNode* SourceNode = SourcePinProxy.GetSubjectChecked<URigVMPin>()->GetNode();
-				FRigVMASTProxy SourceNodeProxy = SourcePin.Key.GetSibling(SourceNode);
+				FRigVMASTProxy SourceNodeProxy = SourcePinProxy.GetSibling(SourceNode);
 
 				check(SourceNodeProxy.IsValid());
 
@@ -1889,12 +1902,12 @@ bool FRigVMParserAST::FoldConstantValuesToLiterals(URigVMGraph* InGraph, URigVMC
 			}
 		}
 
-		const TArray<FRigVMASTProxy>& TargetPins = GetTargetPins(PinToComputeProxy);
-		for (const FRigVMASTProxy& TargetPinProxy : TargetPins)
+		const TArray<int32> LinkIndices = GetTargetLinkIndices(PinToComputeProxy);
+		for (const int32 LinkIndex : LinkIndices)
 		{
 			URigVMPin::FPinOverrideValue OverrideValue;
 			OverrideValue.DefaultValue = DefaultValue;
-			PinOverrides.FindOrAdd(TargetPinProxy) = OverrideValue;
+			PinOverrides.FindOrAdd(Links[LinkIndex].TargetProxy) = OverrideValue;
 		}
 	}
 
@@ -2740,32 +2753,24 @@ void FRigVMParserAST::TraverseChildren(const FRigVMExprAST* InExpr, TFunctionRef
 	}
 }
 
-const TArray<FRigVMASTProxy>& FRigVMParserAST::GetSourcePins(const FRigVMASTProxy& InPinProxy) const
+TArray<int32> FRigVMParserAST::GetSourceLinkIndices(const FRigVMASTProxy& InPinProxy, bool bRecursive) const
 {
-	if (const TArray<FRigVMASTProxy>* LinkedPins = SourceLinks.Find(InPinProxy))
-	{
-		return *LinkedPins;
-	}
-	return EmptyProxyArray;
+	return GetLinkIndices(InPinProxy, true, bRecursive);
 }
 
-const TArray<FRigVMASTProxy>& FRigVMParserAST::GetTargetPins(const FRigVMASTProxy& InPinProxy) const
+TArray<int32> FRigVMParserAST::GetTargetLinkIndices(const FRigVMASTProxy& InPinProxy, bool bRecursive) const
 {
-	if (const TArray<FRigVMASTProxy>* LinkedPins = TargetLinks.Find(InPinProxy))
-	{
-		return *LinkedPins;
-	}
-	return EmptyProxyArray;
+	return GetLinkIndices(InPinProxy, false, bRecursive);
 }
 
-TArray<FRigVMPinProxyPair> FRigVMParserAST::GetSourceLinks(const FRigVMASTProxy& InPinProxy, bool bRecursive) const
+TArray<int32> FRigVMParserAST::GetLinkIndices(const FRigVMASTProxy& InPinProxy, bool bGetSource, bool bRecursive) const
 {
-	const TArray<FRigVMASTProxy>& SourcePins = GetSourcePins(InPinProxy);
+	const static TArray<int32> EmptyIntArray;
+	TArray<int32> LinkIndices;
 
-	TArray<FRigVMPinProxyPair> Pairs;
-	for (const FRigVMASTProxy& SourcePin : SourcePins)
+	if(const TArray<int32>* LinkIndicesPtr = (bGetSource ? SourceLinkIndices : TargetLinkIndices).Find(InPinProxy))
 	{
-		Pairs.Add(FRigVMPinProxyPair(SourcePin, InPinProxy));
+		LinkIndices = *LinkIndicesPtr;
 	}
 
 	if (bRecursive)
@@ -2774,34 +2779,17 @@ TArray<FRigVMPinProxyPair> FRigVMParserAST::GetSourceLinks(const FRigVMASTProxy&
 		for (URigVMPin* SubPin : Pin->GetSubPins())
 		{
 			FRigVMASTProxy SubPinProxy = InPinProxy.GetSibling(SubPin);
-			Pairs.Append(GetSourceLinks(SubPinProxy, true));
+			LinkIndices.Append(GetLinkIndices(SubPinProxy, bGetSource, true));
 		}
 	}
 
-	return Pairs;
+	return LinkIndices;
 }
 
-TArray<FRigVMPinProxyPair> FRigVMParserAST::GetTargetLinks(const FRigVMASTProxy& InPinProxy, bool bRecursive) const
-{
-	const TArray<FRigVMASTProxy>& TargetPins = GetTargetPins(InPinProxy);
 
-	TArray<FRigVMPinProxyPair> Pairs;
-	for (const FRigVMASTProxy& TargetPin : TargetPins)
-	{
-		Pairs.Add(FRigVMPinProxyPair(InPinProxy, TargetPin));
-	}
-
-	if (bRecursive)
-	{
-		URigVMPin* Pin = InPinProxy.GetSubjectChecked<URigVMPin>();
-		for (URigVMPin* SubPin : Pin->GetSubPins())
-		{
-			FRigVMASTProxy SubPinProxy = InPinProxy.GetSibling(SubPin);
-			Pairs.Append(GetTargetLinks(SubPinProxy, true));
-		}
-	}
-
-	return Pairs;
+const FRigVMASTLinkDescription& FRigVMParserAST::GetLink(int32 InLinkIndex) const
+{
+	return Links[InLinkIndex];
 }
 
 void FRigVMParserAST::Inline(URigVMGraph* InGraph)
@@ -2822,8 +2810,9 @@ void FRigVMParserAST::Inline(URigVMGraph* InGraph, const TArray<FRigVMASTProxy>&
 	{
 		URigVMPin::FPinOverrideMap* PinOverrides;
 		TMap<FRigVMASTProxy, FRigVMASTProxy>* SourcePins;
-		TMap<FRigVMASTProxy, TArray<FRigVMASTProxy>>* TargetLinks;
-		TMap<FRigVMASTProxy, TArray<FRigVMASTProxy>>* SourceLinks;
+		TMap<FRigVMASTProxy, TArray<int32>>* TargetLinkIndices;
+		TMap<FRigVMASTProxy, TArray<int32>>* SourceLinkIndices;
+		TArray<FRigVMASTLinkDescription>* Links;
 		TArray<FRigVMASTProxy> LibraryNodeCallstack;
 		FRigVMParserASTSettings* Settings;
 
@@ -2853,8 +2842,14 @@ void FRigVMParserAST::Inline(URigVMGraph* InGraph, const TArray<FRigVMASTProxy>&
 		}
 
 		static FRigVMASTProxy FindSourcePin(const FRigVMASTProxy& InPinProxy, LocalPinTraversalInfo& OutTraversalInfo)
+		{
+			return FindSourcePin(InPinProxy, InPinProxy, OutTraversalInfo);
+		}
+
+		static FRigVMASTProxy FindSourcePin(const FRigVMASTProxy& InPinProxy, const FRigVMASTProxy& InPinProxyForMap, LocalPinTraversalInfo& OutTraversalInfo)
 		{
 			URigVMPin* Pin = InPinProxy.GetSubjectChecked<URigVMPin>();
+			const bool bStoreSourcePinOnMap = InPinProxy == InPinProxyForMap;
 
 			// if this pin is a root on a library
 			if (Pin->GetParentPin() == nullptr)
@@ -2876,7 +2871,10 @@ void FRigVMParserAST::Inline(URigVMGraph* InGraph, const TArray<FRigVMASTProxy>&
 									FRigVMASTProxy ReturnPinProxy = LibraryNodeProxy.GetChild(ReturnPin);
 									FRigVMASTProxy SourcePinProxy = FindSourcePin(ReturnPinProxy, OutTraversalInfo);
 									SourcePinProxy = SourcePinProxy.IsValid() ? SourcePinProxy : ReturnPinProxy;
-									OutTraversalInfo.SourcePins->FindOrAdd(InPinProxy) = SourcePinProxy;
+									if(bStoreSourcePinOnMap)
+									{
+										OutTraversalInfo.SourcePins->FindOrAdd(InPinProxyForMap) = SourcePinProxy;
+									}
 									OutTraversalInfo.LibraryNodeCallstack.Pop();
 									return SourcePinProxy;
 
@@ -2895,7 +2893,10 @@ void FRigVMParserAST::Inline(URigVMGraph* InGraph, const TArray<FRigVMASTProxy>&
 									FRigVMASTProxy EntryPinProxy = InPinProxy.GetSibling(EntryPin);
 									FRigVMASTProxy SourcePinProxy = FindSourcePin(EntryPinProxy, OutTraversalInfo);
 									SourcePinProxy = SourcePinProxy.IsValid() ? SourcePinProxy : EntryPinProxy;
-									OutTraversalInfo.SourcePins->FindOrAdd(InPinProxy) = SourcePinProxy;
+									if(bStoreSourcePinOnMap)
+									{
+										OutTraversalInfo.SourcePins->FindOrAdd(InPinProxyForMap) = SourcePinProxy;
+									}
 									return SourcePinProxy;
 								} 
 							}
@@ -2919,7 +2920,10 @@ void FRigVMParserAST::Inline(URigVMGraph* InGraph, const TArray<FRigVMASTProxy>&
 									FRigVMASTProxy LibraryPinProxy = LibraryNodeProxy.GetSibling(LibraryPin);
 									FRigVMASTProxy SourcePinProxy = FindSourcePin(LibraryPinProxy, OutTraversalInfo);
 									SourcePinProxy = SourcePinProxy.IsValid() ? SourcePinProxy : LibraryPinProxy;
-									OutTraversalInfo.SourcePins->FindOrAdd(InPinProxy) = SourcePinProxy;
+									if(bStoreSourcePinOnMap)
+									{
+										OutTraversalInfo.SourcePins->FindOrAdd(InPinProxyForMap) = SourcePinProxy;
+									}
 									return SourcePinProxy;
 								}
 							}
@@ -2944,11 +2948,11 @@ void FRigVMParserAST::Inline(URigVMGraph* InGraph, const TArray<FRigVMASTProxy>&
 				}
 			}
 
-			if (!bIOPinOnLeftOfLibraryNode)
+			if (!bIOPinOnLeftOfLibraryNode && bStoreSourcePinOnMap)
 			{
 				// note: this map isn't going to work for functions which are referenced.
 				// (since the pin objects are shared between multiple invocation nodes)
-				if (const FRigVMASTProxy* SourcePinProxy = OutTraversalInfo.SourcePins->Find(InPinProxy))
+				if (const FRigVMASTProxy* SourcePinProxy = OutTraversalInfo.SourcePins->Find(InPinProxyForMap))
 				{
 					return *SourcePinProxy;
 				}
@@ -3098,16 +3102,22 @@ void FRigVMParserAST::Inline(URigVMGraph* InGraph, const TArray<FRigVMASTProxy>&
 				}
 			}
 
-			if (!bIOPinOnLeftOfLibraryNode)
+			if (!bIOPinOnLeftOfLibraryNode && bStoreSourcePinOnMap)
 			{
-				OutTraversalInfo.SourcePins->FindOrAdd(InPinProxy) = SourcePinProxy;
+				OutTraversalInfo.SourcePins->FindOrAdd(InPinProxyForMap) = SourcePinProxy;
 			}
 			return SourcePinProxy;
 		}
 
 		static void VisitPin(const FRigVMASTProxy& InPinProxy, LocalPinTraversalInfo& OutTraversalInfo)
 		{
-			FRigVMASTProxy SourcePinProxy = FindSourcePin(InPinProxy, OutTraversalInfo);
+			return VisitPin(InPinProxy, InPinProxy, OutTraversalInfo, FString());
+		}
+
+		static void VisitPin(const FRigVMASTProxy& InPinProxy, const FRigVMASTProxy& InPinProxyForMap,
+			LocalPinTraversalInfo& OutTraversalInfo, const FString& InSegmentPath)
+		{
+			const FRigVMASTProxy SourcePinProxy = FindSourcePin(InPinProxy, InPinProxyForMap, OutTraversalInfo);
 			if (SourcePinProxy.IsValid())
 			{
 				// The source pin is the final determined source pin, since
@@ -3125,13 +3135,30 @@ void FRigVMParserAST::Inline(URigVMGraph* InGraph, const TArray<FRigVMASTProxy>&
 					SourceNode->IsA<URigVMFunctionReturnNode>())
 				{
 					OutTraversalInfo.PinOverrides->FindOrAdd(InPinProxy) = URigVMPin::FPinOverrideValue(SourcePin);
+
+					// for arrays - if there are sub-pins on the determined source, we need to walk those as well
+					if(SourcePin->IsArray())
+					{
+						TArray<URigVMPin*> SourceSubPins = SourcePin->GetSubPins();
+						for (int32 SourceSubPinIndex = 0; SourceSubPinIndex < SourceSubPins.Num(); SourceSubPinIndex++)
+						{
+							URigVMPin* SourceSubPin = SourceSubPins[SourceSubPinIndex];
+							const FRigVMASTProxy SubPinProxy = SourcePinProxy.GetSibling(SourceSubPin);
+							const FString SegmentPath = SourceSubPin->GetSubPinPath(SourcePin, false);
+							VisitPin(SubPinProxy, InPinProxy, OutTraversalInfo, SegmentPath);
+							SourceSubPins.Append(SourceSubPin->GetSubPins());
+						}
+					}
 				}
 				else
 				{
-					if (IsValidLinkForAST(SourcePinProxy, InPinProxy))
+					if (IsValidLinkForAST(SourcePinProxy, InPinProxyForMap))
 					{
-						OutTraversalInfo.SourceLinks->FindOrAdd(InPinProxy).Add(SourcePinProxy);
-						OutTraversalInfo.TargetLinks->FindOrAdd(SourcePinProxy).Add(InPinProxy);
+						FRigVMASTLinkDescription Link(SourcePinProxy, InPinProxyForMap, InSegmentPath);
+						Link.LinkIndex = OutTraversalInfo.Links->Num();
+						OutTraversalInfo.Links->Add(Link);
+						OutTraversalInfo.SourceLinkIndices->FindOrAdd(InPinProxyForMap).Add(Link.LinkIndex);
+						OutTraversalInfo.TargetLinkIndices->FindOrAdd(SourcePinProxy).Add(Link.LinkIndex);
 					}
 				}
 			}
@@ -3202,8 +3229,8 @@ void FRigVMParserAST::Inline(URigVMGraph* InGraph, const TArray<FRigVMASTProxy>&
 	};
 
 	NodeProxies.Reset();
-	SourceLinks.Reset();
-	TargetLinks.Reset();
+	SourceLinkIndices.Reset();
+	TargetLinkIndices.Reset();
 
 	// a) find all of the relevant nodes,
 	//    inline and traverse into library nodes
@@ -3214,8 +3241,9 @@ void FRigVMParserAST::Inline(URigVMGraph* InGraph, const TArray<FRigVMASTProxy>&
 	LocalPinTraversalInfo TraversalInfo;
 	TraversalInfo.PinOverrides = &PinOverrides;
 	TraversalInfo.SourcePins = &SharedOperandPins;
-	TraversalInfo.TargetLinks = &TargetLinks;
-	TraversalInfo.SourceLinks = &SourceLinks;
+	TraversalInfo.TargetLinkIndices = &TargetLinkIndices;
+	TraversalInfo.SourceLinkIndices = &SourceLinkIndices;
+	TraversalInfo.Links = &Links;
 	TraversalInfo.Settings = &Settings;
 
 	for (const FRigVMASTProxy& NodeProxy : NodeProxies)
@@ -3224,10 +3252,10 @@ void FRigVMParserAST::Inline(URigVMGraph* InGraph, const TArray<FRigVMASTProxy>&
 	}
 }
 
-bool FRigVMParserAST::ShouldLinkBeSkipped(const FRigVMPinProxyPair& InLink) const
+bool FRigVMParserAST::ShouldLinkBeSkipped(const FRigVMASTLinkDescription& InLink) const
 {
-	URigVMPin* SourcePin = InLink.Key.GetSubjectChecked<URigVMPin>();
-	URigVMPin* TargetPin = InLink.Value.GetSubjectChecked<URigVMPin>();
+	const URigVMPin* SourcePin = InLink.SourceProxy.GetSubjectChecked<URigVMPin>();
+	const URigVMPin* TargetPin = InLink.TargetProxy.GetSubjectChecked<URigVMPin>();
 
 	for (URigVMLink* LinkToSkip : LinksToSkip)
 	{
@@ -3240,10 +3268,13 @@ bool FRigVMParserAST::ShouldLinkBeSkipped(const FRigVMPinProxyPair& InLink) cons
 	return false;
 }
 
-FString FRigVMParserAST::GetLinkAsString(const FRigVMPinProxyPair& InLink)
+FString FRigVMParserAST::GetLinkAsString(const FRigVMASTLinkDescription& InLink)
 {
-	URigVMPin* SourcePin = InLink.Key.GetSubjectChecked<URigVMPin>();
-	URigVMPin* TargetPin = InLink.Value.GetSubjectChecked<URigVMPin>();
+	const URigVMPin* SourcePin = InLink.SourceProxy.GetSubjectChecked<URigVMPin>();
+	const URigVMPin* TargetPin = InLink.TargetProxy.GetSubjectChecked<URigVMPin>();
+	static const FString EmptyString;
+	static const FString PeriodString = TEXT(".");
 
-	return FString::Printf(TEXT("%s -> %s"), *SourcePin->GetPinPath(), *TargetPin->GetPinPath());
+	return FString::Printf(TEXT("%s -> %s%s%s"), *SourcePin->GetPinPath(), *TargetPin->GetPinPath(),
+		*(InLink.SegmentPath.IsEmpty() ? EmptyString : PeriodString), *InLink.SegmentPath);
 }
\ No newline at end of file
diff --git a/Engine/Source/Developer/RigVMDeveloper/Private/RigVMCompiler/RigVMCompiler.cpp b/Engine/Source/Developer/RigVMDeveloper/Private/RigVMCompiler/RigVMCompiler.cpp
--- a/Engine/Source/Developer/RigVMDeveloper/Private/RigVMCompiler/RigVMCompiler.cpp
+++ b/Engine/Source/Developer/RigVMDeveloper/Private/RigVMCompiler/RigVMCompiler.cpp
@@ -998,21 +998,35 @@ void URigVMCompiler::TraverseAssign(const FRigVMAssignExprAST* InExpr, FRigVMCom
 		{
 			return;
 		}
-
+		
 		// if this is a copy - we should check if operands need offsets
 		if (InExpr->GetType() == FRigVMExprAST::EType::Copy)
 		{
 			struct Local
 			{
-				static void SetupRegisterOffset(URigVM* VM, URigVMPin* Pin, FRigVMOperand& Operand, const FRigVMVarExprAST* VarExpr, bool bSource, FRigVMCompilerWorkData& WorkData)
+				static void SetupRegisterOffset(URigVM* VM, const FRigVMASTLinkDescription& InLink, URigVMPin* Pin,
+					FRigVMOperand& Operand, const FRigVMVarExprAST* VarExpr, bool bSource, FRigVMCompilerWorkData& WorkData)
 				{
+					const bool bHasTargetSegmentPath = !bSource && !InLink.SegmentPath.IsEmpty();
+					
 					URigVMPin* RootPin = Pin->GetRootPin();
-					if (Pin == RootPin)
+					if (Pin == RootPin && !bHasTargetSegmentPath)
 					{
 						return;
 					}
 
 					FString SegmentPath = Pin->GetSegmentPath(false);
+					if(bHasTargetSegmentPath)
+					{
+						if(SegmentPath.IsEmpty())
+						{
+							SegmentPath = InLink.SegmentPath;
+						}
+						else
+						{
+							SegmentPath = URigVMPin::JoinPinPath(SegmentPath, InLink.SegmentPath);
+						}
+					}
 
 					// for select nodes we create a register for each case (since the cases are fixed in size)
 					// thus we do not need to setup a registeroffset for the array element.
@@ -1048,8 +1062,9 @@ void URigVMCompiler::TraverseAssign(const FRigVMAssignExprAST* InExpr, FRigVMCom
 				}
 			};
 
-			Local::SetupRegisterOffset(WorkData.VM, InExpr->GetSourcePin(), Source, SourceExpr, true, WorkData);
-			Local::SetupRegisterOffset(WorkData.VM, InExpr->GetTargetPin(), Target, TargetExpr, false, WorkData);
+			const FRigVMASTLinkDescription& Link = InExpr->GetLink();
+			Local::SetupRegisterOffset(WorkData.VM, Link, InExpr->GetSourcePin(), Source, SourceExpr, true, WorkData);
+			Local::SetupRegisterOffset(WorkData.VM, Link, InExpr->GetTargetPin(), Target, TargetExpr, false, WorkData);
 		}
 
 		FRigVMCopyOp CopyOp = WorkData.VM->GetCopyOpForOperands(Source, Target);
@@ -1722,9 +1737,9 @@ void URigVMCompiler::AddCopyOperator(const FRigVMCopyOp& InOp, const FRigVMAssig
 					const FRigVMASTProxy RootPinProxy = InTargetExpr->GetProxy().GetSibling(RootPin);
 
 					// if the root pin has only links on its subpins
-					if(WorkData.AST->GetSourceLinks(RootPinProxy, false).Num() == 0)
+					if(WorkData.AST->GetSourceLinkIndices(RootPinProxy, false).Num() == 0)
 					{
-						if(WorkData.AST->GetSourceLinks(RootPinProxy, true).Num() > 0)
+						if(WorkData.AST->GetSourceLinkIndices(RootPinProxy, true).Num() > 0)
 						{					
 							FRigVMCompilerWorkData::FCopyOpInfo DeferredCopyOp;
 							DeferredCopyOp.Op = InOp;
diff --git a/Engine/Source/Developer/RigVMDeveloper/Private/RigVMModel/RigVMPin.cpp b/Engine/Source/Developer/RigVMDeveloper/Private/RigVMModel/RigVMPin.cpp
--- a/Engine/Source/Developer/RigVMDeveloper/Private/RigVMModel/RigVMPin.cpp
+++ b/Engine/Source/Developer/RigVMDeveloper/Private/RigVMModel/RigVMPin.cpp
@@ -244,14 +244,14 @@ FString URigVMPin::GetPinPath(bool bUseNodePath) const
 	URigVMPin* ParentPin = GetParentPin();
 	if (ParentPin)
 	{
-		PinPath = FString::Printf(TEXT("%s.%s"), *ParentPin->GetPinPath(), *GetName());
+		PinPath = JoinPinPath(ParentPin->GetPinPath(), GetName());
 	}
 	else
 	{
 		URigVMNode* Node = GetNode();
 		if (Node != nullptr)
 		{
-			PinPath = FString::Printf(TEXT("%s.%s"), *Node->GetNodePath(bUseNodePath), *GetName());
+			PinPath = JoinPinPath(Node->GetNodePath(bUseNodePath), GetName());
 		}
 	}
 #if UE_BUILD_DEBUG
@@ -263,6 +263,25 @@ FString URigVMPin::GetPinPath(bool bUseNodePath) const
 	return PinPath;
 }
 
+FString URigVMPin::GetSubPinPath(const URigVMPin* InParentPin, bool bIncludeParentPinName) const
+{
+	if (const URigVMPin* ParentPin = GetParentPin())
+	{
+		if(ParentPin == InParentPin)
+		{
+			if(bIncludeParentPinName)
+			{
+				return JoinPinPath(ParentPin->GetName(),GetName());
+			}
+		}
+		else
+		{
+			return JoinPinPath(ParentPin->GetSubPinPath(InParentPin, bIncludeParentPinName), GetName());
+		}
+	}
+	return GetName();
+}
+
 FString URigVMPin::GetSegmentPath(bool bIncludeRootPin) const
 {
 	URigVMPin* ParentPin = GetParentPin();
@@ -273,7 +292,7 @@ FString URigVMPin::GetSegmentPath(bool bIncludeRootPin) const
 		{
 			return GetName();
 		}
-		return FString::Printf(TEXT("%s.%s"), *ParentSegmentPath, *GetName());
+		return JoinPinPath(ParentSegmentPath, GetName());
 	}
 
 	if(bIncludeRootPin)
diff --git a/Engine/Source/Developer/RigVMDeveloper/Public/RigVMCompiler/RigVMAST.h b/Engine/Source/Developer/RigVMDeveloper/Public/RigVMCompiler/RigVMAST.h
--- a/Engine/Source/Developer/RigVMDeveloper/Public/RigVMCompiler/RigVMAST.h
+++ b/Engine/Source/Developer/RigVMDeveloper/Public/RigVMCompiler/RigVMAST.h
@@ -36,6 +36,33 @@ class URigVMNode;
 class URigVMGraph;
 class URigVMController;
 
+/*
+ * A structure to describe a link between two proxies
+ */
+struct RIGVMDEVELOPER_API FRigVMASTLinkDescription
+{
+	FRigVMASTLinkDescription() {}
+	
+	FRigVMASTLinkDescription(const FRigVMASTProxy& InSourceProxy, const FRigVMASTProxy& InTargetProxy, const FString& InSegmentPath)
+	: SourceProxy(InSourceProxy)
+	, TargetProxy(InTargetProxy)
+	, SegmentPath(InSegmentPath)
+	, LinkIndex(INDEX_NONE)
+	{}
+
+	FRigVMASTLinkDescription(const FRigVMASTLinkDescription& InOther)
+	: SourceProxy(InOther.SourceProxy)
+	, TargetProxy(InOther.TargetProxy)
+	, SegmentPath(InOther.SegmentPath)
+	, LinkIndex(INDEX_NONE)
+	{}
+
+	FRigVMASTProxy SourceProxy;
+	FRigVMASTProxy TargetProxy;
+	FString SegmentPath;
+	int32 LinkIndex;
+};
+
 /*
  * Base class for an expression within an abstract syntax tree.
  * The base implements parent / child relationships as well as a
@@ -811,18 +838,21 @@ public:
 	FRigVMAssignExprAST(const FRigVMAssignExprAST&) = delete;
 
 	// returns the source proxy this expression is using
-	const FRigVMASTProxy& GetSourceProxy() const { return SourceProxy; }
+	const FRigVMASTLinkDescription& GetLink() const { return Link; }
+
+	// returns the source proxy this expression is using
+	const FRigVMASTProxy& GetSourceProxy() const { return Link.SourceProxy; }
 
 	// returns the source pin for this assignment
 	// @return the source pin for this assignment
-	URigVMPin* GetSourcePin() const { return SourceProxy.GetSubjectChecked<URigVMPin>(); }
+	URigVMPin* GetSourcePin() const { return Link.SourceProxy.GetSubjectChecked<URigVMPin>(); }
 
 	// returns the target proxy this expression is using
-	const FRigVMASTProxy& GetTargetProxy() const { return TargetProxy; }
+	const FRigVMASTProxy& GetTargetProxy() const { return Link.TargetProxy; }
 
 	// returns the target pin for this assignment
 	// @return the target pin for this assignment
-	URigVMPin* GetTargetPin() const { return TargetProxy.GetSubjectChecked<URigVMPin>(); }
+	URigVMPin* GetTargetPin() const { return Link.TargetProxy.GetSubjectChecked<URigVMPin>(); }
 
 	// overload of the type checking mechanism
 	virtual bool IsA(EType InType) const override
@@ -833,17 +863,14 @@ public:
 protected:
 
 	// default constructor (protected so that only parser can access it)
-	FRigVMAssignExprAST(EType InType, const FRigVMASTProxy& InSourceProxy, const FRigVMASTProxy& InTargetProxy)
+	FRigVMAssignExprAST(EType InType, const FRigVMASTLinkDescription& InLink)
 		: FRigVMExprAST(InType)
-		, SourceProxy(InSourceProxy)
-		, TargetProxy(InTargetProxy)
+		, Link(InLink)
 	{
 	}
 
 private:
-
-	FRigVMASTProxy SourceProxy;
-	FRigVMASTProxy TargetProxy;
+	FRigVMASTLinkDescription Link;
 	friend class FRigVMParserAST;
 };
 
@@ -878,8 +905,8 @@ public:
 protected:
 
 	// default constructor (protected so that only parser can access it)
-	FRigVMCopyExprAST(const FRigVMASTProxy& InSourceProxy, const FRigVMASTProxy& InTargetProxy)
-		: FRigVMAssignExprAST(EType::Copy, InSourceProxy, InTargetProxy)
+	FRigVMCopyExprAST(const FRigVMASTLinkDescription& InLink)
+		: FRigVMAssignExprAST(EType::Copy, InLink)
 	{}
 
 
@@ -1283,6 +1310,16 @@ private:
 		return Expr;
 	}
 
+	// make function to create an expression
+	template<class ExprType>
+	ExprType* MakeExpr(FRigVMExprAST::EType InType, const FRigVMASTLinkDescription& InLink)
+	{
+		ExprType* Expr = new ExprType(InType, InLink);
+		Expr->ParserPtr = this;
+		Expr->Index = Expressions.Add(Expr);
+		return Expr;
+	}
+
 	// make function to create an expression
 	template<class ExprType>
 	ExprType* MakeExpr(FRigVMExprAST::EType InType, const FRigVMASTProxy& InProxyA, const FRigVMASTProxy& InProxyB)
@@ -1313,6 +1350,16 @@ private:
 		return Expr;
 	}
 
+	// make function to create an expression
+	template<class ExprType>
+	ExprType* MakeExpr(const FRigVMASTLinkDescription& InLink)
+	{
+		ExprType* Expr = new ExprType(InLink);
+		Expr->ParserPtr = this;
+		Expr->Index = Expressions.Add(Expr);
+		return Expr;
+	}
+
 	// make function to create an expression
 	template<class ExprType>
 	ExprType* MakeExpr()
@@ -1373,10 +1420,12 @@ private:
 	void Inline(URigVMGraph* InGraph, const TArray<FRigVMASTProxy>& InNodeProxies);
 
 	// helper functions to retrieve links for a given pin
-	const TArray<FRigVMASTProxy>& GetSourcePins(const FRigVMASTProxy& InPinProxy) const;
-	const TArray<FRigVMASTProxy>& GetTargetPins(const FRigVMASTProxy& InPinProxy) const;
-	TArray<FRigVMPinProxyPair> GetSourceLinks(const FRigVMASTProxy& InPinProxy, bool bRecursive = false) const;
-	TArray<FRigVMPinProxyPair> GetTargetLinks(const FRigVMASTProxy& InPinProxy, bool bRecursive = false) const;
+	TArray<int32> GetSourceLinkIndices(const FRigVMASTProxy& InPinProxy, bool bRecursive = false) const;
+	TArray<int32> GetTargetLinkIndices(const FRigVMASTProxy& InPinProxy, bool bRecursive = false) const;
+	TArray<int32> GetLinkIndices(const FRigVMASTProxy& InPinProxy, bool bGetSource, bool bRecursive) const;
+
+	// function to retrieve a link given its index
+	const FRigVMASTLinkDescription& GetLink(int32 InLinkIndex) const;
 
 	// traverse a single mutable node (constructs entry, call extern and other expressions)
 	FRigVMExprAST* TraverseMutableNode(const FRigVMASTProxy& InNodeProxy, FRigVMExprAST* InParentExpr);
@@ -1394,11 +1443,11 @@ private:
 	FRigVMExprAST* TraversePin(const FRigVMASTProxy& InPinProxy, FRigVMExprAST* InParentExpr);
 
 	// traverse a single link (constructs assign + copy expressions)
-	FRigVMExprAST* TraverseLink(const FRigVMPinProxyPair& InLink, FRigVMExprAST* InParentExpr);
+	FRigVMExprAST* TraverseLink(int32 InLinkIndex, FRigVMExprAST* InParentExpr);
 
-	bool ShouldLinkBeSkipped(const FRigVMPinProxyPair& InLink) const;
+	bool ShouldLinkBeSkipped(const FRigVMASTLinkDescription& InLink) const;
 
-	static FString GetLinkAsString(const FRigVMPinProxyPair& InLink);
+	static FString GetLinkAsString(const FRigVMASTLinkDescription& InLink);
 
 	TMap<FRigVMASTProxy, FRigVMExprAST*> SubjectToExpression;
 	TMap<FRigVMASTProxy, int32> NodeExpressionIndex;
@@ -1409,8 +1458,9 @@ private:
 
 	TArray<FRigVMASTProxy> NodeProxies;
 	TMap<FRigVMASTProxy, FRigVMASTProxy> SharedOperandPins;
-	TMap<FRigVMASTProxy, TArray<FRigVMASTProxy>> TargetLinks;
-	TMap<FRigVMASTProxy, TArray<FRigVMASTProxy>> SourceLinks;
+	TMap<FRigVMASTProxy, TArray<int32>> TargetLinkIndices;
+	TMap<FRigVMASTProxy, TArray<int32>> SourceLinkIndices;
+	TArray<FRigVMASTLinkDescription> Links;
 	static const TArray<FRigVMASTProxy> EmptyProxyArray;
 	URigVMPin::FPinOverrideMap PinOverrides;
 
@@ -1430,6 +1480,7 @@ private:
 	mutable TMap<TTuple<const FRigVMExprAST*,const FRigVMExprAST*>, int32> MinIndexOfChildWithinParent;
 
 	friend class FRigVMExprAST;
+	friend class FRigVMAssignExprAST;
 	friend class URigVMCompiler;
 };
 
diff --git a/Engine/Source/Developer/RigVMDeveloper/Public/RigVMModel/RigVMPin.h b/Engine/Source/Developer/RigVMDeveloper/Public/RigVMModel/RigVMPin.h
--- a/Engine/Source/Developer/RigVMDeveloper/Public/RigVMModel/RigVMPin.h
+++ b/Engine/Source/Developer/RigVMDeveloper/Public/RigVMModel/RigVMPin.h
@@ -132,6 +132,11 @@ public:
 	UFUNCTION(BlueprintCallable, Category = RigVMPin)
 	FString GetPinPath(bool bUseNodePath = false) const;
 
+	// Returns a . separated path containing all names of the pin and its owners
+	// until we hit the provided parent pin.
+	UFUNCTION(BlueprintCallable, Category = RigVMPin)
+	FString GetSubPinPath(const URigVMPin* InParentPin, bool bIncludeParentPinName = false) const;
+
 	// Returns a . separated path containing all names of the pin within its main
 	// memory owner / storage. This is typically used to create an offset pointer
 	// within memory (FRigVMRegisterOffset).
