RigVM: Ignore pin default overrides for parent pins if sub pins also are provided with default overrides

diff --git a/Engine/Source/Developer/RigVMDeveloper/Private/RigVMCompiler/RigVMAST.cpp b/Engine/Source/Developer/RigVMDeveloper/Private/RigVMCompiler/RigVMAST.cpp
--- a/Engine/Source/Developer/RigVMDeveloper/Private/RigVMCompiler/RigVMAST.cpp
+++ b/Engine/Source/Developer/RigVMDeveloper/Private/RigVMCompiler/RigVMAST.cpp
@@ -3250,6 +3250,28 @@ void FRigVMParserAST::Inline(URigVMGraph* InGraph, const TArray<FRigVMASTProxy>&
 	{
 		LocalPinTraversalInfo::VisitNode(NodeProxy, TraversalInfo);
 	}
+
+	// once we are done with the inlining we may need to clean up pin value overrides for pins
+	// that also have overrides on sub pins
+	TArray<FRigVMASTProxy> PinOverridesToRemove;
+	for(const TPair<FRigVMASTProxy, URigVMPin::FPinOverrideValue>& Override : PinOverrides)
+	{
+		if(URigVMPin* Pin = Override.Key.GetSubject<URigVMPin>())
+		{
+			for(URigVMPin* SubPin : Pin->GetSubPins())
+			{
+				const FRigVMASTProxy SubPinProxy = Override.Key.GetSibling(SubPin);
+				if(PinOverrides.Contains(SubPinProxy))
+				{
+					PinOverridesToRemove.Add(Override.Key);
+				}
+			}
+		}
+	}
+	for(const FRigVMASTProxy& ProxyToRemove : PinOverridesToRemove)
+	{
+		PinOverrides.Remove(ProxyToRemove);
+	}
 }
 
 bool FRigVMParserAST::ShouldLinkBeSkipped(const FRigVMASTLinkDescription& InLink) const
