RigVM: Static Analysis - Fix nullptr access

diff --git a/Engine/Source/Developer/RigVMDeveloper/Private/RigVMCompiler/RigVMAST.cpp b/Engine/Source/Developer/RigVMDeveloper/Private/RigVMCompiler/RigVMAST.cpp
--- a/Engine/Source/Developer/RigVMDeveloper/Private/RigVMCompiler/RigVMAST.cpp
+++ b/Engine/Source/Developer/RigVMDeveloper/Private/RigVMCompiler/RigVMAST.cpp
@@ -230,7 +230,10 @@ int32 FRigVMExprAST::GetMinChildIndexWithinParent(const FRigVMExprAST* InParentE
 		}
 	}
 
-	Parser->MinIndexOfChildWithinParent.Add(MapKey, MinIndex);
+	if(Parser)
+	{
+		Parser->MinIndexOfChildWithinParent.Add(MapKey, MinIndex);
+	}
 	return MinIndex;
 }
 
