Control Rig: Fix connecting to Mirror template node hangs editor

diff --git a/Engine/Source/Developer/RigVMDeveloper/Private/RigVMModel/RigVMController.cpp b/Engine/Source/Developer/RigVMDeveloper/Private/RigVMModel/RigVMController.cpp
--- a/Engine/Source/Developer/RigVMDeveloper/Private/RigVMModel/RigVMController.cpp
+++ b/Engine/Source/Developer/RigVMDeveloper/Private/RigVMModel/RigVMController.cpp
@@ -897,7 +897,7 @@ URigVMVariableNode* URigVMController::AddVariableNode(const FName& InVariableNam
 		InCPPTypeObject = URigVMPin::FindObjectFromCPPTypeObjectPath<UObject>(InCPPType);
 	}
 
-	FString CPPType = PostProcessCPPType(InCPPType, InCPPTypeObject);
+	FString CPPType = RigVMTypeUtils::PostProcessCPPType(InCPPType, InCPPTypeObject);
 	
 	FString Name = GetValidNodeName(InNodeName.IsEmpty() ? FString(TEXT("VariableNode")) : InNodeName);
 	URigVMVariableNode* Node = NewObject<URigVMVariableNode>(Graph, *Name);
@@ -8260,7 +8260,7 @@ FName URigVMController::AddExposedPin(const FName& InPinName, ERigVMPinDirection
 	}, false, true);
 
 	URigVMPin* Pin = NewObject<URigVMPin>(LibraryNode, PinName);
-	Pin->CPPType = PostProcessCPPType(InCPPType, CPPTypeObject);
+	Pin->CPPType = RigVMTypeUtils::PostProcessCPPType(InCPPType, CPPTypeObject);
 	Pin->CPPTypeObjectPath = InCPPTypeObjectPath;
 	Pin->bIsConstant = false;
 	Pin->Direction = InDirection;
@@ -10098,7 +10098,7 @@ URigVMIfNode* URigVMController::AddIfNode(const FString& InCPPType, const FName&
 		}
 	}
 
-	FString CPPType = PostProcessCPPType(InCPPType, CPPTypeObject);
+	FString CPPType = RigVMTypeUtils::PostProcessCPPType(InCPPType, CPPTypeObject);
 
 	FString DefaultValue;
 	if(UScriptStruct* ScriptStruct = Cast<UScriptStruct>(CPPTypeObject))
@@ -10214,7 +10214,7 @@ URigVMSelectNode* URigVMController::AddSelectNode(const FString& InCPPType, cons
 		}
 	}
 
-	FString CPPType = PostProcessCPPType(InCPPType, CPPTypeObject);
+	FString CPPType = RigVMTypeUtils::PostProcessCPPType(InCPPType, CPPTypeObject);
 
 	FString DefaultValue;
 	if (UScriptStruct* ScriptStruct = Cast<UScriptStruct>(CPPTypeObject))
@@ -10575,7 +10575,7 @@ URigVMArrayNode* URigVMController::AddArrayNode(ERigVMOpCode InOpCode, const FSt
 		InCPPTypeObject = URigVMPin::FindObjectFromCPPTypeObjectPath<UObject>(CPPType);
 	}
 
-	CPPType = PostProcessCPPType(CPPType, InCPPTypeObject);
+	CPPType = RigVMTypeUtils::PostProcessCPPType(CPPType, InCPPTypeObject);
 	
 	const FString Name = GetValidNodeName(InNodeName.IsEmpty() ? FString(TEXT("ArrayNode")) : InNodeName);
 	URigVMArrayNode* Node = NewObject<URigVMArrayNode>(Graph, *Name);
@@ -11473,6 +11473,8 @@ void URigVMController::ConfigurePinFromProperty(FProperty* InProperty, URigVMPin
 	{
 		InOutPin->CPPTypeObjectPath = *InOutPin->CPPTypeObject->GetPathName();
 	}
+
+	InOutPin->CPPType = RigVMTypeUtils::PostProcessCPPType(InOutPin->CPPType, InOutPin->GetCPPTypeObject());
 }
 
 void URigVMController::ConfigurePinFromPin(URigVMPin* InOutPin, URigVMPin* InPin)
@@ -12930,36 +12932,6 @@ void URigVMController::PostProcessDefaultValue(URigVMPin* Pin, FString& OutDefau
 	}
 }
 
-FString URigVMController::PostProcessCPPType(const FString& InCPPType, UObject* InCPPTypeObject)
-{
-	FString CPPType = InCPPType;
-	
-	if (const UClass* Class = Cast<UClass>(InCPPTypeObject))
-	{
-		CPPType = FString::Printf(RigVMTypeUtils::TObjectPtrTemplate, Class->GetPrefixCPP(), *Class->GetName());
-	}
-	else if (const UScriptStruct* ScriptStruct = Cast<UScriptStruct>(InCPPTypeObject))
-	{
-		CPPType = ScriptStruct->GetStructCPPName();
-	}
-	else if (UEnum* Enum = Cast<UEnum>(InCPPTypeObject))
-	{
-		CPPType = RigVMTypeUtils::CPPTypeFromEnum(Enum);
-	}
-
-	if(CPPType != InCPPType)
-	{
-		FString TemplateType = InCPPType;
-		while (RigVMTypeUtils::IsArrayType(TemplateType))
-		{
-			CPPType = RigVMTypeUtils::ArrayTypeFromBaseType(CPPType);
-			TemplateType = RigVMTypeUtils::BaseTypeFromArrayType(TemplateType);
-		}		
-	}
-	
-	return CPPType;
-}
-
 void URigVMController::ResolveTemplateNodePins(URigVMTemplateNode* InNode, bool bSetupUndoRedo)
 {
 	if (InNode == nullptr)
@@ -13276,7 +13248,7 @@ bool URigVMController::ResolveWildCardPin(const FString& InPinPath, const FStrin
 		}
 	}
 
-	const FString CPPType = PostProcessCPPType(InCPPType, CPPTypeObject);
+	const FString CPPType = RigVMTypeUtils::PostProcessCPPType(InCPPType, CPPTypeObject);
 
 	if (URigVMPin* Pin = Graph->FindPin(InPinPath))
 	{
@@ -14022,7 +13994,7 @@ bool URigVMController::EnsurePinValidity(URigVMPin* InPin, bool bRecursive)
 		}
 	}
 
-	InPin->CPPType = PostProcessCPPType(InPin->CPPType, InPin->GetCPPTypeObject());
+	InPin->CPPType = RigVMTypeUtils::PostProcessCPPType(InPin->CPPType, InPin->GetCPPTypeObject());
 
 	if(bRecursive)
 	{
@@ -14056,7 +14028,7 @@ void URigVMController::EnsureLocalVariableValidity()
 		{
 			// CPPType can become invalid when the type object is defined by
 			// an asset that have changed name or asset path, user defined struct is one possibility
-			Variable.CPPType = PostProcessCPPType(Variable.CPPType, Variable.CPPTypeObject);
+			Variable.CPPType = RigVMTypeUtils::PostProcessCPPType(Variable.CPPType, Variable.CPPTypeObject);
 		}
 	}
 }
diff --git a/Engine/Source/Developer/RigVMDeveloper/Private/RigVMModel/RigVMPin.cpp b/Engine/Source/Developer/RigVMDeveloper/Private/RigVMModel/RigVMPin.cpp
--- a/Engine/Source/Developer/RigVMDeveloper/Private/RigVMModel/RigVMPin.cpp
+++ b/Engine/Source/Developer/RigVMDeveloper/Private/RigVMModel/RigVMPin.cpp
@@ -563,7 +563,7 @@ int32 URigVMPin::GetArraySize() const
 
 FString URigVMPin::GetCPPType() const
 {
-	return URigVMController::PostProcessCPPType(CPPType, GetCPPTypeObject());
+	return RigVMTypeUtils::PostProcessCPPType(CPPType, GetCPPTypeObject());
 }
 
 FString URigVMPin::GetArrayElementCppType() const
@@ -1241,7 +1241,7 @@ void URigVMPin::UpdateTypeInformationIfRequired() const
 	{
 		// refresh the type string 
 		URigVMPin* MutableThis = (URigVMPin*)this;
-		MutableThis->CPPType = URigVMController::PostProcessCPPType(CPPType, CPPTypeObject);
+		MutableThis->CPPType = RigVMTypeUtils::PostProcessCPPType(CPPType, CPPTypeObject);
 	}
 }
 
diff --git a/Engine/Source/Developer/RigVMDeveloper/Public/RigVMModel/RigVMController.h b/Engine/Source/Developer/RigVMDeveloper/Public/RigVMModel/RigVMController.h
--- a/Engine/Source/Developer/RigVMDeveloper/Public/RigVMModel/RigVMController.h
+++ b/Engine/Source/Developer/RigVMDeveloper/Public/RigVMModel/RigVMController.h
@@ -941,7 +941,6 @@ private:
 	static FLinearColor GetColorFromMetadata(const FString& InMetadata);
 	static void CreateDefaultValueForStructIfRequired(UScriptStruct* InStruct, FString& InOutDefaultValue);
 	static void PostProcessDefaultValue(URigVMPin* Pin, FString& OutDefaultValue);
-	static FString PostProcessCPPType(const FString& InCPPType, UObject* InCPPTypeObject);
 
 	bool ResolveWildCardPin(URigVMPin* InPinToResolve, URigVMPin* InTemplatePin, bool bSetupUndoRedo, bool bTraverseNode = true, bool bTraverseParentPins = true, bool bTraverseLinks = true, bool bAllowFloatingPointCasts = false);
 	bool ResolveWildCardPinImpl(URigVMPin* InPinToResolve, const FString& InCPPType, UObject* InCPPTypeObject, bool bSetupUndoRedo, bool bTraverseNode = true, bool bTraverseParentPins = true, bool bTraverseLinks = true, bool bAllowFloatingPointCasts = false);
diff --git a/Engine/Source/Runtime/RigVM/Private/RigVMCore/RigVMTemplate.cpp b/Engine/Source/Runtime/RigVM/Private/RigVMCore/RigVMTemplate.cpp
--- a/Engine/Source/Runtime/RigVM/Private/RigVMCore/RigVMTemplate.cpp
+++ b/Engine/Source/Runtime/RigVM/Private/RigVMCore/RigVMTemplate.cpp
@@ -47,6 +47,7 @@ FRigVMTemplateArgument::FRigVMTemplateArgument(FProperty* InProperty)
 	{
 		Type.CPPTypeObject = ByteProperty->Enum;
 	}
+	Type.CPPType = RigVMTypeUtils::PostProcessCPPType(Type.CPPType, Type.CPPTypeObject);
 
 	Types.Add(Type);
 }
diff --git a/Engine/Source/Runtime/RigVM/Public/RigVMTypeUtils.h b/Engine/Source/Runtime/RigVM/Public/RigVMTypeUtils.h
--- a/Engine/Source/Runtime/RigVM/Public/RigVMTypeUtils.h
+++ b/Engine/Source/Runtime/RigVM/Public/RigVMTypeUtils.h
@@ -90,4 +90,34 @@ namespace RigVMTypeUtils
 		static const FString WildCardArrayCPPType = ArrayTypeFromBaseType(GetWildCardCPPType()); 
 		return WildCardArrayCPPType;
 	}
+
+	FORCEINLINE FString PostProcessCPPType(const FString& InCPPType, UObject* InCPPTypeObject)
+	{
+		FString CPPType = InCPPType;
+	
+		if (const UClass* Class = Cast<UClass>(InCPPTypeObject))
+		{
+			CPPType = FString::Printf(RigVMTypeUtils::TObjectPtrTemplate, Class->GetPrefixCPP(), *Class->GetName());
+		}
+		else if (const UScriptStruct* ScriptStruct = Cast<UScriptStruct>(InCPPTypeObject))
+		{
+			CPPType = ScriptStruct->GetStructCPPName();
+		}
+		else if (UEnum* Enum = Cast<UEnum>(InCPPTypeObject))
+		{
+			CPPType = RigVMTypeUtils::CPPTypeFromEnum(Enum);
+		}
+
+		if(CPPType != InCPPType)
+		{
+			FString TemplateType = InCPPType;
+			while (RigVMTypeUtils::IsArrayType(TemplateType))
+			{
+				CPPType = RigVMTypeUtils::ArrayTypeFromBaseType(CPPType);
+				TemplateType = RigVMTypeUtils::BaseTypeFromArrayType(TemplateType);
+			}		
+		}
+	
+		return CPPType;
+	}
 }
\ No newline at end of file
