Control Rig: Add aggregate node

diff --git a/Engine/Source/Developer/RigVMDeveloper/Private/RigVMModel/Nodes/RigVMAggregateNode.cpp b/Engine/Source/Developer/RigVMDeveloper/Private/RigVMModel/Nodes/RigVMAggregateNode.cpp
new file mode 100644
--- /dev/null
+++ b/Engine/Source/Developer/RigVMDeveloper/Private/RigVMModel/Nodes/RigVMAggregateNode.cpp
@@ -0,0 +1,141 @@
+﻿// Copyright Epic Games, Inc. All Rights Reserved.
+
+#include "RigVMModel/Nodes/RigVMAggregateNode.h"
+
+#include "RigVMModel/RigVMController.h"
+
+URigVMAggregateNode::URigVMAggregateNode()
+	: Super()
+	, FirstInnerNodeCache(nullptr)
+{
+}
+
+bool URigVMAggregateNode::IsInputAggregate() const
+{
+	return GetFirstAggregatePin()->GetDirection() == ERigVMPinDirection::Input;
+}
+
+URigVMNode* URigVMAggregateNode::GetFirstInnerNode() const
+{
+	if (FirstInnerNodeCache == nullptr)
+	{
+		// Find any inner node (not entry or return)
+		URigVMNode* InnerNode = nullptr;
+		{
+			for (URigVMNode* Node : GetContainedGraph()->GetNodes())
+			{
+				if (!Node->IsA<URigVMFunctionEntryNode>() && !Node->IsA<URigVMFunctionReturnNode>())
+				{
+					InnerNode = Node;
+					break;
+				}
+			}
+
+			if (!InnerNode || !InnerNode->IsAggregate())
+			{
+				return nullptr;
+			}
+		}
+
+		// Find the aggregate pins
+		URigVMGraph* Graph = GetContainedGraph();
+		if (URigVMFunctionEntryNode* EntryNode = Graph->GetEntryNode())
+		{
+			URigVMPin* Arg1Pin = nullptr;
+			const URigVMPin* Argument1 = InnerNode->GetFirstAggregatePin();
+			if (Argument1->GetDirection() == ERigVMPinDirection::Input)
+			{
+				Arg1Pin = EntryNode->FindPin(Argument1->GetName());
+			}
+			else
+			{
+				const URigVMPin* OppositeArgument = InnerNode->GetOppositeAggregatePin();
+				Arg1Pin = EntryNode->FindPin(OppositeArgument->GetName());				
+			}
+			if (Arg1Pin)
+			{
+				TArray<URigVMPin*> TargetPins = Arg1Pin->GetLinkedTargetPins();
+				if (TargetPins.Num() > 0)
+				{
+					FirstInnerNodeCache = TargetPins[0]->GetNode(); 
+				}
+			}
+		}
+	}
+	return FirstInnerNodeCache;
+}
+
+URigVMPin* URigVMAggregateNode::GetFirstAggregatePin() const
+{
+	const URigVMNode* FirstNode = GetFirstInnerNode();
+	return FirstNode->GetFirstAggregatePin();
+}
+
+URigVMPin* URigVMAggregateNode::GetSecondAggregatePin() const
+{
+	const URigVMNode* FirstNode = GetFirstInnerNode();
+	return FirstNode->GetSecondAggregatePin();
+}
+
+URigVMPin* URigVMAggregateNode::GetOppositeAggregatePin() const
+{
+	const URigVMNode* FirstNode = GetFirstInnerNode();
+	return FirstNode->GetOppositeAggregatePin();
+}
+
+void URigVMAggregateNode::InvalidateCache()
+{
+	FirstInnerNodeCache = nullptr;
+}
+
+FString URigVMAggregateNode::GetNodeTitle() const
+{
+	if (URigVMNode* InnerNode = GetFirstInnerNode())
+	{
+		return InnerNode->GetNodeTitle();
+	}
+	
+	return Super::GetNodeTitle();
+}
+
+FName URigVMAggregateNode::GetMethodName() const
+{
+	if (URigVMUnitNode* InnerNode = Cast<URigVMUnitNode>(GetFirstInnerNode()))
+	{
+		return InnerNode->GetMethodName();
+	}
+	
+	return NAME_None;
+}
+
+FText URigVMAggregateNode::GetToolTipText() const
+{
+	if (URigVMNode* InnerNode = GetFirstInnerNode())
+	{
+		return InnerNode->GetToolTipText();
+	}
+	
+	return Super::GetToolTipText();
+}
+
+FText URigVMAggregateNode::GetToolTipTextForPin(const URigVMPin* InPin) const
+{
+	if (URigVMPin* EntryPin = GetContainedGraph()->GetEntryNode()->FindPin(InPin->GetName()))
+	{
+		const TArray<URigVMPin*> Targets = EntryPin->GetLinkedTargetPins();
+		if (Targets.Num() > 0)
+		{
+			return Targets[0]->GetToolTipText();
+		}
+	}
+	else if(URigVMPin* ReturnPin = GetContainedGraph()->GetReturnNode()->FindPin(InPin->GetName()))
+	{
+		const TArray<URigVMPin*> Sources = ReturnPin->GetLinkedSourcePins();
+		if (Sources.Num() > 0)
+		{
+			return Sources[0]->GetToolTipText();
+		}
+	}
+	
+	return Super::GetToolTipTextForPin(InPin);
+}
diff --git a/Engine/Source/Developer/RigVMDeveloper/Private/RigVMModel/Nodes/RigVMUnitNode.cpp b/Engine/Source/Developer/RigVMDeveloper/Private/RigVMModel/Nodes/RigVMUnitNode.cpp
--- a/Engine/Source/Developer/RigVMDeveloper/Private/RigVMModel/Nodes/RigVMUnitNode.cpp
+++ b/Engine/Source/Developer/RigVMDeveloper/Private/RigVMModel/Nodes/RigVMUnitNode.cpp
@@ -1,6 +1,8 @@
 // Copyright Epic Games, Inc. All Rights Reserved.
 
 #include "RigVMModel/Nodes/RigVMUnitNode.h"
+
+#include "Animation/Rig.h"
 #include "RigVMCore/RigVMStruct.h"
 
 FString URigVMUnitNode::GetNodeTitle() const
@@ -107,6 +109,125 @@ FString URigVMUnitNode::GetDeprecatedMetadata() const
 	return FString();
 }
 
+bool URigVMUnitNode::IsAggregate() const
+{
+	TArray<URigVMPin*> AggregateInputs = GetAggregateInputs();
+	TArray<URigVMPin*> AggregateOutputs = GetAggregateOutputs();
+
+	if ((AggregateInputs.Num() == 2 && AggregateOutputs.Num() == 1) ||
+		(AggregateInputs.Num() == 1 && AggregateOutputs.Num() == 2))
+	{
+		TArray<URigVMPin*> AggregateAll = AggregateInputs;
+		AggregateAll.Append(AggregateOutputs);
+		for (int32 i = 1; i < 3; ++i)
+		{
+			if (AggregateAll[0]->GetCPPType() != AggregateAll[i]->GetCPPType() ||
+				AggregateAll[0]->GetCPPTypeObject() != AggregateAll[i]->GetCPPTypeObject())
+			{
+				return false;
+			}
+		}
+		
+		return true;
+	}
+	
+	return false;
+}
+
+URigVMPin* URigVMUnitNode::GetFirstAggregatePin() const
+{
+	TArray<URigVMPin*> Inputs = GetAggregateInputs();
+	TArray<URigVMPin*> Outputs = GetAggregateOutputs();
+	if (Inputs.Num() == 2 && Outputs.Num() == 1)
+	{
+		return Inputs[0];
+	}
+	if (Inputs.Num() == 1 && Outputs.Num() == 2)
+	{
+		return Outputs[0];
+	}
+	return nullptr;
+}
+
+URigVMPin* URigVMUnitNode::GetSecondAggregatePin() const
+{
+	TArray<URigVMPin*> Inputs = GetAggregateInputs();
+	TArray<URigVMPin*> Outputs = GetAggregateOutputs();
+	if (Inputs.Num() == 2 && Outputs.Num() == 1)
+	{
+		return Inputs[1];
+	}
+	if (Inputs.Num() == 1 && Outputs.Num() == 2)
+	{
+		return Outputs[1];
+	}
+	return nullptr;
+}
+
+URigVMPin* URigVMUnitNode::GetOppositeAggregatePin() const
+{
+	TArray<URigVMPin*> Inputs = GetAggregateInputs();
+	TArray<URigVMPin*> Outputs = GetAggregateOutputs();
+	if (Inputs.Num() == 2 && Outputs.Num() == 1)
+	{
+		return Outputs[0];
+	}
+	if (Inputs.Num() == 1 && Outputs.Num() == 2)
+	{
+		return Inputs[0];
+	}
+	return nullptr;
+}
+
+bool URigVMUnitNode::IsInputAggregate() const
+{
+	return GetAggregateInputs().Num() == 2;
+}
+
+TArray<URigVMPin*> URigVMUnitNode::GetAggregateInputs() const
+{
+	TArray<URigVMPin*> AggregateInputs;
+	if (UScriptStruct* Struct = GetScriptStruct())
+	{
+		for (URigVMPin* Pin : GetPins())
+		{
+			if (Pin->GetDirection() == ERigVMPinDirection::Input)
+			{
+				if (FProperty* Property = Struct->FindPropertyByName(Pin->GetFName()))
+				{
+					if (Property->HasMetaData(FRigVMStruct::AggregateMetaName))
+					{
+						AggregateInputs.Add(Pin);
+					}
+				}			
+			}
+		}
+	}
+	return AggregateInputs;
+}
+
+TArray<URigVMPin*> URigVMUnitNode::GetAggregateOutputs() const
+{
+	TArray<URigVMPin*> AggregateOutputs;
+	if (UScriptStruct* Struct = GetScriptStruct())
+	{
+		for (URigVMPin* Pin : GetPins())
+		{
+			if (Pin->GetDirection() == ERigVMPinDirection::Output)
+			{
+				if (FProperty* Property = Struct->FindPropertyByName(Pin->GetFName()))
+				{
+					if (Property->HasMetaData(FRigVMStruct::AggregateMetaName))
+					{
+						AggregateOutputs.Add(Pin);
+					}
+				}			
+			}
+		}
+	}
+	return AggregateOutputs;
+}
+
 UScriptStruct* URigVMUnitNode::GetScriptStruct() const
 {
 	return ScriptStruct;
diff --git a/Engine/Source/Developer/RigVMDeveloper/Private/RigVMModel/RigVMController.cpp b/Engine/Source/Developer/RigVMDeveloper/Private/RigVMModel/RigVMController.cpp
--- a/Engine/Source/Developer/RigVMDeveloper/Private/RigVMModel/RigVMController.cpp
+++ b/Engine/Source/Developer/RigVMDeveloper/Private/RigVMModel/RigVMController.cpp
@@ -5,6 +5,7 @@
 #include "RigVMModel/Nodes/RigVMFunctionEntryNode.h"
 #include "RigVMModel/Nodes/RigVMFunctionReturnNode.h"
 #include "RigVMModel/Nodes/RigVMFunctionReferenceNode.h"
+#include "RigVMModel/Nodes/RigVMAggregateNode.h"
 #include "RigVMCore/RigVMRegistry.h"
 #include "RigVMCore/RigVMExecuteContext.h"
 #include "RigVMCore/RigVMUnknownType.h"
@@ -3487,7 +3488,7 @@ FName URigVMController::GetUniqueName(const FName& InName, TFunction<bool(const
 	return *Name;
 }
 
-URigVMCollapseNode* URigVMController::CollapseNodes(const TArray<FName>& InNodeNames, const FString& InCollapseNodeName, bool bSetupUndoRedo, bool bPrintPythonCommand)
+URigVMCollapseNode* URigVMController::CollapseNodes(const TArray<FName>& InNodeNames, const FString& InCollapseNodeName, bool bSetupUndoRedo, bool bPrintPythonCommand, bool bIsAggregate)
 {
 	if (!IsValidGraph())
 	{
@@ -3509,7 +3510,7 @@ URigVMCollapseNode* URigVMController::CollapseNodes(const TArray<FName>& InNodeN
 		Nodes.AddUnique(Node);
 	}
 
-	URigVMCollapseNode* Node = CollapseNodes(Nodes, InCollapseNodeName, bSetupUndoRedo);
+	URigVMCollapseNode* Node = CollapseNodes(Nodes, InCollapseNodeName, bSetupUndoRedo, bIsAggregate);
 	if (Node && bPrintPythonCommand)
 	{
 		FString ArrayStr = TEXT("[");
@@ -3577,7 +3578,7 @@ TArray<URigVMNode*> URigVMController::ExpandLibraryNode(const FName& InNodeName,
 
 #endif
 
-URigVMCollapseNode* URigVMController::CollapseNodes(const TArray<URigVMNode*>& InNodes, const FString& InCollapseNodeName, bool bSetupUndoRedo)
+URigVMCollapseNode* URigVMController::CollapseNodes(const TArray<URigVMNode*>& InNodes, const FString& InCollapseNodeName, bool bSetupUndoRedo, bool bIsAggregate)
 {
 	URigVMGraph* Graph = GetGraph();
 	check(Graph);
@@ -3588,6 +3589,48 @@ URigVMCollapseNode* URigVMController::CollapseNodes(const TArray<URigVMNode*>& I
 		return nullptr;
 	}
 
+	if (InNodes.IsEmpty())
+	{
+		ReportError(TEXT("No nodes specified to collapse."));
+		return nullptr;
+	}
+
+#if UE_RIGVM_ENABLE_AGGREGATE_NODES
+	if (bIsAggregate)
+	{
+		if (URigVMUnitNode* UnitNode = Cast<URigVMUnitNode>(InNodes[0]))
+		{
+			for (int32 i=1; i<InNodes.Num(); ++i)
+			{
+				if (URigVMUnitNode* OtherUnitNode = Cast<URigVMUnitNode>(InNodes[i]))
+				{
+					if (UnitNode->GetScriptStruct() == OtherUnitNode->GetScriptStruct())
+					{
+						continue;
+					}
+				}
+				ReportError(TEXT("Cannot aggregate nodes of different types."));
+				return nullptr;
+			}
+		}
+		else if (URigVMFunctionReferenceNode* FunctionReferenceNode = Cast<URigVMFunctionReferenceNode>(InNodes[0]))
+		{
+			for (int32 i=1; i<InNodes.Num(); ++i)
+			{
+				if (URigVMFunctionReferenceNode* OtherFunctionReferenceNode = Cast<URigVMFunctionReferenceNode>(InNodes[i]))
+				{
+					if(FunctionReferenceNode->GetReferencedNode() == OtherFunctionReferenceNode->GetReferencedNode())
+					{
+						continue;
+					}
+				}
+				ReportError(TEXT("Cannot aggregate nodes of different types."));
+				return nullptr;
+			}
+		}
+	}
+#endif
+
 	TArray<URigVMNode*> Nodes;
 	for (URigVMNode* Node : InNodes)
 	{
@@ -3797,12 +3840,24 @@ URigVMCollapseNode* URigVMController::CollapseNodes(const TArray<URigVMNode*>& I
 
 	if (bSetupUndoRedo)
 	{
-		CollapseAction = FRigVMCollapseNodesAction(this, Nodes, CollapseNodeName); 
+		CollapseAction = FRigVMCollapseNodesAction(this, Nodes, CollapseNodeName, bIsAggregate); 
 		CollapseAction.Title = TEXT("Collapse Nodes");
 		ActionStack->BeginAction(CollapseAction);
 	}
 
+#if UE_RIGVM_ENABLE_AGGREGATE_NODES
+	URigVMCollapseNode* CollapseNode = nullptr;
+	if (bIsAggregate)
+	{
+		CollapseNode = NewObject<URigVMAggregateNode>(Graph, *CollapseNodeName);		
+	}
+	else
+	{
+		CollapseNode = NewObject<URigVMCollapseNode>(Graph, *CollapseNodeName);		
+	}
+#else
 	URigVMCollapseNode* CollapseNode = NewObject<URigVMCollapseNode>(Graph, *CollapseNodeName);
+#endif
 	CollapseNode->ContainedGraph = NewObject<URigVMGraph>(CollapseNode, TEXT("ContainedGraph"));
 	CollapseNode->Position = Center;
 	Graph->Nodes.Add(CollapseNode);
@@ -4145,6 +4200,15 @@ URigVMCollapseNode* URigVMController::CollapseNodes(const TArray<URigVMNode*>& I
 		RemoveNodeByName(NodeToRemove, false, true);
 	}
 
+	if (!InCollapseNodeName.IsEmpty() && CollapseNodeName != InCollapseNodeName)
+	{		
+		FString ValidName = GetValidNodeName(InCollapseNodeName);
+		if (ValidName == InCollapseNodeName)
+		{
+			RenameNode(CollapseNode, *ValidName, bSetupUndoRedo);
+		}
+	}
+
 	if (bSetupUndoRedo)
 	{
 		ActionStack->EndAction(CollapseAction);
@@ -6868,6 +6932,438 @@ FString URigVMController::GetPinInitialDefaultValueFromStruct(UScriptStruct* Scr
 	return DefaultValue;
 }
 
+FString URigVMController::AddAggregatePin(const FString& InNodeName, const FString& InPinName, const FString& InDefaultValue, bool bSetupUndoRedo, bool bPrintPythonCommand)
+{
+#if UE_RIGVM_ENABLE_AGGREGATE_NODES
+	 
+	if (!IsValidGraph())
+	{
+		return FString();
+	}
+
+	URigVMGraph* Graph = GetGraph();
+	check(Graph);
+
+	URigVMNode* Node = Graph->FindNodeByName(*InNodeName);
+	if (!Node)
+	{
+		return FString();
+	}
+
+	return AddAggregatePin(Node, InPinName, InDefaultValue, bSetupUndoRedo, bPrintPythonCommand);
+}
+
+FString URigVMController::AddAggregatePin(URigVMNode* InNode, const FString& InPinName, const FString& InDefaultValue, bool bSetupUndoRedo, bool bPrintPythonCommand)
+{
+	if (!InNode)
+	{
+		return FString();
+	}
+
+	if (!IsValidNodeForGraph(InNode))
+	{
+		return FString();
+	}
+
+	URigVMAggregateNode* AggregateNode = Cast<URigVMAggregateNode>(InNode);
+	URigVMUnitNode* UnitNode = Cast<URigVMUnitNode>(InNode);
+	URigVMFunctionReferenceNode* FunctionReferenceNode = Cast<URigVMFunctionReferenceNode>(InNode);
+
+	FRigVMControllerCompileBracketScope CompileScope(this);
+	FRigVMBaseAction Action;
+	if (bSetupUndoRedo)
+	{
+		Action.Title = FString::Printf(TEXT("Add Aggregate Pin"));
+		ActionStack->BeginAction(Action);
+	}
+
+	if (UnitNode || FunctionReferenceNode)
+	{
+		if (UnitNode && !UnitNode->IsAggregate())
+		{
+			if (bSetupUndoRedo)
+			{
+				ActionStack->CancelAction(Action);
+			}
+			return FString();
+		}
+
+		bool bAggregateInputs = false;
+		URigVMPin* Arg1 = nullptr;
+		URigVMPin* Arg2 = nullptr;
+		URigVMPin* ArgOpposite = nullptr;
+		if (UnitNode)
+		{
+			const TArray<URigVMPin*> AggregateInputs = UnitNode->GetAggregateInputs();
+			const TArray<URigVMPin*> AggregateOutputs = UnitNode->GetAggregateOutputs();
+
+			if (AggregateInputs.Num() == 2 && AggregateOutputs.Num() == 1)
+			{
+				bAggregateInputs = true;
+				Arg1 = AggregateInputs[0];
+				Arg2 = AggregateInputs[1];
+				ArgOpposite = AggregateOutputs[0];
+			}
+			else if (AggregateInputs.Num() == 1 && AggregateOutputs.Num() == 2)
+			{
+				bAggregateInputs = false;
+				Arg1 = AggregateOutputs[0];
+				Arg2 = AggregateOutputs[1];
+				ArgOpposite = AggregateInputs[0];
+			}
+			else
+			{
+				if (bSetupUndoRedo)
+				{
+					ActionStack->CancelAction(Action);
+				}
+				return FString();	
+			}
+		}
+		else if (FunctionReferenceNode)
+		{
+			
+		}
+
+		if (!Arg1 || !Arg2 || !ArgOpposite)
+		{
+			if (bSetupUndoRedo)
+			{
+				ActionStack->CancelAction(Action);
+			}
+			return FString();
+		}
+
+		if (Arg1->GetCPPType() != Arg2->GetCPPType() || Arg1->GetCPPTypeObject() != Arg2->GetCPPTypeObject() ||
+			Arg1->GetCPPType() != ArgOpposite->GetCPPType() || Arg1->GetCPPTypeObject() != ArgOpposite->GetCPPTypeObject())
+		{
+			if (bSetupUndoRedo)
+			{
+				ActionStack->CancelAction(Action);
+			}
+			return FString();
+		}
+
+		const FString AggregateArg1 = Arg1->GetName();
+		const FString AggregateArg2 = Arg2->GetName();
+		const FString AggregateOppositeArg = ArgOpposite->GetName();
+	
+		TArray<URigVMNode*> Nodes = {InNode};
+		URigVMCollapseNode* CollapseNode = CollapseNodes(Nodes, InNode->GetName(), bSetupUndoRedo, true);
+		if (!CollapseNode)
+		{
+			if (bSetupUndoRedo)
+			{
+				ActionStack->CancelAction(Action);
+			}
+			return FString();
+		}
+
+		AggregateNode = Cast<URigVMAggregateNode>(CollapseNode);
+		if (AggregateNode)
+		{
+			FRigVMControllerGraphGuard GraphGuard(this, AggregateNode->GetContainedGraph(), bSetupUndoRedo);
+			if (UnitNode)
+			{
+				uint16 Index = 0;
+				for (TFieldIterator<FProperty> It(UnitNode->GetScriptStruct()); It; ++It, ++Index)
+				{
+					FName PropertyName = It->GetFName();
+					if (URigVMPin* AggregatePin = AggregateNode->FindPin(PropertyName.ToString()))
+					{
+						SetExposedPinIndex(PropertyName, Index, bSetupUndoRedo);
+						continue;
+					}
+
+					URigVMPin* Pin = UnitNode->FindPin(PropertyName.ToString());
+					AddExposedPin(PropertyName, Pin->GetDirection(), Pin->GetCPPType(), *Pin->GetCPPTypeObject()->GetPathName(), Pin->GetDefaultValue());
+
+					FString PropertyNameStr = PropertyName.ToString();
+					if (Pin->GetDirection() == ERigVMPinDirection::Input)
+					{
+						AddLink(FString::Printf(TEXT("Entry.%s"), *PropertyNameStr), FString::Printf(TEXT("%s.%s"), *InNode->GetName(), *PropertyNameStr), bSetupUndoRedo);
+					}
+					else
+					{
+						AddLink(FString::Printf(TEXT("%s.%s"), *InNode->GetName(), *PropertyNameStr), FString::Printf(TEXT("Return.%s"), *PropertyNameStr), bSetupUndoRedo);
+					}
+				}
+			}
+			else if (FunctionReferenceNode)
+			{
+				
+			}
+		}
+		else
+		{
+			if (bSetupUndoRedo)
+			{
+				ActionStack->CancelAction(Action);
+			}
+			return FString();
+		}
+	}
+	
+	if (!AggregateNode)
+	{
+		if (bSetupUndoRedo)
+		{
+			ActionStack->CancelAction(Action);
+		}
+		return FString();
+	}
+
+	URigVMPin* NewPin = nullptr;	
+	{
+		FRigVMControllerGraphGuard GraphGuard(this, AggregateNode->GetContainedGraph(), bSetupUndoRedo);
+
+		URigVMNode* InnerNode = (AggregateNode == nullptr) ?
+									((UnitNode == nullptr) ? (URigVMNode*) FunctionReferenceNode : UnitNode)
+								: AggregateNode->GetFirstInnerNode();
+		URigVMNode* NewNode = nullptr;
+		if (URigVMUnitNode* InnerUnitNode = Cast<URigVMUnitNode>(InnerNode))
+		{		
+			NewNode = AddUnitNode(InnerUnitNode->GetScriptStruct(), InnerUnitNode->GetMethodName(), FVector2D::ZeroVector, FString(), bSetupUndoRedo);		
+		}
+		else if (URigVMFunctionReferenceNode* InnerReferenceNode = Cast<URigVMFunctionReferenceNode>(InnerNode))
+		{
+			NewNode = AddFunctionReferenceNode(FunctionReferenceNode->GetReferencedNode(), FVector2D::ZeroVector, FString(), bSetupUndoRedo);
+		}
+		else
+		{
+			if (bSetupUndoRedo)
+			{
+				ActionStack->CancelAction(Action);
+			}
+			return FString();
+		}
+
+		const URigVMPin* Arg1 = AggregateNode->GetFirstAggregatePin();
+		FName NewPinName = AddExposedPin(*InPinName, Arg1->GetDirection(), Arg1->GetCPPType(), *Arg1->GetCPPTypeObject()->GetPathName(), InDefaultValue, bSetupUndoRedo);
+		NewPin = AggregateNode->FindPin(NewPinName.ToString());
+		URigVMPin* NewUnitPinArg1 = NewNode->GetFirstAggregatePin();
+		URigVMPin* NewUnitPinArg2 = NewNode->GetSecondAggregatePin();
+		URigVMPin* NewUnitPinOppositeArg = NewNode->GetOppositeAggregatePin();
+		if(AggregateNode->IsInputAggregate())
+		{		
+			URigVMFunctionEntryNode* EntryNode = AggregateNode->GetEntryNode();
+			URigVMPin* EntryPin = EntryNode->FindPin(NewPinName.ToString());
+			URigVMPin* ReturnPin = AggregateNode->GetReturnNode()->FindPin(NewUnitPinOppositeArg->GetName());
+			URigVMPin* PreviousReturnPin = ReturnPin->GetLinkedSourcePins()[0];					
+		
+			BreakAllLinks(ReturnPin, true, bSetupUndoRedo);
+			AddLink(PreviousReturnPin, NewUnitPinArg1, bSetupUndoRedo);						
+			AddLink(EntryPin, NewUnitPinArg2, bSetupUndoRedo);
+			AddLink(NewUnitPinOppositeArg, ReturnPin, bSetupUndoRedo);
+		}
+		else
+		{
+			URigVMFunctionReturnNode* ReturnNode = AggregateNode->GetReturnNode();
+			URigVMPin* NewReturnPin = ReturnNode->FindPin(NewPinName.ToString());
+			URigVMPin* OldReturnPin = ReturnNode->GetPins()[ReturnNode->GetPins().Num()-2];
+			URigVMPin* PreviousReturnPin = OldReturnPin->GetLinkedSourcePins()[0];
+
+			BreakAllLinks(OldReturnPin, true, bSetupUndoRedo);
+			AddLink(PreviousReturnPin, NewUnitPinOppositeArg, bSetupUndoRedo);						
+			AddLink(NewUnitPinArg1, OldReturnPin, bSetupUndoRedo);
+			AddLink(NewUnitPinArg2, NewReturnPin, bSetupUndoRedo);
+
+		}
+
+		// Connect other input pins
+		for (URigVMPin* OtherEntryPin : AggregateNode->GetEntryNode()->GetPins())
+		{
+			if (OtherEntryPin->GetName() != NewUnitPinArg1->GetName() &&
+				OtherEntryPin->GetName() != NewUnitPinArg2->GetName() &&
+				OtherEntryPin->GetName() != NewUnitPinOppositeArg->GetName())
+			{
+				AddLink(OtherEntryPin, NewNode->FindPin(OtherEntryPin->GetName()), bSetupUndoRedo);
+			}
+		}
+	}
+
+	if (!NewPin)
+	{
+		if (bSetupUndoRedo)
+		{
+			ActionStack->CancelAction(Action);
+		}
+		return FString();
+	}
+
+	if (bSetupUndoRedo)
+	{
+		ActionStack->EndAction(Action);
+	}
+
+	if (bPrintPythonCommand)
+	{
+		const FString GraphName = GetSanitizedGraphName(GetGraph()->GetGraphName());
+		const FString NodePath = GetSanitizedPinPath(InNode->GetNodePath());
+
+		RigVMPythonUtils::Print(GetGraphOuterName(),
+			FString::Printf(TEXT("blueprint.get_controller_by_name('%s').add_aggregate_pin('%s', '%s', '%s')"),
+			*GraphName,
+			*NodePath,
+			*InPinName,
+			*InDefaultValue));
+	}
+	
+	return NewPin->GetPinPath();
+
+#else
+	return FString();
+#endif
+}
+
+bool URigVMController::RemoveAggregatePin(const FString& InPinPath, bool bSetupUndoRedo, bool bPrintPythonCommand)
+{
+#if UE_RIGVM_ENABLE_AGGREGATE_NODES
+	
+	if (!IsValidGraph())
+	{
+		return false;
+	}
+
+	URigVMGraph* Graph = GetGraph();
+	check(Graph);
+
+	URigVMPin* Pin = Graph->FindPin(*InPinPath);
+	if (!Pin)
+	{
+		return false;
+	}
+
+	return RemoveAggregatePin(Pin, bSetupUndoRedo, bPrintPythonCommand);
+}
+
+bool URigVMController::RemoveAggregatePin(URigVMPin* InPin, bool bSetupUndoRedo, bool bPrintPythonCommand)
+{
+	if (!InPin)
+	{
+		return false;
+	}
+
+	if (InPin->GetParentPin())
+	{
+		return false;
+	}
+
+	FRigVMControllerCompileBracketScope CompileScope(this);
+	FRigVMBaseAction Action;
+	if (bSetupUndoRedo)
+	{
+		Action.Title = FString::Printf(TEXT("Remove Aggregate Pin"));
+		ActionStack->BeginAction(Action);
+	}
+
+	bool bSuccess = false;
+	if (URigVMAggregateNode* AggregateNode = Cast<URigVMAggregateNode>(InPin->GetNode()))
+	{
+		URigVMGraph* Graph = AggregateNode->GetContainedGraph();
+		if (AggregateNode->IsInputAggregate())
+		{
+			if (URigVMFunctionEntryNode* EntryNode = Graph->GetEntryNode())
+			{
+				if (URigVMPin* EntryPin = EntryNode->FindPin(InPin->GetName()))
+				{
+					if (EntryPin->GetLinkedTargetPins().Num() > 0)
+					{
+						FRigVMControllerGraphGuard GraphGuard(this, AggregateNode->GetContainedGraph(), bSetupUndoRedo);
+						URigVMPin* TargetPin = EntryPin->GetLinkedTargetPins()[0];
+					
+						URigVMNode* NodeToRemove = TargetPin->GetNode();
+						URigVMPin* ResultPin = NodeToRemove->GetOppositeAggregatePin();
+						URigVMPin* NextNodePin = ResultPin->GetLinkedTargetPins()[0];
+
+						if (NodeToRemove == AggregateNode->FirstInnerNodeCache)
+						{
+							AggregateNode->InvalidateCache();
+						}
+
+						const FString FirstAggregatePin = AggregateNode->GetFirstAggregatePin()->GetName();
+						const FString SecondAggregatePin = AggregateNode->GetSecondAggregatePin()->GetName();
+						FString OtherArg = TargetPin->GetName() == FirstAggregatePin ? SecondAggregatePin : FirstAggregatePin;
+						BreakAllLinks(NextNodePin, true, bSetupUndoRedo);
+						RewireLinks(NodeToRemove->FindPin(OtherArg), NextNodePin, true, bSetupUndoRedo);
+						RemoveNode(NodeToRemove, bSetupUndoRedo);
+						RemoveExposedPin(*InPin->GetName(), bSetupUndoRedo);
+						bSuccess = true;						
+					}
+				}
+			}
+		}
+		else
+		{
+			if (URigVMFunctionReturnNode* ReturnNode = Graph->GetReturnNode())
+			{
+				if (URigVMPin* ReturnPin = ReturnNode->FindPin(InPin->GetName()))
+				{
+					if (ReturnPin->GetLinkedSourcePins().Num() > 0)
+					{
+						FRigVMControllerGraphGuard GraphGuard(this, AggregateNode->GetContainedGraph(), bSetupUndoRedo);
+						URigVMPin* SourcePin = ReturnPin->GetLinkedSourcePins()[0];
+					
+						URigVMNode* NodeToRemove = SourcePin->GetNode();
+						URigVMPin* OppositePin = NodeToRemove->GetOppositeAggregatePin();
+						URigVMPin* NextNodePin = OppositePin->GetLinkedSourcePins()[0];
+						URigVMNode* NextNode = NextNodePin->GetNode();
+
+						if (NodeToRemove == AggregateNode->FirstInnerNodeCache)
+						{
+							AggregateNode->InvalidateCache();
+						}
+
+						const FString FirstAggregatePin = AggregateNode->GetFirstAggregatePin()->GetName();
+						const FString SecondAggregatePin = AggregateNode->GetSecondAggregatePin()->GetName();
+						FString OtherArg = SourcePin->GetName() == FirstAggregatePin ? SecondAggregatePin : FirstAggregatePin;
+						BreakAllLinks(NextNodePin, false, bSetupUndoRedo);
+						RewireLinks(NodeToRemove->FindPin(OtherArg), NextNodePin, false, bSetupUndoRedo);
+						RemoveNode(NodeToRemove, bSetupUndoRedo);
+						RemoveExposedPin(*InPin->GetName(), bSetupUndoRedo);
+						bSuccess = true;	
+					}
+				}
+			}			
+		}
+
+		if (bSuccess && AggregateNode->GetContainedNodes().Num() == 3)
+		{
+			TArray<URigVMNode*> NodesEjected = ExpandLibraryNode(AggregateNode, bSetupUndoRedo);
+			bSuccess = NodesEjected.Num() == 1;
+		}		
+	}
+
+	if (bSetupUndoRedo)
+	{
+		if (bSuccess)
+		{
+			ActionStack->EndAction(Action);
+		}
+		else
+		{
+			ActionStack->CancelAction(Action);			
+		}
+	}
+
+	if (bSuccess && bPrintPythonCommand)
+	{
+		const FString GraphName = GetSanitizedGraphName(GetGraph()->GetGraphName());
+		const FString PinPath = GetSanitizedPinPath(InPin->GetPinPath());
+
+		RigVMPythonUtils::Print(GetGraphOuterName(),
+			FString::Printf(TEXT("blueprint.get_controller_by_name('%s').remove_aggregate_pin('%s')"),
+			*GraphName,
+			*PinPath));
+	}
+
+	return bSuccess;
+
+#else
+	return false;
+#endif
+}
+
 FString URigVMController::AddArrayPin(const FString& InArrayPinPath, const FString& InDefaultValue, bool bSetupUndoRedo, bool bPrintPythonCommand)
 {
 	return InsertArrayPin(InArrayPinPath, INDEX_NONE, InDefaultValue, bSetupUndoRedo, bPrintPythonCommand);
@@ -7742,8 +8238,7 @@ bool URigVMController::AddLink(const FString& InOutputPinPath, const FString& In
 	return bSuccess;
 }
 
-bool URigVMController::AddLink(URigVMPin* OutputPin, URigVMPin* InputPin, bool bSetupUndoRedo,
-	ERigVMPinDirection InUserDirection)
+bool URigVMController::AddLink(URigVMPin* OutputPin, URigVMPin* InputPin, bool bSetupUndoRedo, ERigVMPinDirection InUserDirection)
 {
 	if(OutputPin == nullptr)
 	{
@@ -8404,7 +8899,7 @@ bool URigVMController::RemoveExposedPin(const FName& InPinName, bool bSetupUndoR
 
 	bool bSuccessfullyRemovedPin = false;
 	{
-		FRigVMControllerGraphGuard GraphGuard(this, LibraryNode->GetGraph(), false);
+		FRigVMControllerGraphGuard GraphGuard(this, LibraryNode->GetGraph(), bSetupUndoRedo);
 		bSuccessfullyRemovedPin = RemovePin(Pin, bSetupUndoRedo, true);
 	}
 
@@ -8815,7 +9310,7 @@ bool URigVMController::SetExposedPinIndex(const FName& InPinName, int32 InNewInd
 
 	if (Pin->GetPinIndex() == InNewIndex)
 	{
-		return false;
+		return true; // Nothing to do, do not fail
 	}
 
 	if (InNewIndex < 0 || InNewIndex >= LibraryNode->GetPins().Num())
diff --git a/Engine/Source/Developer/RigVMDeveloper/Private/RigVMModel/RigVMControllerActions.cpp b/Engine/Source/Developer/RigVMDeveloper/Private/RigVMModel/RigVMControllerActions.cpp
--- a/Engine/Source/Developer/RigVMDeveloper/Private/RigVMModel/RigVMControllerActions.cpp
+++ b/Engine/Source/Developer/RigVMDeveloper/Private/RigVMModel/RigVMControllerActions.cpp
@@ -1507,12 +1507,12 @@ bool FRigVMImportNodeFromTextAction::Redo(URigVMController* InController)
 }
 
 FRigVMCollapseNodesAction::FRigVMCollapseNodesAction()
-	: LibraryNodePath()
+	: LibraryNodePath(), bIsAggregate(false)
 {
 }
 
-FRigVMCollapseNodesAction::FRigVMCollapseNodesAction(URigVMController* InController, const TArray<URigVMNode*>& InNodes, const FString& InNodePath)
-	: LibraryNodePath(InNodePath)
+FRigVMCollapseNodesAction::FRigVMCollapseNodesAction(URigVMController* InController, const TArray<URigVMNode*>& InNodes, const FString& InNodePath, const bool bIsAggregate)
+	: LibraryNodePath(InNodePath), bIsAggregate(bIsAggregate)
 {
 	TArray<FName> NodesToExport;
 	for (URigVMNode* InNode : InNodes)
@@ -1576,7 +1576,7 @@ bool FRigVMCollapseNodesAction::Redo(URigVMController* InController)
 		NodeNames.Add(*NodePath);
 	}
 
-	URigVMLibraryNode* LibraryNode = InController->CollapseNodes(NodeNames, LibraryNodePath, false);
+	URigVMLibraryNode* LibraryNode = InController->CollapseNodes(NodeNames, LibraryNodePath, false, false, bIsAggregate);
 	if (LibraryNode)
 	{
 		return FRigVMBaseAction::Redo(InController);
@@ -1768,12 +1768,18 @@ bool FRigVMAddExposedPinAction::Redo(URigVMController* InController)
 	return false;
 }
 
+FRigVMRemoveExposedPinAction::FRigVMRemoveExposedPinAction()
+	: PinIndex(INDEX_NONE)
+{
+}
+
 FRigVMRemoveExposedPinAction::FRigVMRemoveExposedPinAction(URigVMPin* InPin)
 	: PinName(InPin->GetName())
 	, Direction(InPin->GetDirection())
 	, CPPType(InPin->GetCPPType())
 	, CPPTypeObjectPath()
 	, DefaultValue(InPin->GetDefaultValue())
+	, PinIndex(InPin->GetPinIndex())
 {
 	if (UObject* CPPTypeObject = InPin->GetCPPTypeObject())
 	{
@@ -1785,7 +1791,10 @@ bool FRigVMRemoveExposedPinAction::Undo(URigVMController* InController)
 {
 	if(!InController->AddExposedPin(*PinName, Direction, CPPType, *CPPTypeObjectPath, DefaultValue, false).IsNone())
 	{
-		return FRigVMBaseAction::Undo(InController);
+		if (InController->SetExposedPinIndex(*PinName, PinIndex, false))
+		{
+			return FRigVMBaseAction::Undo(InController);
+		}
 	}
 	return false;
 }
@@ -1831,7 +1840,7 @@ FRigVMSetPinIndexAction::FRigVMSetPinIndexAction()
 }
 
 FRigVMSetPinIndexAction::FRigVMSetPinIndexAction(URigVMPin* InPin, int32 InNewIndex)
-	: PinPath(InPin->GetPinPath())
+	: PinPath(InPin->GetName())
 	, OldIndex(InPin->GetPinIndex())
 	, NewIndex(InNewIndex)
 {
@@ -1843,12 +1852,12 @@ bool FRigVMSetPinIndexAction::Undo(URigVMController* InController)
 	{
 		return false;
 	}
-	return InController->SetExposedPinIndex(*PinPath, OldIndex);
+	return InController->SetExposedPinIndex(*PinPath, OldIndex, false);
 }
 
 bool FRigVMSetPinIndexAction::Redo(URigVMController* InController)
 {
-	if (!InController->SetExposedPinIndex(*PinPath, NewIndex))
+	if (!InController->SetExposedPinIndex(*PinPath, NewIndex, false))
 	{
 		return false;
 	}
diff --git a/Engine/Source/Developer/RigVMDeveloper/Public/RigVMModel/Nodes/RigVMAggregateNode.h b/Engine/Source/Developer/RigVMDeveloper/Public/RigVMModel/Nodes/RigVMAggregateNode.h
new file mode 100644
--- /dev/null
+++ b/Engine/Source/Developer/RigVMDeveloper/Public/RigVMModel/Nodes/RigVMAggregateNode.h
@@ -0,0 +1,54 @@
+﻿// Copyright Epic Games, Inc. All Rights Reserved.
+
+#pragma once
+
+#include "RigVMUnitNode.h"
+#include "RigVMCollapseNode.h"
+#include "RigVMModel/Nodes/RigVMUnitNode.h"
+#include "RigVMAggregateNode.generated.h"
+
+/**
+ * The Aggregate Node contains a subgraph of nodes with aggregate pins (2in+1out or 1out+2in) connected
+ * to each other. For example, a unit node IntAdd which adds 2 integers and provides Result=A+B could have
+ * A, B and Result as aggregates in order to add additional input pins to construct an Aggregate Node that computes
+ * Result=A+B+C.
+ */
+UCLASS(BlueprintType)
+class RIGVMDEVELOPER_API URigVMAggregateNode : public URigVMCollapseNode
+{
+	GENERATED_BODY()
+
+public:
+
+	// default constructor
+	URigVMAggregateNode();
+
+	// URigVMUnitNode interface
+	virtual FString GetNodeTitle() const override;
+	virtual FName GetMethodName() const;
+	virtual  FText GetToolTipText() const override;
+	virtual FText GetToolTipTextForPin(const URigVMPin* InPin) const override;
+
+	virtual bool IsAggregate() const override { return true; };
+	virtual bool IsInputAggregate() const override;
+
+	URigVMNode* GetFirstInnerNode() const;
+
+	UFUNCTION(BlueprintPure, Category = Aggregate)
+	URigVMPin* GetFirstAggregatePin() const;
+
+	UFUNCTION(BlueprintPure, Category = Aggregate)
+	URigVMPin* GetSecondAggregatePin() const;
+
+	UFUNCTION(BlueprintPure, Category = Aggregate)
+	URigVMPin* GetOppositeAggregatePin() const;
+
+private:
+
+	void InvalidateCache();
+
+	mutable URigVMNode* FirstInnerNodeCache;
+
+	friend class URigVMController;
+};
+
diff --git a/Engine/Source/Developer/RigVMDeveloper/Public/RigVMModel/Nodes/RigVMUnitNode.h b/Engine/Source/Developer/RigVMDeveloper/Public/RigVMModel/Nodes/RigVMUnitNode.h
--- a/Engine/Source/Developer/RigVMDeveloper/Public/RigVMModel/Nodes/RigVMUnitNode.h
+++ b/Engine/Source/Developer/RigVMDeveloper/Public/RigVMModel/Nodes/RigVMUnitNode.h
@@ -27,10 +27,20 @@ public:
 	virtual FName GetEventName() const override;
 	virtual bool IsLoopNode() const override;
 	virtual bool CanBeUpgraded() const override { return GetUpgradeInfo().IsValid(); }
+	virtual bool IsAggregate() const override;
+	virtual URigVMPin* GetFirstAggregatePin() const override;
+	virtual URigVMPin* GetSecondAggregatePin() const override;
+	virtual URigVMPin* GetOppositeAggregatePin() const override;
+	virtual bool IsInputAggregate() const override;
+	
+	TArray<URigVMPin*> GetAggregateInputs() const;
+	TArray<URigVMPin*> GetAggregateOutputs() const;
 
 	bool IsDeprecated() const;
 	FString GetDeprecatedMetadata() const;
 
+	
+	
 	// Returns the UStruct for this unit node
 	// (the struct declaring the RIGVM_METHOD)
 	UFUNCTION(BlueprintCallable, Category = RigVMUnitNode)
diff --git a/Engine/Source/Developer/RigVMDeveloper/Public/RigVMModel/RigVMController.h b/Engine/Source/Developer/RigVMDeveloper/Public/RigVMModel/RigVMController.h
--- a/Engine/Source/Developer/RigVMDeveloper/Public/RigVMModel/RigVMController.h
+++ b/Engine/Source/Developer/RigVMDeveloper/Public/RigVMModel/RigVMController.h
@@ -24,6 +24,10 @@
 #define UE_RIGVM_ENABLE_TEMPLATE_NODES 1
 #endif
 
+#ifndef UE_RIGVM_ENABLE_AGGREGATE_NODES
+#define UE_RIGVM_ENABLE_AGGREGATE_NODES 0
+#endif
+
 class URigVMActionStack;
 
 UENUM()
@@ -421,7 +425,7 @@ public:
 
 	// Turns a series of nodes into a Collapse node
 	UFUNCTION(BlueprintCallable, Category = RigVMController)
-	URigVMCollapseNode* CollapseNodes(const TArray<FName>& InNodeNames, const FString& InCollapseNodeName = TEXT(""), bool bSetupUndoRedo = true, bool bPrintPythonCommand = false);
+	URigVMCollapseNode* CollapseNodes(const TArray<FName>& InNodeNames, const FString& InCollapseNodeName = TEXT(""), bool bSetupUndoRedo = true, bool bPrintPythonCommand = false, bool bIsAggregate = false);
 
 	// Turns a library node into its contained nodes
 	UFUNCTION(BlueprintCallable, Category = RigVMController)
@@ -594,6 +598,17 @@ public:
 	UFUNCTION(BlueprintCallable, Category = RigVMController)
 	bool ResetPinDefaultValue(const FString& InPinPath, bool bSetupUndoRedo = true, bool bPrintPythonCommand = false);
 
+	UFUNCTION(BlueprintCallable, Category = RigVMController)
+	FString AddAggregatePin(const FString& InNodeName, const FString& InPinName, const FString& InDefaultValue = TEXT(""), bool bSetupUndoRedo = true, bool bPrintPythonCommand = false);
+
+	UFUNCTION(BlueprintCallable, Category = RigVMController)
+	bool RemoveAggregatePin(const FString& InPinPath, bool bSetupUndoRedo = true, bool bPrintPythonCommand = false);
+
+#if UE_RIGVM_ENABLE_AGGREGATE_NODES
+	FString AddAggregatePin(URigVMNode* InNode, const FString& InPinName, const FString& InDefaultValue = TEXT(""), bool bSetupUndoRedo = true, bool bPrintPythonCommand = false);
+	bool RemoveAggregatePin(URigVMPin* InPin, bool bSetupUndoRedo = true, bool bPrintPythonCommand = false);
+#endif
+
 	// Adds an array element pin to the end of an array pin.
 	// This causes a PinArraySizeChanged modified event.
 	UFUNCTION(BlueprintCallable, Category = RigVMController)
@@ -911,7 +926,7 @@ private:
 	bool SetVariableName(URigVMVariableNode* InVariableNode, const FName& InVariableName, bool bSetupUndoRedo);
 	static void ForEveryPinRecursively(URigVMPin* InPin, TFunction<void(URigVMPin*)> OnEachPinFunction);
 	static void ForEveryPinRecursively(URigVMNode* InNode, TFunction<void(URigVMPin*)> OnEachPinFunction);
-	URigVMCollapseNode* CollapseNodes(const TArray<URigVMNode*>& InNodes, const FString& InCollapseNodeName, bool bSetupUndoRedo);
+	URigVMCollapseNode* CollapseNodes(const TArray<URigVMNode*>& InNodes, const FString& InCollapseNodeName, bool bSetupUndoRedo, bool bIsAggregate);
 	TArray<URigVMNode*> ExpandLibraryNode(URigVMLibraryNode* InNode, bool bSetupUndoRedo);
 	URigVMFunctionReferenceNode* PromoteCollapseNodeToFunctionReferenceNode(URigVMCollapseNode* InCollapseNode, bool bSetupUndoRedo, const FString& InExistingFunctionDefinitionPath);
 	URigVMCollapseNode* PromoteFunctionReferenceNodeToCollapseNode(URigVMFunctionReferenceNode* InFunctionRefNode, bool bSetupUndoRedo, bool bRemoveFunctionDefinition);
diff --git a/Engine/Source/Developer/RigVMDeveloper/Public/RigVMModel/RigVMControllerActions.h b/Engine/Source/Developer/RigVMDeveloper/Public/RigVMModel/RigVMControllerActions.h
--- a/Engine/Source/Developer/RigVMDeveloper/Public/RigVMModel/RigVMControllerActions.h
+++ b/Engine/Source/Developer/RigVMDeveloper/Public/RigVMModel/RigVMControllerActions.h
@@ -1302,7 +1302,7 @@ struct FRigVMCollapseNodesAction : public FRigVMBaseAction
 public:
 
 	FRigVMCollapseNodesAction();
-	FRigVMCollapseNodesAction(URigVMController* InController, const TArray<URigVMNode*>& InNodes, const FString& InNodePath);
+	FRigVMCollapseNodesAction(URigVMController* InController, const TArray<URigVMNode*>& InNodes, const FString& InNodePath, const bool bIsAggregate);
 	virtual ~FRigVMCollapseNodesAction() {};
 	virtual UScriptStruct* GetScriptStruct() const override { return FRigVMCollapseNodesAction::StaticStruct(); }
 	virtual bool Undo(URigVMController* InController) override;
@@ -1319,6 +1319,9 @@ public:
 
 	UPROPERTY()
 	TArray<FString> CollapsedNodesLinks;
+
+	UPROPERTY()
+	bool bIsAggregate;
 };
 
 /**
@@ -1461,7 +1464,7 @@ struct FRigVMRemoveExposedPinAction : public FRigVMBaseAction
 
 public:
 
-	FRigVMRemoveExposedPinAction() {}
+	FRigVMRemoveExposedPinAction();
 	FRigVMRemoveExposedPinAction(URigVMPin* InPin);
 	virtual ~FRigVMRemoveExposedPinAction() {};
 	virtual UScriptStruct* GetScriptStruct() const override { return FRigVMRemoveExposedPinAction::StaticStruct(); }
@@ -1482,6 +1485,9 @@ public:
 
 	UPROPERTY()
 	FString DefaultValue;
+
+	UPROPERTY()
+	int32 PinIndex;
 };
 
 /**
diff --git a/Engine/Source/Developer/RigVMDeveloper/Public/RigVMModel/RigVMNode.h b/Engine/Source/Developer/RigVMDeveloper/Public/RigVMModel/RigVMNode.h
--- a/Engine/Source/Developer/RigVMDeveloper/Public/RigVMModel/RigVMNode.h
+++ b/Engine/Source/Developer/RigVMDeveloper/Public/RigVMModel/RigVMNode.h
@@ -217,6 +217,12 @@ public:
 
 	UFUNCTION(BlueprintCallable, Category = RigVMNode)
 	void SetExecutionIsHaltedAtThisNode(const bool bValue) { bHaltedAtThisNode = bValue; }
+
+	virtual bool IsAggregate() const { return false; }
+	virtual URigVMPin* GetFirstAggregatePin() const { return nullptr; }
+	virtual URigVMPin* GetSecondAggregatePin() const { return nullptr; }
+	virtual URigVMPin* GetOppositeAggregatePin() const { return nullptr; }
+	virtual bool IsInputAggregate() const { return false; }
 private:
 
 	static const FString NodeColorName;
diff --git a/Engine/Source/Runtime/RigVM/Private/RigVMCore/RigVMStruct.cpp b/Engine/Source/Runtime/RigVM/Private/RigVMCore/RigVMStruct.cpp
--- a/Engine/Source/Runtime/RigVM/Private/RigVMCore/RigVMStruct.cpp
+++ b/Engine/Source/Runtime/RigVM/Private/RigVMCore/RigVMStruct.cpp
@@ -199,6 +199,7 @@ const FName FRigVMStruct::NodeColorMetaName("NodeColor");
 const FName FRigVMStruct::IconMetaName("Icon");
 const FName FRigVMStruct::KeywordsMetaName("Keywords");
 const FName FRigVMStruct::TemplateNameMetaName = FRigVMRegistry::TemplateNameMetaName;
+const FName FRigVMStruct::AggregateMetaName("Aggregate");
 const FName FRigVMStruct::ExpandPinByDefaultMetaName("ExpandByDefault");
 const FName FRigVMStruct::DefaultArraySizeMetaName("DefaultArraySize");
 const FName FRigVMStruct::VaryingMetaName("Varying");
diff --git a/Engine/Source/Runtime/RigVM/Public/RigVMCore/RigVMStruct.h b/Engine/Source/Runtime/RigVM/Public/RigVMCore/RigVMStruct.h
--- a/Engine/Source/Runtime/RigVM/Public/RigVMCore/RigVMStruct.h
+++ b/Engine/Source/Runtime/RigVM/Public/RigVMCore/RigVMStruct.h
@@ -257,6 +257,7 @@ public:
 	static const FName IconMetaName;
 	static const FName KeywordsMetaName;
 	static const FName TemplateNameMetaName;
+	static const FName AggregateMetaName;
 	static const FName ExpandPinByDefaultMetaName;
 	static const FName DefaultArraySizeMetaName;
 	static const FName VaryingMetaName;
