Controlr Rig: UI implementation of Aggregate nodes

diff --git a/Engine/Plugins/Animation/ControlRig/Source/ControlRigDeveloper/Private/ControlRigBlueprint.cpp b/Engine/Plugins/Animation/ControlRig/Source/ControlRigDeveloper/Private/ControlRigBlueprint.cpp
--- a/Engine/Plugins/Animation/ControlRig/Source/ControlRigDeveloper/Private/ControlRigBlueprint.cpp
+++ b/Engine/Plugins/Animation/ControlRig/Source/ControlRigDeveloper/Private/ControlRigBlueprint.cpp
@@ -27,6 +27,7 @@
 #include "AssetRegistryModule.h"
 #include "RigVMPythonUtils.h"
 #include "RigVMTypeUtils.h"
+#include "RigVMModel/Nodes/RigVMAggregateNode.h"
 
 #if WITH_EDITOR
 #include "IControlRigEditorModule.h"
@@ -454,6 +455,14 @@ void UControlRigBlueprint::PostLoad()
 		}
 		UbergraphPages = NewUberGraphPages;
 
+		TArray<TGuardValue<bool>> EditableGraphGuards;
+		{
+			for (URigVMGraph* Graph : GetAllModels())
+			{
+				EditableGraphGuards.Emplace(Graph->bEditable, true);
+			}
+		}
+		
 		InitializeModelIfRequired(false /* recompile vm */);
 
 		PatchFunctionReferencesOnLoad();
@@ -4550,11 +4559,17 @@ void UControlRigBlueprint::CreateEdGraphForCollapseNodeIfNeeded(URigVMCollapseNo
 
 			if (!bSubGraphExists)
 			{
+				bool bEditable = true;
+				if (InNode->IsA<URigVMAggregateNode>())
+				{
+					bEditable = false;
+				}
+				
 				// create a sub graph
 				UControlRigGraph* SubRigGraph = NewObject<UControlRigGraph>(RigGraph, *InNode->GetEditorSubGraphName(), RF_Transactional);
 				SubRigGraph->Schema = UControlRigGraphSchema::StaticClass();
 				SubRigGraph->bAllowRenaming = 1;
-				SubRigGraph->bEditable = 1;
+				SubRigGraph->bEditable = bEditable;
 				SubRigGraph->bAllowDeletion = 1;
 				SubRigGraph->ModelNodePath = ContainedGraph->GetNodePath();
 				SubRigGraph->bIsFunctionDefinition = false;
diff --git a/Engine/Plugins/Animation/ControlRig/Source/ControlRigDeveloper/Private/Graph/ControlRigGraphNode.cpp b/Engine/Plugins/Animation/ControlRig/Source/ControlRigDeveloper/Private/Graph/ControlRigGraphNode.cpp
--- a/Engine/Plugins/Animation/ControlRig/Source/ControlRigDeveloper/Private/Graph/ControlRigGraphNode.cpp
+++ b/Engine/Plugins/Animation/ControlRig/Source/ControlRigDeveloper/Private/Graph/ControlRigGraphNode.cpp
@@ -1062,6 +1062,16 @@ URigVMPin* UControlRigGraphNode::GetModelPinFromPinPath(const FString& InPinPath
 	return nullptr;
 }
 
+void UControlRigGraphNode::HandleAddAggregateElement(const FString& InNodePath)
+{
+	DECLARE_SCOPE_HIERARCHICAL_COUNTER_FUNC()
+
+	if (URigVMController* Controller = GetController())
+	{
+		Controller->AddAggregatePin(InNodePath, FString(), FString(), true, true);
+	}	
+}
+
 void UControlRigGraphNode::SetupPinDefaultsFromModel(UEdGraphPin* Pin)
 {
 	DECLARE_SCOPE_HIERARCHICAL_COUNTER_FUNC()
diff --git a/Engine/Plugins/Animation/ControlRig/Source/ControlRigDeveloper/Public/Graph/ControlRigGraphNode.h b/Engine/Plugins/Animation/ControlRig/Source/ControlRigDeveloper/Public/Graph/ControlRigGraphNode.h
--- a/Engine/Plugins/Animation/ControlRig/Source/ControlRigDeveloper/Public/Graph/ControlRigGraphNode.h
+++ b/Engine/Plugins/Animation/ControlRig/Source/ControlRigDeveloper/Public/Graph/ControlRigGraphNode.h
@@ -135,6 +135,9 @@ public:
 
 	URigVMPin* GetModelPinFromPinPath(const FString& InPinPath) const;
 
+	/** Add a new element to the aggregate node referred to by the property path */
+	void HandleAddAggregateElement(const FString& InNodePath);
+
 	/** Add a new array element to the array referred to by the property path */
 	void HandleAddArrayElement(FString InPinPath);
 	
diff --git a/Engine/Plugins/Animation/ControlRig/Source/ControlRigEditor/Private/ControlRigEditorModule.cpp b/Engine/Plugins/Animation/ControlRig/Source/ControlRigEditor/Private/ControlRigEditorModule.cpp
--- a/Engine/Plugins/Animation/ControlRig/Source/ControlRigEditor/Private/ControlRigEditorModule.cpp
+++ b/Engine/Plugins/Animation/ControlRig/Source/ControlRigEditor/Private/ControlRigEditorModule.cpp
@@ -125,6 +125,7 @@
 #include "ControlRigSpaceChannelEditors.h"
 #include "UserDefinedStructure/UserDefinedStructEditorData.h"
 #include "UObject/FieldIterator.h"
+#include "RigVMModel/Nodes/RigVMAggregateNode.h"
 
 #define LOCTEXT_NAMESPACE "ControlRigEditorModule"
 
@@ -1293,6 +1294,20 @@ void FControlRigEditorModule::GetContextMenuActions(const UControlRigGraphSchema
 						));
 					}
 
+					if (URigVMAggregateNode* AggregateNode = Cast<URigVMAggregateNode>(ModelPin->GetNode()))
+					{
+						FToolMenuSection& Section = Menu->AddSection("EdGraphSchemaAggregatePin", LOCTEXT("AggregatePin", "Aggregates"));
+						Section.AddMenuEntry(
+							"RemoveAggregatePin",
+							LOCTEXT("RemoveAggregatePin", "Remove Aggregate Element"),
+							LOCTEXT("RemoveAggregatePin_Tooltip", "Removes the selected element from the aggregate"),
+							FSlateIcon(),
+							FUIAction(FExecuteAction::CreateLambda([Controller, ModelPin]() {
+								Controller->RemoveAggregatePin(ModelPin->GetPinPath(), true, true);
+							})
+						));
+					}
+
 					if (URigVMTemplateNode* TemplateNode = Cast<URigVMTemplateNode>(ModelPin->GetNode()))
 					{
 						FToolMenuSection& TemplatesSection = Menu->AddSection("EdGraphSchemaTemplates", LOCTEXT("TemplatesHeader", "Templates"));
diff --git a/Engine/Plugins/Animation/ControlRig/Source/ControlRigEditor/Private/Graph/SControlRigGraphNode.cpp b/Engine/Plugins/Animation/ControlRig/Source/ControlRigEditor/Private/Graph/SControlRigGraphNode.cpp
--- a/Engine/Plugins/Animation/ControlRig/Source/ControlRigEditor/Private/Graph/SControlRigGraphNode.cpp
+++ b/Engine/Plugins/Animation/ControlRig/Source/ControlRigEditor/Private/Graph/SControlRigGraphNode.cpp
@@ -580,6 +580,70 @@ bool SControlRigGraphNode::UseLowDetailPinNames() const
 	return false;
 }
 
+void SControlRigGraphNode::CreateInputSideAddButton(TSharedPtr<SVerticalBox> InputBox)
+{
+	if (ModelNode.IsValid() && ModelNode->IsAggregate() && ModelNode->IsInputAggregate())
+	{
+		TSharedRef<SWidget> AddPinButton = AddPinButtonContent(
+		   NSLOCTEXT("ControlRigAggregateNode", "ControlRigAggregateNodeAddPinButton", "Add pin"),
+		   NSLOCTEXT("ControlRigAggregateNode", "ControlRigAggregateNodeAddPinButton_Tooltip", "Adds a input pin to the node"),
+		   false);
+
+		FMargin AddPinPadding = Settings->GetInputPinPadding();
+		AddPinPadding.Top += 6.0f;
+
+		InputBox->AddSlot()
+			.AutoHeight()
+			.VAlign(VAlign_Center)
+			.Padding(AddPinPadding)
+			[
+				AddPinButton
+			];
+	}
+	else
+	{
+		SGraphNode::CreateInputSideAddButton(InputBox);
+	}
+}
+
+void SControlRigGraphNode::CreateOutputSideAddButton(TSharedPtr<SVerticalBox> InputBox)
+{
+	if (ModelNode.IsValid() && ModelNode->IsAggregate() && !ModelNode->IsInputAggregate())
+	{
+		TSharedRef<SWidget> AddPinButton = AddPinButtonContent(
+		   NSLOCTEXT("LayeredBoneBlendNode", "LayeredBoneBlendNodeAddPinButton", "Add pin"),
+		   NSLOCTEXT("LayeredBoneBlendNode", "LayeredBoneBlendNodeAddPinButton_Tooltip", "Adds a output pin to the node"),
+		   false);
+
+		FMargin AddPinPadding = Settings->GetOutputPinPadding();
+		AddPinPadding.Top += 6.0f;
+
+		InputBox->AddSlot()
+			.AutoHeight()
+			.VAlign(VAlign_Center)
+			.Padding(AddPinPadding)
+			[
+				AddPinButton
+			];
+	}
+	else
+	{
+		SGraphNode::CreateOutputSideAddButton(InputBox);
+	}
+}
+
+FReply SControlRigGraphNode::OnAddPin()
+{
+	if (ModelNode.IsValid())
+	{
+		if (UControlRigGraphNode* ControlRigGraphNode = Cast<UControlRigGraphNode>(GraphNode))
+		{
+			ControlRigGraphNode->HandleAddAggregateElement(ModelNode->GetNodePath());
+		}
+	}
+	return FReply::Handled();
+}
+
 bool SControlRigGraphNode::UseLowDetailNodeContent() const
 {
 	if(LastHighDetailSize.IsNearlyZero())
diff --git a/Engine/Plugins/Animation/ControlRig/Source/ControlRigEditor/Private/Graph/SControlRigGraphNode.h b/Engine/Plugins/Animation/ControlRig/Source/ControlRigEditor/Private/Graph/SControlRigGraphNode.h
--- a/Engine/Plugins/Animation/ControlRig/Source/ControlRigEditor/Private/Graph/SControlRigGraphNode.h
+++ b/Engine/Plugins/Animation/ControlRig/Source/ControlRigEditor/Private/Graph/SControlRigGraphNode.h
@@ -53,6 +53,12 @@ public:
 	virtual bool IsHidingPinWidgets() const override { return UseLowDetailNodeContent(); }
 	virtual bool UseLowDetailPinNames() const override;
 
+	virtual void CreateInputSideAddButton(TSharedPtr<SVerticalBox> InputBox) override;
+	virtual void CreateOutputSideAddButton(TSharedPtr<SVerticalBox> InputBox) override;
+
+	/** Callback function executed when Add pin button is clicked */
+	virtual FReply OnAddPin() override;
+
 private:
 
 	bool UseLowDetailNodeContent() const;
diff --git a/Engine/Source/Developer/RigVMDeveloper/Private/RigVMModel/Nodes/RigVMAggregateNode.cpp b/Engine/Source/Developer/RigVMDeveloper/Private/RigVMModel/Nodes/RigVMAggregateNode.cpp
--- a/Engine/Source/Developer/RigVMDeveloper/Private/RigVMModel/Nodes/RigVMAggregateNode.cpp
+++ b/Engine/Source/Developer/RigVMDeveloper/Private/RigVMModel/Nodes/RigVMAggregateNode.cpp
@@ -7,85 +7,180 @@
 URigVMAggregateNode::URigVMAggregateNode()
 	: Super()
 	, FirstInnerNodeCache(nullptr)
+	, LastInnerNodeCache(nullptr)
 {
 }
 
 bool URigVMAggregateNode::IsInputAggregate() const
 {
-	return GetFirstAggregatePin()->GetDirection() == ERigVMPinDirection::Input;
+	for (URigVMNode* Node : GetContainedNodes())
+	{
+		if (!Node->IsA<URigVMFunctionEntryNode>() && !Node->IsA<URigVMFunctionReturnNode>())
+		{
+			if (URigVMPin* FirstAggregatePin = Node->GetFirstAggregatePin())
+			{
+				return FirstAggregatePin->GetDirection() == ERigVMPinDirection::Input;
+			}
+		}
+	}
+	return false;
 }
 
 URigVMNode* URigVMAggregateNode::GetFirstInnerNode() const
 {
+#if UE_RIGVM_AGGREGATE_NODES_ENABLED
+
 	if (FirstInnerNodeCache == nullptr)
 	{
-		// Find any inner node (not entry or return)
-		URigVMNode* InnerNode = nullptr;
+		if (GetContainedNodes().Num() < 3)
+		{
+			return nullptr;
+		}
+		
+		if (IsInputAggregate())
 		{
-			for (URigVMNode* Node : GetContainedGraph()->GetNodes())
+			// Find node connected twice to the entry (through aggregate arguments)
+			FString Arg1Name = GetFirstAggregatePin()->GetName();
+			FString Arg2Name = GetSecondAggregatePin()->GetName();
+			URigVMFunctionEntryNode* EntryNode = GetEntryNode();
+			TArray<URigVMNode*> ConnectedNodes;
+			for (URigVMPin* EntryPin : EntryNode->GetPins())
 			{
-				if (!Node->IsA<URigVMFunctionEntryNode>() && !Node->IsA<URigVMFunctionReturnNode>())
+				TArray<URigVMPin*> TargetPins = EntryPin->GetLinkedTargetPins();
+				if (TargetPins.Num() > 0)
 				{
-					InnerNode = Node;
-					break;
+					if (TargetPins[0]->GetName() == Arg1Name || TargetPins[0]->GetName() == Arg2Name)
+					{
+						URigVMNode* TargetNode = TargetPins[0]->GetNode();
+						if (ConnectedNodes.Contains(TargetNode))
+						{
+							FirstInnerNodeCache = TargetNode;
+							return FirstInnerNodeCache;
+						}
+
+						ConnectedNodes.Add(TargetNode);
+					}
 				}
 			}
-
-			if (!InnerNode || !InnerNode->IsAggregate())
+		}
+		else
+		{
+			// Find node connected to entry throught the opposite aggregate argument 
+			FString ArgOppositeName = GetOppositeAggregatePin()->GetName();
+			URigVMFunctionEntryNode* EntryNode = GetEntryNode();
+			if (URigVMPin* EntryPin = EntryNode->FindPin(ArgOppositeName))
 			{
-				return nullptr;
+				TArray<URigVMPin*> TargetPins = EntryPin->GetLinkedTargetPins();
+				if (TargetPins.Num() > 0)
+				{
+					FirstInnerNodeCache = TargetPins[0]->GetNode();
+					return FirstInnerNodeCache;
+				}
 			}
 		}
+	}
+#endif
+	return FirstInnerNodeCache;
+}
+
+URigVMNode* URigVMAggregateNode::GetLastInnerNode() const
+{
+#if UE_RIGVM_AGGREGATE_NODES_ENABLED
 
-		// Find the aggregate pins
-		URigVMGraph* Graph = GetContainedGraph();
-		if (URigVMFunctionEntryNode* EntryNode = Graph->GetEntryNode())
+	if (LastInnerNodeCache == nullptr)
+	{
+		FString ArgOppositeName = GetOppositeAggregatePin()->GetName();
+		if (IsInputAggregate())
 		{
-			URigVMPin* Arg1Pin = nullptr;
-			const URigVMPin* Argument1 = InnerNode->GetFirstAggregatePin();
-			if (Argument1->GetDirection() == ERigVMPinDirection::Input)
-			{
-				Arg1Pin = EntryNode->FindPin(Argument1->GetName());
-			}
-			else
+			// Find node connected to return throught the opposite aggregate argument
+			URigVMFunctionReturnNode* ReturnNode = GetReturnNode();
+			if (URigVMPin* ReturnPin = ReturnNode->FindPin(ArgOppositeName))
 			{
-				const URigVMPin* OppositeArgument = InnerNode->GetOppositeAggregatePin();
-				Arg1Pin = EntryNode->FindPin(OppositeArgument->GetName());				
+				TArray<URigVMPin*> SourcePins = ReturnPin->GetLinkedSourcePins();
+				if(SourcePins.Num() > 0)
+				{
+					LastInnerNodeCache = SourcePins[0]->GetNode();
+					return LastInnerNodeCache;
+				}
 			}
-			if (Arg1Pin)
+		}
+		else
+		{
+			// Find node connected twice to the return (through aggregate arguments)
+			FString Arg1Name = GetFirstAggregatePin()->GetName();
+			FString Arg2Name = GetSecondAggregatePin()->GetName();
+			URigVMFunctionReturnNode* ReturnNode = GetReturnNode();
+			TArray<URigVMNode*> ConnectedNodes;
+			for (URigVMPin* ReturnPin : ReturnNode->GetPins())
 			{
-				TArray<URigVMPin*> TargetPins = Arg1Pin->GetLinkedTargetPins();
-				if (TargetPins.Num() > 0)
+				TArray<URigVMPin*> SourcePins = ReturnPin->GetLinkedSourcePins();
+				if (SourcePins.Num() > 0)
 				{
-					FirstInnerNodeCache = TargetPins[0]->GetNode(); 
+					if (SourcePins[0]->GetName() == Arg1Name || SourcePins[0]->GetName() == Arg2Name)
+					{
+						URigVMNode* SourceNode = SourcePins[0]->GetNode();
+						if (ConnectedNodes.Contains(SourceNode))
+						{
+							LastInnerNodeCache = SourceNode;
+							return LastInnerNodeCache;
+						}
+
+						ConnectedNodes.Add(SourceNode);
+					}
 				}
 			}
 		}
 	}
-	return FirstInnerNodeCache;
+#endif
+	return LastInnerNodeCache;
 }
 
 URigVMPin* URigVMAggregateNode::GetFirstAggregatePin() const
 {
-	const URigVMNode* FirstNode = GetFirstInnerNode();
-	return FirstNode->GetFirstAggregatePin();
+#if UE_RIGVM_AGGREGATE_NODES_ENABLED
+	for (URigVMNode* Node : GetContainedNodes())
+	{
+		if (!Node->IsA<URigVMFunctionEntryNode>() && !Node->IsA<URigVMFunctionReturnNode>())
+		{
+			return Node->GetFirstAggregatePin();
+		}
+	}
+#endif
+	return nullptr;
 }
 
 URigVMPin* URigVMAggregateNode::GetSecondAggregatePin() const
 {
-	const URigVMNode* FirstNode = GetFirstInnerNode();
-	return FirstNode->GetSecondAggregatePin();
+#if UE_RIGVM_AGGREGATE_NODES_ENABLED
+	for (URigVMNode* Node : GetContainedNodes())
+	{
+		if (!Node->IsA<URigVMFunctionEntryNode>() && !Node->IsA<URigVMFunctionReturnNode>())
+		{
+			return Node->GetSecondAggregatePin();
+		}
+	}
+#endif
+	return nullptr;
 }
 
 URigVMPin* URigVMAggregateNode::GetOppositeAggregatePin() const
 {
-	const URigVMNode* FirstNode = GetFirstInnerNode();
-	return FirstNode->GetOppositeAggregatePin();
+#if UE_RIGVM_AGGREGATE_NODES_ENABLED
+	for (URigVMNode* Node : GetContainedNodes())
+	{
+		if (!Node->IsA<URigVMFunctionEntryNode>() && !Node->IsA<URigVMFunctionReturnNode>())
+		{
+			return Node->GetOppositeAggregatePin();
+		}
+	}
+#endif
+	return nullptr;
 }
 
 void URigVMAggregateNode::InvalidateCache()
 {
 	FirstInnerNodeCache = nullptr;
+	LastInnerNodeCache = nullptr;
 }
 
 FString URigVMAggregateNode::GetNodeTitle() const
diff --git a/Engine/Source/Developer/RigVMDeveloper/Private/RigVMModel/Nodes/RigVMUnitNode.cpp b/Engine/Source/Developer/RigVMDeveloper/Private/RigVMModel/Nodes/RigVMUnitNode.cpp
--- a/Engine/Source/Developer/RigVMDeveloper/Private/RigVMModel/Nodes/RigVMUnitNode.cpp
+++ b/Engine/Source/Developer/RigVMDeveloper/Private/RigVMModel/Nodes/RigVMUnitNode.cpp
@@ -111,6 +111,7 @@ FString URigVMUnitNode::GetDeprecatedMetadata() const
 
 bool URigVMUnitNode::IsAggregate() const
 {
+#if UE_RIGVM_AGGREGATE_NODES_ENABLED
 	TArray<URigVMPin*> AggregateInputs = GetAggregateInputs();
 	TArray<URigVMPin*> AggregateOutputs = GetAggregateOutputs();
 
@@ -130,12 +131,14 @@ bool URigVMUnitNode::IsAggregate() const
 		
 		return true;
 	}
+#endif
 	
 	return false;
 }
 
 URigVMPin* URigVMUnitNode::GetFirstAggregatePin() const
 {
+#if UE_RIGVM_AGGREGATE_NODES_ENABLED
 	TArray<URigVMPin*> Inputs = GetAggregateInputs();
 	TArray<URigVMPin*> Outputs = GetAggregateOutputs();
 	if (Inputs.Num() == 2 && Outputs.Num() == 1)
@@ -146,11 +149,13 @@ URigVMPin* URigVMUnitNode::GetFirstAggregatePin() const
 	{
 		return Outputs[0];
 	}
+#endif
 	return nullptr;
 }
 
 URigVMPin* URigVMUnitNode::GetSecondAggregatePin() const
 {
+#if UE_RIGVM_AGGREGATE_NODES_ENABLED
 	TArray<URigVMPin*> Inputs = GetAggregateInputs();
 	TArray<URigVMPin*> Outputs = GetAggregateOutputs();
 	if (Inputs.Num() == 2 && Outputs.Num() == 1)
@@ -161,11 +166,13 @@ URigVMPin* URigVMUnitNode::GetSecondAggregatePin() const
 	{
 		return Outputs[1];
 	}
+#endif
 	return nullptr;
 }
 
 URigVMPin* URigVMUnitNode::GetOppositeAggregatePin() const
 {
+#if UE_RIGVM_AGGREGATE_NODES_ENABLED
 	TArray<URigVMPin*> Inputs = GetAggregateInputs();
 	TArray<URigVMPin*> Outputs = GetAggregateOutputs();
 	if (Inputs.Num() == 2 && Outputs.Num() == 1)
@@ -176,6 +183,7 @@ URigVMPin* URigVMUnitNode::GetOppositeAggregatePin() const
 	{
 		return Inputs[0];
 	}
+#endif
 	return nullptr;
 }
 
@@ -187,6 +195,7 @@ bool URigVMUnitNode::IsInputAggregate() const
 TArray<URigVMPin*> URigVMUnitNode::GetAggregateInputs() const
 {
 	TArray<URigVMPin*> AggregateInputs;
+#if UE_RIGVM_AGGREGATE_NODES_ENABLED
 	if (UScriptStruct* Struct = GetScriptStruct())
 	{
 		for (URigVMPin* Pin : GetPins())
@@ -203,12 +212,14 @@ TArray<URigVMPin*> URigVMUnitNode::GetAggregateInputs() const
 			}
 		}
 	}
+#endif
 	return AggregateInputs;
 }
 
 TArray<URigVMPin*> URigVMUnitNode::GetAggregateOutputs() const
 {
 	TArray<URigVMPin*> AggregateOutputs;
+#if UE_RIGVM_AGGREGATE_NODES_ENABLED	
 	if (UScriptStruct* Struct = GetScriptStruct())
 	{
 		for (URigVMPin* Pin : GetPins())
@@ -225,9 +236,25 @@ TArray<URigVMPin*> URigVMUnitNode::GetAggregateOutputs() const
 			}
 		}
 	}
+#endif
 	return AggregateOutputs;
 }
 
+FName URigVMUnitNode::GetNextAggregateName(const FName& InLastAggregatePinName) const
+{
+#if UE_RIGVM_AGGREGATE_NODES_ENABLED	
+	if(UScriptStruct* Struct = GetScriptStruct())
+	{
+		check(Struct->IsChildOf(FRigVMStruct::StaticStruct()));
+
+		const TSharedPtr<FStructOnScope> StructOnScope = ConstructStructInstance();
+		const FRigVMStruct* StructMemory = (const FRigVMStruct*)StructOnScope->GetStructMemory();
+		return StructMemory->GetNextAggregateName(InLastAggregatePinName);
+	}
+#endif
+	return FName();
+}
+
 UScriptStruct* URigVMUnitNode::GetScriptStruct() const
 {
 	return ScriptStruct;
diff --git a/Engine/Source/Developer/RigVMDeveloper/Private/RigVMModel/RigVMController.cpp b/Engine/Source/Developer/RigVMDeveloper/Private/RigVMModel/RigVMController.cpp
--- a/Engine/Source/Developer/RigVMDeveloper/Private/RigVMModel/RigVMController.cpp
+++ b/Engine/Source/Developer/RigVMDeveloper/Private/RigVMModel/RigVMController.cpp
@@ -66,6 +66,7 @@ URigVMController::URigVMController()
 	, bReportWarningsAndErrors(true)
 	, bIgnoreRerouteCompactnessChanges(false)
 	, UserLinkDirection(ERigVMPinDirection::Invalid)
+	, bIsTransacting(false)
 	, bIsRunningUnitTest(false)
 	, bIsFullyResolvingTemplateNode(false)
 {
@@ -79,6 +80,7 @@ URigVMController::URigVMController(const FObjectInitializer& ObjectInitializer)
 	, bReportWarningsAndErrors(true)
 	, bIgnoreRerouteCompactnessChanges(false)
 	, UserLinkDirection(ERigVMPinDirection::Invalid)
+	, bIsTransacting(false)
 	, bIsRunningUnitTest(false)
 	, bIsFullyResolvingTemplateNode(false)
 {
@@ -656,6 +658,11 @@ URigVMUnitNode* URigVMController::AddUnitNode(UScriptStruct* InScriptStruct, con
 		return nullptr;
 	}
 
+	if (!bIsTransacting && !IsGraphEditable())
+	{
+		return nullptr;
+	}
+
 	if (GetGraph()->IsA<URigVMFunctionLibrary>())
 	{
 		ReportError(TEXT("Cannot add unit nodes to function library graphs."));
@@ -813,6 +820,11 @@ URigVMUnitNode* URigVMController::AddUnitNodeFromStructPath(const FString& InScr
 		return nullptr;
 	}
 
+	if (!bIsTransacting && !IsGraphEditable())
+	{
+		return nullptr;
+	}
+
 	UScriptStruct* ScriptStruct = URigVMPin::FindObjectFromCPPTypeObjectPath<UScriptStruct>(InScriptStructPath);
 	if (ScriptStruct == nullptr)
 	{
@@ -830,6 +842,11 @@ URigVMVariableNode* URigVMController::AddVariableNode(const FName& InVariableNam
 		return nullptr;
 	}
 
+	if (!bIsTransacting && !IsGraphEditable())
+	{
+		return nullptr;
+	}
+
 	URigVMGraph* Graph = GetGraph();
 	check(Graph);
 
@@ -1022,6 +1039,11 @@ URigVMVariableNode* URigVMController::AddVariableNodeFromObjectPath(const FName&
 		return nullptr;
 	}
 
+	if (!bIsTransacting && !IsGraphEditable())
+	{
+		return nullptr;
+	}
+
 	UObject* CPPTypeObject = nullptr;
 	if (!InCPPTypeObjectPath.IsEmpty())
 	{
@@ -1043,6 +1065,11 @@ void URigVMController::RefreshVariableNode(const FName& InNodeName, const FName&
 		return;
 	}
 
+	if (!bIsTransacting && !IsGraphEditable())
+	{
+		return;
+	}
+
 	URigVMGraph* Graph = GetGraph();
 	check(Graph);
 
@@ -1109,6 +1136,11 @@ void URigVMController::OnExternalVariableRemoved(const FName& InVarName, bool bS
 		return;
 	}
 
+	if (!bIsTransacting && !IsGraphEditable())
+	{
+		return;
+	}
+
 	if (!InVarName.IsValid())
 	{
 		return;
@@ -1150,6 +1182,7 @@ void URigVMController::OnExternalVariableRemoved(const FName& InVarName, bool bS
 		else if(URigVMCollapseNode* CollapseNode = Cast<URigVMCollapseNode>(Node))
 		{
 			FRigVMControllerGraphGuard GraphGuard(this, CollapseNode->GetContainedGraph(), bSetupUndoRedo);
+			TGuardValue<bool> GuardEditGraph(CollapseNode->ContainedGraph->bEditable, true);
 
 			// call this function for the contained graph recursively 
 			OnExternalVariableRemoved(InVarName, bSetupUndoRedo);
@@ -1196,6 +1229,11 @@ bool URigVMController::OnExternalVariableRenamed(const FName& InOldVarName, cons
 		return false;
 	}
 
+	if (!bIsTransacting && !IsGraphEditable())
+	{
+		return false;
+	}
+
 	if (!InOldVarName.IsValid() || !InNewVarName.IsValid())
 	{
 		return false;
@@ -1237,6 +1275,7 @@ bool URigVMController::OnExternalVariableRenamed(const FName& InOldVarName, cons
 		else if(URigVMCollapseNode* CollapseNode = Cast<URigVMCollapseNode>(Node))
 		{
 			FRigVMControllerGraphGuard GraphGuard(this, CollapseNode->GetContainedGraph(), bSetupUndoRedo);
+			TGuardValue<bool> GuardEditGraph(CollapseNode->ContainedGraph->bEditable, true);
 			OnExternalVariableRenamed(InOldVarName, InNewVarName, bSetupUndoRedo);
 
 			// if we are a function we need to notify all references!
@@ -1286,6 +1325,11 @@ void URigVMController::OnExternalVariableTypeChanged(const FName& InVarName, con
 		return;
 	}
 
+	if (!bIsTransacting && !IsGraphEditable())
+	{
+		return;
+	}
+
 	if (!InVarName.IsValid())
 	{
 		return;
@@ -1327,6 +1371,7 @@ void URigVMController::OnExternalVariableTypeChanged(const FName& InVarName, con
 		else if(URigVMCollapseNode* CollapseNode = Cast<URigVMCollapseNode>(Node))
 		{
 			FRigVMControllerGraphGuard GraphGuard(this, CollapseNode->GetContainedGraph(), bSetupUndoRedo);
+			TGuardValue<bool> GuardEditGraph(CollapseNode->ContainedGraph->bEditable, true);
 			OnExternalVariableTypeChanged(InVarName, InCPPType, InCPPTypeObject, bSetupUndoRedo);
 
 			// if we are a function we need to notify all references!
@@ -1383,6 +1428,11 @@ void URigVMController::OnExternalVariableTypeChangedFromObjectPath(const FName&
 		return;
 	}
 
+	if (!bIsTransacting && !IsGraphEditable())
+	{
+		return;
+	}
+
 	UObject* CPPTypeObject = nullptr;
 	if (!InCPPTypeObjectPath.IsEmpty())
 	{
@@ -1404,6 +1454,11 @@ URigVMVariableNode* URigVMController::ReplaceParameterNodeWithVariable(const FNa
 		return nullptr;
 	}
 
+	if (!bIsTransacting && !IsGraphEditable())
+	{
+		return nullptr;
+	}
+
 	URigVMGraph* Graph = GetGraph();
 	check(Graph);
 
@@ -1451,6 +1506,11 @@ URigVMTemplateNode* URigVMController::ReplaceUnitNodeWithTemplateNode(const FNam
 		return nullptr;
 	}
 
+	if (!bIsTransacting && !IsGraphEditable())
+	{
+		return nullptr;
+	}
+
 	URigVMUnitNode* UnitNode = Cast<URigVMUnitNode>(GetGraph()->FindNodeByName(InNodeName));
 	if(UnitNode == nullptr)
 	{
@@ -1522,6 +1582,11 @@ bool URigVMController::UnresolveTemplateNodes(const TArray<FName>& InNodeNames,
 		return false;
 	}
 
+	if (!bIsTransacting && !IsGraphEditable())
+	{
+		return false;
+	}
+
 	TArray<URigVMTemplateNode*> Nodes;
 	for (const FName& NodeName : InNodeNames)
 	{
@@ -1563,6 +1628,12 @@ bool URigVMController::UnresolveTemplateNodes(const TArray<URigVMTemplateNode*>&
 		return false;
 	}
 
+
+	if (!bIsTransacting && !IsGraphEditable())
+	{
+		return false;
+	}
+	
 	// check if any of the nodes needs to be unresolved
 	const bool bHasNodeToResolve = InNodes.ContainsByPredicate( [](const URigVMTemplateNode* Node) -> bool
 	{
@@ -1674,6 +1745,11 @@ TArray<URigVMNode*> URigVMController::UpgradeNodes(const TArray<FName>& InNodeNa
 		return Nodes;
 	}
 
+	if (!bIsTransacting && !IsGraphEditable())
+	{
+		return Nodes;
+	}
+
 	for (const FName& NodeName : InNodeNames)
 	{
 		if (URigVMNode* Node = GetGraph()->FindNodeByName(NodeName))
@@ -1723,6 +1799,11 @@ TArray<URigVMNode*> URigVMController::UpgradeNodes(const TArray<URigVMNode*>& In
 		return TArray<URigVMNode*>();
 	}
 
+	if (!bIsTransacting && !IsGraphEditable())
+	{
+		return TArray<URigVMNode*>();
+	}
+
 	bool bFoundAnyNodeToUpgrade = false;
 	for(URigVMNode* Node : InNodes)
 	{
@@ -1929,6 +2010,11 @@ URigVMParameterNode* URigVMController::AddParameterNodeFromObjectPath(const FNam
 		return nullptr;
 	}
 
+	if (!bIsTransacting && !IsGraphEditable())
+	{
+		return nullptr;
+	}
+
 	UObject* CPPTypeObject = nullptr;
 	if (!InCPPTypeObjectPath.IsEmpty())
 	{
@@ -1950,6 +2036,11 @@ URigVMCommentNode* URigVMController::AddCommentNode(const FString& InCommentText
 		return nullptr;
 	}
 
+	if (!bIsTransacting && !IsGraphEditable())
+	{
+		return nullptr;
+	}
+
 	URigVMGraph* Graph = GetGraph();
 	check(Graph);
 
@@ -2008,6 +2099,11 @@ URigVMRerouteNode* URigVMController::AddRerouteNodeOnLink(URigVMLink* InLink, bo
 		return nullptr;
 	}
 
+	if (!bIsTransacting && !IsGraphEditable())
+	{
+		return nullptr;
+	}
+
 	if (GetGraph()->IsA<URigVMFunctionLibrary>())
 	{
 		ReportError(TEXT("Cannot add reroutes to function library graphs."));
@@ -2070,6 +2166,11 @@ URigVMRerouteNode* URigVMController::AddRerouteNodeOnLinkPath(const FString& InL
 		return nullptr;
 	}
 
+	if (!bIsTransacting && !IsGraphEditable())
+	{
+		return nullptr;
+	}
+
 	URigVMGraph* Graph = GetGraph();
 	check(Graph);
 
@@ -2084,6 +2185,11 @@ URigVMRerouteNode* URigVMController::AddRerouteNodeOnPin(const FString& InPinPat
 		return nullptr;
 	}
 
+	if (!bIsTransacting && !IsGraphEditable())
+	{
+		return nullptr;
+	}
+
 	URigVMGraph* Graph = GetGraph();
 	check(Graph);
 
@@ -2194,6 +2300,11 @@ URigVMInjectionInfo* URigVMController::AddInjectedNode(const FString& InPinPath,
 		return nullptr;
 	}
 
+	if (!bIsTransacting && !IsGraphEditable())
+	{
+		return nullptr;
+	}
+
 	URigVMGraph* Graph = GetGraph();
 	check(Graph);
 
@@ -2408,6 +2519,11 @@ URigVMInjectionInfo* URigVMController::AddInjectedNodeFromStructPath(const FStri
 		return nullptr;
 	}
 
+	if (!bIsTransacting && !IsGraphEditable())
+	{
+		return nullptr;
+	}
+
 	UScriptStruct* ScriptStruct = URigVMPin::FindObjectFromCPPTypeObjectPath<UScriptStruct>(InScriptStructPath);
 	if (ScriptStruct == nullptr)
 	{
@@ -2425,6 +2541,11 @@ bool URigVMController::RemoveInjectedNode(const FString& InPinPath, bool bAsInpu
 		return false;
 	}
 
+	if (!bIsTransacting && !IsGraphEditable())
+	{
+		return false;
+	}
+
 	URigVMGraph* Graph = GetGraph();
 	check(Graph);
 
@@ -2525,6 +2646,11 @@ URigVMInjectionInfo* URigVMController::InjectNodeIntoPin(const FString& InPinPat
 		return nullptr;
 	}
 
+	if (!bIsTransacting && !IsGraphEditable())
+	{
+		return nullptr;
+	}
+
 	URigVMGraph* Graph = GetGraph();
 	check(Graph);
 
@@ -2544,6 +2670,11 @@ URigVMInjectionInfo* URigVMController::InjectNodeIntoPin(URigVMPin* InPin, bool
 		return nullptr;
 	}
 
+	if (!bIsTransacting && !IsGraphEditable())
+	{
+		return nullptr;
+	}
+	
 	URigVMGraph* Graph = GetGraph();
 	check(Graph);
 
@@ -2658,6 +2789,11 @@ URigVMNode* URigVMController::EjectNodeFromPin(const FString& InPinPath, bool bS
 		return nullptr;
 	}
 
+	if (!bIsTransacting && !IsGraphEditable())
+	{
+		return nullptr;
+	}
+
 	URigVMGraph* Graph = GetGraph();
 	check(Graph);
 
@@ -2677,6 +2813,11 @@ URigVMNode* URigVMController::EjectNodeFromPin(URigVMPin* InPin, bool bSetupUndo
 		return nullptr;
 	}
 
+	if (!bIsTransacting && !IsGraphEditable())
+	{
+		return nullptr;
+	}
+
 	URigVMGraph* Graph = GetGraph();
 	check(Graph);
 
@@ -2980,6 +3121,11 @@ bool URigVMController::CanImportNodesFromText(const FString& InText)
 		return false;
 	}
 
+	if (!bIsTransacting && !IsGraphEditable())
+	{
+		return false;
+	}
+
 	if (GetGraph()->IsA<URigVMFunctionLibrary>())
 	{
 		return false;
@@ -2997,6 +3143,11 @@ TArray<FName> URigVMController::ImportNodesFromText(const FString& InText, bool
 		return NodeNames;
 	}
 
+	if (!bIsTransacting && !IsGraphEditable())
+	{
+		return NodeNames;
+	}
+
 	URigVMGraph* Graph = GetGraph();
 	check(Graph);
 
@@ -3020,6 +3171,15 @@ TArray<FName> URigVMController::ImportNodesFromText(const FString& InText, bool
 		ActionStack->BeginAction(AddNodesAction);
 	}
 
+	TArray<TGuardValue<bool>> EditGuards;
+	for (URigVMNode* CreatedNode : Factory.CreatedNodes)
+	{
+		if (URigVMLibraryNode* LibraryNode = Cast<URigVMLibraryNode>(CreatedNode))
+		{
+			EditGuards.Emplace(LibraryNode->GetContainedGraph()->bEditable, true);
+		}
+	}
+
 	FRigVMUnitNodeCreatedContext::FScope UnitNodeCreatedScope(UnitNodeCreatedContext, ERigVMNodeCreatedReason::Paste);
 	for (URigVMNode* CreatedNode : Factory.CreatedNodes)
 	{
@@ -3057,6 +3217,7 @@ TArray<FName> URigVMController::ImportNodesFromText(const FString& InText, bool
 			{
 				{
 					FRigVMControllerGraphGuard GraphGuard(this, CollapseNode->GetContainedGraph(), false);
+					TGuardValue<bool> GuardEditGraph(CollapseNode->ContainedGraph->bEditable, true);
 					ReattachLinksToPinObjects();
 				}
 				
@@ -3250,6 +3411,11 @@ URigVMLibraryNode* URigVMController::LocalizeFunction(
 		return nullptr;
 	}
 
+	if (!bIsTransacting && !IsGraphEditable())
+	{
+		return nullptr;
+	}
+
 	if(InFunctionDefinition == nullptr)
 	{
 		return nullptr;
@@ -3280,6 +3446,11 @@ TMap<URigVMLibraryNode*, URigVMLibraryNode*> URigVMController::LocalizeFunctions
 		return LocalizedFunctions;
 	}
 
+	if (!bIsTransacting && !IsGraphEditable())
+	{
+		return LocalizedFunctions;
+	}
+
 	URigVMGraph* Graph = GetGraph();
 	check(Graph);
 
@@ -3495,6 +3666,11 @@ URigVMCollapseNode* URigVMController::CollapseNodes(const TArray<FName>& InNodeN
 		return nullptr;
 	}
 
+	if (!bIsTransacting && !IsGraphEditable())
+	{
+		return nullptr;
+	}
+
 	URigVMGraph* Graph = GetGraph();
 	check(Graph);
 
@@ -3543,6 +3719,11 @@ TArray<URigVMNode*> URigVMController::ExpandLibraryNode(const FName& InNodeName,
 		return TArray<URigVMNode*>();
 	}
 
+	if (!bIsTransacting && !IsGraphEditable())
+	{
+		return TArray<URigVMNode*>();
+	}
+
 	URigVMGraph* Graph = GetGraph();
 	check(Graph);
 
@@ -3583,6 +3764,11 @@ URigVMCollapseNode* URigVMController::CollapseNodes(const TArray<URigVMNode*>& I
 	URigVMGraph* Graph = GetGraph();
 	check(Graph);
 
+	if (!bIsTransacting && !IsGraphEditable())
+	{
+		return nullptr;
+	}
+
 	if (Graph->IsA<URigVMFunctionLibrary>())
 	{
 		ReportError(TEXT("Cannot collapse nodes in function library graphs."));
@@ -3595,7 +3781,7 @@ URigVMCollapseNode* URigVMController::CollapseNodes(const TArray<URigVMNode*>& I
 		return nullptr;
 	}
 
-#if UE_RIGVM_ENABLE_AGGREGATE_NODES
+#if UE_RIGVM_AGGREGATE_NODES_ENABLED
 	if (bIsAggregate)
 	{
 		if (URigVMUnitNode* UnitNode = Cast<URigVMUnitNode>(InNodes[0]))
@@ -3845,7 +4031,7 @@ URigVMCollapseNode* URigVMController::CollapseNodes(const TArray<URigVMNode*>& I
 		ActionStack->BeginAction(CollapseAction);
 	}
 
-#if UE_RIGVM_ENABLE_AGGREGATE_NODES
+#if UE_RIGVM_AGGREGATE_NODES_ENABLED
 	URigVMCollapseNode* CollapseNode = nullptr;
 	if (bIsAggregate)
 	{
@@ -3859,6 +4045,15 @@ URigVMCollapseNode* URigVMController::CollapseNodes(const TArray<URigVMNode*>& I
 	URigVMCollapseNode* CollapseNode = NewObject<URigVMCollapseNode>(Graph, *CollapseNodeName);
 #endif
 	CollapseNode->ContainedGraph = NewObject<URigVMGraph>(CollapseNode, TEXT("ContainedGraph"));
+
+#if UE_RIGVM_AGGREGATE_NODES_ENABLED
+	if (bIsAggregate)
+	{
+		CollapseNode->ContainedGraph->bEditable = false;
+	}
+	TGuardValue<bool> GuardEditGraph(CollapseNode->ContainedGraph->bEditable, true);
+#endif
+	
 	CollapseNode->Position = Center;
 	Graph->Nodes.Add(CollapseNode);
 
@@ -4219,6 +4414,11 @@ URigVMCollapseNode* URigVMController::CollapseNodes(const TArray<URigVMNode*>& I
 
 TArray<URigVMNode*> URigVMController::ExpandLibraryNode(URigVMLibraryNode* InNode, bool bSetupUndoRedo)
 {
+	if (!bIsTransacting && !IsGraphEditable())
+	{
+		return TArray<URigVMNode*>();
+	}
+	
 	if (!IsValidNodeForGraph(InNode))
 	{
 		return TArray<URigVMNode*>();
@@ -4924,6 +5124,11 @@ FName URigVMController::PromoteCollapseNodeToFunctionReferenceNode(const FName&
 		return NAME_None;
 	}
 
+	if (!bIsTransacting && !IsGraphEditable())
+	{
+		return NAME_None;
+	}
+
 	URigVMGraph* Graph = GetGraph();
 	check(Graph);
 
@@ -4952,6 +5157,11 @@ FName URigVMController::PromoteFunctionReferenceNodeToCollapseNode(const FName&
 		return NAME_None;
 	}
 
+	if (!bIsTransacting && !IsGraphEditable())
+	{
+		return NAME_None;
+	}
+
 	URigVMGraph* Graph = GetGraph();
 	check(Graph);
 
@@ -4965,6 +5175,11 @@ FName URigVMController::PromoteFunctionReferenceNodeToCollapseNode(const FName&
 
 URigVMFunctionReferenceNode* URigVMController::PromoteCollapseNodeToFunctionReferenceNode(URigVMCollapseNode* InCollapseNode, bool bSetupUndoRedo, const FString& InExistingFunctionDefinitionPath)
 {
+	if (!bIsTransacting && !IsGraphEditable())
+	{
+		return nullptr;
+	}
+	
 	if (!IsValidNodeForGraph(InCollapseNode))
 	{
 		return nullptr;
@@ -5092,6 +5307,11 @@ URigVMFunctionReferenceNode* URigVMController::PromoteCollapseNodeToFunctionRefe
 
 URigVMCollapseNode* URigVMController::PromoteFunctionReferenceNodeToCollapseNode(URigVMFunctionReferenceNode* InFunctionRefNode, bool bSetupUndoRedo, bool bRemoveFunctionDefinition)
 {
+	if (!bIsTransacting && !IsGraphEditable())
+	{
+		return nullptr;
+	}
+	
 	if (!IsValidNodeForGraph(InFunctionRefNode))
 	{
 		return nullptr;
@@ -5228,6 +5448,11 @@ void URigVMController::SetReferencedFunction(URigVMFunctionReferenceNode* InFunc
 	{
 		return;
 	}
+
+	if (!bIsTransacting && !IsGraphEditable())
+	{
+		return;
+	}
 	
 	URigVMLibraryNode* OldReferencedNode = InFunctionRefNode->GetReferencedNode();
 	if(OldReferencedNode != InNewReferencedNode)
@@ -5291,6 +5516,11 @@ void URigVMController::ReportRemovedLink(const FString& InSourcePinPath, const F
 
 bool URigVMController::RemoveNode(URigVMNode* InNode, bool bSetupUndoRedo, bool bRecursive, bool bPrintPythonCommand, bool bRelinkPins)
 {
+	if (!bIsTransacting && !IsGraphEditable())
+	{
+		return false;
+	}
+	
 	if (!IsValidNodeForGraph(InNode))
 	{
 		return false;
@@ -5583,6 +5813,11 @@ bool URigVMController::RemoveNodeByName(const FName& InNodeName, bool bSetupUndo
 		return false;
 	}
 
+	if (!bIsTransacting && !IsGraphEditable())
+	{
+		return false;
+	}
+
 	URigVMGraph* Graph = GetGraph();
 	check(Graph);
 
@@ -5591,6 +5826,11 @@ bool URigVMController::RemoveNodeByName(const FName& InNodeName, bool bSetupUndo
 
 bool URigVMController::RenameNode(URigVMNode* InNode, const FName& InNewName, bool bSetupUndoRedo, bool bPrintPythonCommand)
 {
+	if (!bIsTransacting && !IsGraphEditable())
+	{
+		return false;
+	}
+	
 	if (!IsValidNodeForGraph(InNode))
 	{
 		return false;
@@ -5664,6 +5904,11 @@ bool URigVMController::RenameNode(URigVMNode* InNode, const FName& InNewName, bo
 
 bool URigVMController::SelectNode(URigVMNode* InNode, bool bSelect, bool bSetupUndoRedo, bool bPrintPythonCommand)
 {
+	if (!bIsTransacting && !IsGraphEditable())
+	{
+		return false;
+	}
+	
 	if (!IsValidNodeForGraph(InNode))
 	{
 		return false;
@@ -5697,6 +5942,11 @@ bool URigVMController::SelectNodeByName(const FName& InNodeName, bool bSelect, b
 		return false;
 	}
 
+	if (!bIsTransacting && !IsGraphEditable())
+	{
+		return false;
+	}
+
 	URigVMGraph* Graph = GetGraph();
 	check(Graph);
 
@@ -5710,6 +5960,11 @@ bool URigVMController::ClearNodeSelection(bool bSetupUndoRedo, bool bPrintPython
 		return false;
 	}
 
+	if (!bIsTransacting && !IsGraphEditable())
+	{
+		return false;
+	}
+
 	return SetNodeSelection(TArray<FName>(), bSetupUndoRedo, bPrintPythonCommand);
 }
 
@@ -5720,6 +5975,11 @@ bool URigVMController::SetNodeSelection(const TArray<FName>& InNodeNames, bool b
 		return false;
 	}
 
+	if (!bIsTransacting && !IsGraphEditable())
+	{
+		return false;
+	}
+
 	URigVMGraph* Graph = GetGraph();
 	check(Graph);
 
@@ -5818,6 +6078,11 @@ bool URigVMController::SetNodeSelection(const TArray<FName>& InNodeNames, bool b
 
 bool URigVMController::SetNodePosition(URigVMNode* InNode, const FVector2D& InPosition, bool bSetupUndoRedo, bool bMergeUndoAction, bool bPrintPythonCommand)
 {
+	if (!bIsTransacting && !IsGraphEditable())
+	{
+		return false;
+	}
+	
 	if (!IsValidNodeForGraph(InNode))
 	{
 		return false;
@@ -5867,6 +6132,11 @@ bool URigVMController::SetNodePositionByName(const FName& InNodeName, const FVec
 		return false;
 	}
 
+	if (!bIsTransacting && !IsGraphEditable())
+	{
+		return false;
+	}
+
 	URigVMGraph* Graph = GetGraph();
 	check(Graph);
 
@@ -5876,6 +6146,11 @@ bool URigVMController::SetNodePositionByName(const FName& InNodeName, const FVec
 
 bool URigVMController::SetNodeSize(URigVMNode* InNode, const FVector2D& InSize, bool bSetupUndoRedo, bool bMergeUndoAction, bool bPrintPythonCommand)
 {
+	if (!bIsTransacting && !IsGraphEditable())
+	{
+		return false;
+	}
+	
 	if (!IsValidNodeForGraph(InNode))
 	{
 		return false;
@@ -5924,6 +6199,11 @@ bool URigVMController::SetNodeSizeByName(const FName& InNodeName, const FVector2
 		return false;
 	}
 
+	if (!bIsTransacting && !IsGraphEditable())
+	{
+		return false;
+	}
+
 	URigVMGraph* Graph = GetGraph();
 	check(Graph);
 
@@ -5933,6 +6213,11 @@ bool URigVMController::SetNodeSizeByName(const FName& InNodeName, const FVector2
 
 bool URigVMController::SetNodeColor(URigVMNode* InNode, const FLinearColor& InColor, bool bSetupUndoRedo, bool bMergeUndoAction, bool bPrintPythonCommand)
 {
+	if (!bIsTransacting && !IsGraphEditable())
+	{
+		return false;
+	}
+
 	if (!IsValidNodeForGraph(InNode))
 	{
 		return false;
@@ -5993,6 +6278,11 @@ bool URigVMController::SetNodeColorByName(const FName& InNodeName, const FLinear
 		return false;
 	}
 
+	if (!bIsTransacting && !IsGraphEditable())
+	{
+		return false;
+	}
+
 	URigVMGraph* Graph = GetGraph();
 	check(Graph);
 
@@ -6002,6 +6292,11 @@ bool URigVMController::SetNodeColorByName(const FName& InNodeName, const FLinear
 
 bool URigVMController::SetNodeCategory(URigVMCollapseNode* InNode, const FString& InCategory, bool bSetupUndoRedo, bool bMergeUndoAction, bool bPrintPythonCommand)
 {
+	if (!bIsTransacting && !IsGraphEditable())
+	{
+		return false;
+	}
+	
 	if (!IsValidNodeForGraph(InNode))
 	{
 		return false;
@@ -6050,6 +6345,11 @@ bool URigVMController::SetNodeCategoryByName(const FName& InNodeName, const FStr
 		return false;
 	}
 
+	if (!bIsTransacting && !IsGraphEditable())
+	{
+		return false;
+	}
+
 	URigVMGraph* Graph = GetGraph();
 	check(Graph);
 
@@ -6059,6 +6359,11 @@ bool URigVMController::SetNodeCategoryByName(const FName& InNodeName, const FStr
 
 bool URigVMController::SetNodeKeywords(URigVMCollapseNode* InNode, const FString& InKeywords, bool bSetupUndoRedo, bool bMergeUndoAction, bool bPrintPythonCommand)
 {
+	if (!bIsTransacting && !IsGraphEditable())
+	{
+		return false;
+	}
+	
 	if (!IsValidNodeForGraph(InNode))
 	{
 		return false;
@@ -6107,6 +6412,11 @@ bool URigVMController::SetNodeKeywordsByName(const FName& InNodeName, const FStr
 		return false;
 	}
 
+	if (!bIsTransacting && !IsGraphEditable())
+	{
+		return false;
+	}
+
 	URigVMGraph* Graph = GetGraph();
 	check(Graph);
 
@@ -6116,6 +6426,11 @@ bool URigVMController::SetNodeKeywordsByName(const FName& InNodeName, const FStr
 
 bool URigVMController::SetNodeDescription(URigVMCollapseNode* InNode, const FString& InDescription, bool bSetupUndoRedo, bool bMergeUndoAction, bool bPrintPythonCommand)
 {
+	if (!bIsTransacting && !IsGraphEditable())
+	{
+		return false;
+	}
+	
 	if (!IsValidNodeForGraph(InNode))
 	{
 		return false;
@@ -6164,6 +6479,11 @@ bool URigVMController::SetNodeDescriptionByName(const FName& InNodeName, const F
 		return false;
 	}
 
+	if (!bIsTransacting && !IsGraphEditable())
+	{
+		return false;
+	}
+
 	URigVMGraph* Graph = GetGraph();
 	check(Graph);
 
@@ -6173,6 +6493,11 @@ bool URigVMController::SetNodeDescriptionByName(const FName& InNodeName, const F
 
 bool URigVMController::SetCommentText(URigVMNode* InNode, const FString& InCommentText, const int32& InCommentFontSize, const bool& bInCommentBubbleVisible, const bool& bInCommentColorBubble, bool bSetupUndoRedo, bool bPrintPythonCommand)
 {
+	if (!bIsTransacting && !IsGraphEditable())
+	{
+		return false;
+	}
+	
 	if (!IsValidNodeForGraph(InNode))
 	{
 		return false;
@@ -6229,6 +6554,11 @@ bool URigVMController::SetCommentTextByName(const FName& InNodeName, const FStri
 		return false;
 	}
 
+	if (!bIsTransacting && !IsGraphEditable())
+	{
+		return false;
+	}
+
 	URigVMGraph* Graph = GetGraph();
 	check(Graph);
 
@@ -6238,6 +6568,11 @@ bool URigVMController::SetCommentTextByName(const FName& InNodeName, const FStri
 
 bool URigVMController::SetRerouteCompactness(URigVMNode* InNode, bool bShowAsFullNode, bool bSetupUndoRedo, bool bPrintPythonCommand)
 {
+	if (!bIsTransacting && !IsGraphEditable())
+	{
+		return false;
+	}
+	
 	if (!IsValidNodeForGraph(InNode))
 	{
 		return false;
@@ -6279,6 +6614,11 @@ bool URigVMController::SetRerouteCompactnessByName(const FName& InNodeName, bool
 		return false;
 	}
 
+	if (!bIsTransacting && !IsGraphEditable())
+	{
+		return false;
+	}
+
 	URigVMGraph* Graph = GetGraph();
 	check(Graph);
 
@@ -6293,6 +6633,11 @@ bool URigVMController::RenameVariable(const FName& InOldName, const FName& InNew
 		return false;
 	}
 
+	if (!bIsTransacting && !IsGraphEditable())
+	{
+		return false;
+	}
+
 	if (InOldName == InNewName)
 	{
 		ReportWarning(TEXT("RenameVariable: InOldName and InNewName are equal."));
@@ -6410,6 +6755,11 @@ bool URigVMController::SetPinExpansion(const FString& InPinPath, bool bIsExpande
 		return false;
 	}
 
+	if (!bIsTransacting && !IsGraphEditable())
+	{
+		return false;
+	}
+
 	URigVMGraph* Graph = GetGraph();
 	check(Graph);
 
@@ -6437,6 +6787,11 @@ bool URigVMController::SetPinExpansion(const FString& InPinPath, bool bIsExpande
 
 bool URigVMController::SetPinExpansion(URigVMPin* InPin, bool bIsExpanded, bool bSetupUndoRedo)
 {
+	if (!bIsTransacting && !IsGraphEditable())
+	{
+		return false;
+	}
+	
 	// If there is nothing to do, just return success
 	if (InPin->GetSubPins().Num() == 0 || InPin->IsExpanded() == bIsExpanded)
 	{
@@ -6477,6 +6832,11 @@ bool URigVMController::SetPinIsWatched(const FString& InPinPath, bool bIsWatched
 		return false;
 	}
 
+	if (!bIsTransacting && !IsGraphEditable())
+	{
+		return false;
+	}
+
 	URigVMGraph* Graph = GetGraph();
 	check(Graph);
 
@@ -6492,6 +6852,11 @@ bool URigVMController::SetPinIsWatched(const FString& InPinPath, bool bIsWatched
 
 bool URigVMController::SetPinIsWatched(URigVMPin* InPin, bool bIsWatched, bool bSetupUndoRedo)
 {
+	if (!bIsTransacting && !IsGraphEditable())
+	{
+		return false;
+	}
+	
 	if (!IsValidPinForGraph(InPin))
 	{
 		return false;
@@ -6569,6 +6934,11 @@ bool URigVMController::SetPinDefaultValue(const FString& InPinPath, const FStrin
 		return false;
 	}
 
+	if (!bIsTransacting && !IsGraphEditable())
+	{
+		return false;
+	}
+
 	URigVMGraph* Graph = GetGraph();
 	check(Graph);
 
@@ -6618,6 +6988,11 @@ bool URigVMController::SetPinDefaultValue(const FString& InPinPath, const FStrin
 
 bool URigVMController::SetPinDefaultValue(URigVMPin* InPin, const FString& InDefaultValue, bool bResizeArrays, bool bSetupUndoRedo, bool bMergeUndoAction, bool bNotify)
 {
+	if (!bIsTransacting && !IsGraphEditable())
+	{
+		return false;
+	}
+	
 	check(InPin);
 
 	if(!InPin->IsUObject())
@@ -6735,6 +7110,11 @@ bool URigVMController::ResetPinDefaultValue(const FString& InPinPath, bool bSetu
 		return false;
 	}
 
+	if (!bIsTransacting && !IsGraphEditable())
+	{
+		return false;
+	}
+
 	URigVMGraph* Graph = GetGraph();
 	check(Graph);
 
@@ -6768,6 +7148,11 @@ bool URigVMController::ResetPinDefaultValue(const FString& InPinPath, bool bSetu
 
 bool URigVMController::ResetPinDefaultValue(URigVMPin* InPin, bool bSetupUndoRedo)
 {
+	if (!bIsTransacting && !IsGraphEditable())
+	{
+		return false;
+	}
+	
 	check(InPin);
 
 	URigVMNode* RigVMNode = InPin->GetNode();
@@ -6934,13 +7319,18 @@ FString URigVMController::GetPinInitialDefaultValueFromStruct(UScriptStruct* Scr
 
 FString URigVMController::AddAggregatePin(const FString& InNodeName, const FString& InPinName, const FString& InDefaultValue, bool bSetupUndoRedo, bool bPrintPythonCommand)
 {
-#if UE_RIGVM_ENABLE_AGGREGATE_NODES
+#if UE_RIGVM_AGGREGATE_NODES_ENABLED
 	 
 	if (!IsValidGraph())
 	{
 		return FString();
 	}
 
+	if (!bIsTransacting && !IsGraphEditable())
+	{
+		return FString();
+	}
+
 	URigVMGraph* Graph = GetGraph();
 	check(Graph);
 
@@ -6955,6 +7345,11 @@ FString URigVMController::AddAggregatePin(const FString& InNodeName, const FStri
 
 FString URigVMController::AddAggregatePin(URigVMNode* InNode, const FString& InPinName, const FString& InDefaultValue, bool bSetupUndoRedo, bool bPrintPythonCommand)
 {
+	if (!bIsTransacting && !IsGraphEditable())
+	{
+		return FString();
+	}
+	
 	if (!InNode)
 	{
 		return FString();
@@ -7063,6 +7458,7 @@ FString URigVMController::AddAggregatePin(URigVMNode* InNode, const FString& InP
 		if (AggregateNode)
 		{
 			FRigVMControllerGraphGuard GraphGuard(this, AggregateNode->GetContainedGraph(), bSetupUndoRedo);
+			TGuardValue<bool> GuardEditGraph(GetGraph()->bEditable, true);
 			if (UnitNode)
 			{
 				uint16 Index = 0;
@@ -7116,6 +7512,7 @@ FString URigVMController::AddAggregatePin(URigVMNode* InNode, const FString& InP
 	URigVMPin* NewPin = nullptr;	
 	{
 		FRigVMControllerGraphGuard GraphGuard(this, AggregateNode->GetContainedGraph(), bSetupUndoRedo);
+		TGuardValue<bool> GuardEditGraph(GetGraph()->bEditable, true);
 
 		URigVMNode* InnerNode = (AggregateNode == nullptr) ?
 									((UnitNode == nullptr) ? (URigVMNode*) FunctionReferenceNode : UnitNode)
@@ -7133,23 +7530,58 @@ FString URigVMController::AddAggregatePin(URigVMNode* InNode, const FString& InP
 		{
 			if (bSetupUndoRedo)
 			{
-				ActionStack->CancelAction(Action);
+				ActionStack->CancelAction(Action);
+			}
+			return FString();
+		}
+
+		FName NewName = *InPinName;
+		if (NewName.IsNone())
+		{
+			URigVMNode* LastInnerNode = AggregateNode->GetLastInnerNode();
+			URigVMPin* SecondAggregateInnerPin = LastInnerNode->GetSecondAggregatePin();
+			FString LastAggregateName;
+			if (AggregateNode->IsInputAggregate())
+			{
+				TArray<URigVMPin*> SourcePins = SecondAggregateInnerPin->GetLinkedSourcePins();
+				if (SourcePins.Num() > 0)
+				{
+					LastAggregateName = SourcePins[0]->GetName();
+				}
+			}
+			else
+			{
+				TArray<URigVMPin*> TargetPins = SecondAggregateInnerPin->GetLinkedTargetPins();
+				if (TargetPins.Num() > 0)
+				{
+					LastAggregateName = TargetPins[0]->GetName();
+				}
 			}
-			return FString();
-		}
 
+			if (URigVMUnitNode* InnerUnitNode = Cast<URigVMUnitNode>(LastInnerNode))
+			{
+				NewName = InnerUnitNode->GetNextAggregateName(*LastAggregateName);
+			}
+		}
+		if (NewName.IsNone())
+		{
+			NewName = InnerNode->GetSecondAggregatePin()->GetFName();
+		}
+		
 		const URigVMPin* Arg1 = AggregateNode->GetFirstAggregatePin();
-		FName NewPinName = AddExposedPin(*InPinName, Arg1->GetDirection(), Arg1->GetCPPType(), *Arg1->GetCPPTypeObject()->GetPathName(), InDefaultValue, bSetupUndoRedo);
+		FName NewPinName = AddExposedPin(NewName, Arg1->GetDirection(), Arg1->GetCPPType(), *Arg1->GetCPPTypeObject()->GetPathName(), InDefaultValue, bSetupUndoRedo);
 		NewPin = AggregateNode->FindPin(NewPinName.ToString());
 		URigVMPin* NewUnitPinArg1 = NewNode->GetFirstAggregatePin();
 		URigVMPin* NewUnitPinArg2 = NewNode->GetSecondAggregatePin();
 		URigVMPin* NewUnitPinOppositeArg = NewNode->GetOppositeAggregatePin();
+		URigVMNode* PreviousNode = nullptr;
 		if(AggregateNode->IsInputAggregate())
 		{		
 			URigVMFunctionEntryNode* EntryNode = AggregateNode->GetEntryNode();
 			URigVMPin* EntryPin = EntryNode->FindPin(NewPinName.ToString());
 			URigVMPin* ReturnPin = AggregateNode->GetReturnNode()->FindPin(NewUnitPinOppositeArg->GetName());
-			URigVMPin* PreviousReturnPin = ReturnPin->GetLinkedSourcePins()[0];					
+			URigVMPin* PreviousReturnPin = ReturnPin->GetLinkedSourcePins()[0];
+			PreviousNode = PreviousReturnPin->GetNode();
 		
 			BreakAllLinks(ReturnPin, true, bSetupUndoRedo);
 			AddLink(PreviousReturnPin, NewUnitPinArg1, bSetupUndoRedo);						
@@ -7162,24 +7594,33 @@ FString URigVMController::AddAggregatePin(URigVMNode* InNode, const FString& InP
 			URigVMPin* NewReturnPin = ReturnNode->FindPin(NewPinName.ToString());
 			URigVMPin* OldReturnPin = ReturnNode->GetPins()[ReturnNode->GetPins().Num()-2];
 			URigVMPin* PreviousReturnPin = OldReturnPin->GetLinkedSourcePins()[0];
+			PreviousNode = PreviousReturnPin->GetNode();
 
 			BreakAllLinks(OldReturnPin, true, bSetupUndoRedo);
 			AddLink(PreviousReturnPin, NewUnitPinOppositeArg, bSetupUndoRedo);						
 			AddLink(NewUnitPinArg1, OldReturnPin, bSetupUndoRedo);
 			AddLink(NewUnitPinArg2, NewReturnPin, bSetupUndoRedo);
-
 		}
 
+		// Rearrange the graph nodes
+		URigVMFunctionReturnNode* ReturnNode = AggregateNode->GetReturnNode();
+		FVector2D NodeDimensions(200, 150);
+		SetNodePosition(NewNode, PreviousNode->GetPosition() + NodeDimensions, bSetupUndoRedo);
+		SetNodePosition(ReturnNode, NewNode->GetPosition() + NodeDimensions, bSetupUndoRedo);
+
 		// Connect other input pins
-		for (URigVMPin* OtherEntryPin : AggregateNode->GetEntryNode()->GetPins())
+		for (URigVMPin* OtherInputPin : AggregateNode->GetFirstInnerNode()->GetPins())
 		{
-			if (OtherEntryPin->GetName() != NewUnitPinArg1->GetName() &&
-				OtherEntryPin->GetName() != NewUnitPinArg2->GetName() &&
-				OtherEntryPin->GetName() != NewUnitPinOppositeArg->GetName())
+			if (OtherInputPin->GetName() != NewUnitPinArg1->GetName() &&
+				OtherInputPin->GetName() != NewUnitPinArg2->GetName() &&
+				OtherInputPin->GetName() != NewUnitPinOppositeArg->GetName())
 			{
+				URigVMPin* OtherEntryPin = AggregateNode->GetEntryNode()->FindPin(OtherInputPin->GetName());
 				AddLink(OtherEntryPin, NewNode->FindPin(OtherEntryPin->GetName()), bSetupUndoRedo);
 			}
 		}
+
+		AggregateNode->LastInnerNodeCache = NewNode;
 	}
 
 	if (!NewPin)
@@ -7218,13 +7659,18 @@ FString URigVMController::AddAggregatePin(URigVMNode* InNode, const FString& InP
 
 bool URigVMController::RemoveAggregatePin(const FString& InPinPath, bool bSetupUndoRedo, bool bPrintPythonCommand)
 {
-#if UE_RIGVM_ENABLE_AGGREGATE_NODES
+#if UE_RIGVM_AGGREGATE_NODES_ENABLED
 	
 	if (!IsValidGraph())
 	{
 		return false;
 	}
 
+	if (!bIsTransacting && !IsGraphEditable())
+	{
+		return false;
+	}
+
 	URigVMGraph* Graph = GetGraph();
 	check(Graph);
 
@@ -7239,6 +7685,11 @@ bool URigVMController::RemoveAggregatePin(const FString& InPinPath, bool bSetupU
 
 bool URigVMController::RemoveAggregatePin(URigVMPin* InPin, bool bSetupUndoRedo, bool bPrintPythonCommand)
 {
+	if (!bIsTransacting && !IsGraphEditable())
+	{
+		return false;
+	}
+	
 	if (!InPin)
 	{
 		return false;
@@ -7270,13 +7721,15 @@ bool URigVMController::RemoveAggregatePin(URigVMPin* InPin, bool bSetupUndoRedo,
 					if (EntryPin->GetLinkedTargetPins().Num() > 0)
 					{
 						FRigVMControllerGraphGuard GraphGuard(this, AggregateNode->GetContainedGraph(), bSetupUndoRedo);
+						TGuardValue<bool> GuardEditGraph(GetGraph()->bEditable, true);
+						
 						URigVMPin* TargetPin = EntryPin->GetLinkedTargetPins()[0];
 					
 						URigVMNode* NodeToRemove = TargetPin->GetNode();
 						URigVMPin* ResultPin = NodeToRemove->GetOppositeAggregatePin();
 						URigVMPin* NextNodePin = ResultPin->GetLinkedTargetPins()[0];
 
-						if (NodeToRemove == AggregateNode->FirstInnerNodeCache)
+						if (NodeToRemove == AggregateNode->FirstInnerNodeCache || NodeToRemove == AggregateNode->LastInnerNodeCache)
 						{
 							AggregateNode->InvalidateCache();
 						}
@@ -7302,6 +7755,8 @@ bool URigVMController::RemoveAggregatePin(URigVMPin* InPin, bool bSetupUndoRedo,
 					if (ReturnPin->GetLinkedSourcePins().Num() > 0)
 					{
 						FRigVMControllerGraphGuard GraphGuard(this, AggregateNode->GetContainedGraph(), bSetupUndoRedo);
+						TGuardValue<bool> GuardEditGraph(GetGraph()->bEditable, true);
+						
 						URigVMPin* SourcePin = ReturnPin->GetLinkedSourcePins()[0];
 					
 						URigVMNode* NodeToRemove = SourcePin->GetNode();
@@ -7309,7 +7764,7 @@ bool URigVMController::RemoveAggregatePin(URigVMPin* InPin, bool bSetupUndoRedo,
 						URigVMPin* NextNodePin = OppositePin->GetLinkedSourcePins()[0];
 						URigVMNode* NextNode = NextNodePin->GetNode();
 
-						if (NodeToRemove == AggregateNode->FirstInnerNodeCache)
+						if (NodeToRemove == AggregateNode->FirstInnerNodeCache || NodeToRemove == AggregateNode->LastInnerNodeCache)
 						{
 							AggregateNode->InvalidateCache();
 						}
@@ -7376,6 +7831,11 @@ FString URigVMController::DuplicateArrayPin(const FString& InArrayElementPinPath
 		return FString();
 	}
 
+	if (!bIsTransacting && !IsGraphEditable())
+	{
+		return FString();
+	}
+
 	URigVMGraph* Graph = GetGraph();
 	check(Graph);
 
@@ -7407,6 +7867,11 @@ FString URigVMController::InsertArrayPin(const FString& InArrayPinPath, int32 In
 		return FString();
 	}
 
+	if (!bIsTransacting && !IsGraphEditable())
+	{
+		return FString();
+	}
+
 	URigVMGraph* Graph = GetGraph();
 	check(Graph);
 
@@ -7440,6 +7905,11 @@ FString URigVMController::InsertArrayPin(const FString& InArrayPinPath, int32 In
 
 URigVMPin* URigVMController::InsertArrayPin(URigVMPin* ArrayPin, int32 InIndex, const FString& InDefaultValue, bool bSetupUndoRedo)
 {
+	if (!bIsTransacting && !IsGraphEditable())
+	{
+		return nullptr;
+	}
+	
 	if (!ArrayPin->IsArray())
 	{
 		ReportErrorf(TEXT("Pin '%s' is not an array."), *ArrayPin->GetPinPath());
@@ -7527,6 +7997,11 @@ bool URigVMController::RemoveArrayPin(const FString& InArrayElementPinPath, bool
 		return false;
 	}
 
+	if (!bIsTransacting && !IsGraphEditable())
+	{
+		return false;
+	}
+
 	URigVMGraph* Graph = GetGraph();
 	check(Graph);
 
@@ -7594,6 +8069,11 @@ bool URigVMController::RemoveArrayPin(const FString& InArrayElementPinPath, bool
 
 bool URigVMController::RemovePin(URigVMPin* InPinToRemove, bool bSetupUndoRedo, bool bNotify)
 {
+	if (!bIsTransacting && !IsGraphEditable())
+	{
+		return false;
+	}
+	
 	FRigVMControllerCompileBracketScope CompileScope(this);
 	if (bSetupUndoRedo)
 	{
@@ -7643,6 +8123,11 @@ bool URigVMController::SetArrayPinSize(const FString& InArrayPinPath, int32 InSi
 		return false;
 	}
 
+	if (!bIsTransacting && !IsGraphEditable())
+	{
+		return false;
+	}
+
 	URigVMGraph* Graph = GetGraph();
 	check(Graph);
 
@@ -7729,6 +8214,11 @@ bool URigVMController::BindPinToVariable(const FString& InPinPath, const FString
 		return false;
 	}
 
+	if (!bIsTransacting && !IsGraphEditable())
+	{
+		return false;
+	}
+
 	URigVMGraph* Graph = GetGraph();
 	check(Graph);
 
@@ -7765,6 +8255,11 @@ bool URigVMController::BindPinToVariable(const FString& InPinPath, const FString
 
 bool URigVMController::BindPinToVariable(URigVMPin* InPin, const FString& InNewBoundVariablePath, bool bSetupUndoRedo, const FString& InVariableNodeName)
 {
+	if (!bIsTransacting && !IsGraphEditable())
+	{
+		return false;
+	}
+	
 	if (!IsValidPinForGraph(InPin))
 	{
 		return false;
@@ -7900,6 +8395,11 @@ bool URigVMController::UnbindPinFromVariable(const FString& InPinPath, bool bSet
 		return false;
 	}
 
+	if (!bIsTransacting && !IsGraphEditable())
+	{
+		return false;
+	}
+
 	URigVMGraph* Graph = GetGraph();
 	check(Graph);
 
@@ -7926,6 +8426,11 @@ bool URigVMController::UnbindPinFromVariable(const FString& InPinPath, bool bSet
 
 bool URigVMController::UnbindPinFromVariable(URigVMPin* InPin, bool bSetupUndoRedo)
 {
+	if (!bIsTransacting && !IsGraphEditable())
+	{
+		return false;
+	}
+	
 	if (!IsValidPinForGraph(InPin))
 	{
 		return false;
@@ -7969,6 +8474,11 @@ bool URigVMController::MakeBindingsFromVariableNode(const FName& InNodeName, boo
 		return false;
 	}
 
+	if (!bIsTransacting && !IsGraphEditable())
+	{
+		return false;
+	}
+
 	URigVMGraph* Graph = GetGraph();
 	check(Graph);
 
@@ -7982,6 +8492,11 @@ bool URigVMController::MakeBindingsFromVariableNode(const FName& InNodeName, boo
 
 bool URigVMController::MakeBindingsFromVariableNode(URigVMVariableNode* InNode, bool bSetupUndoRedo)
 {
+	if (!bIsTransacting && !IsGraphEditable())
+	{
+		return false;
+	}
+	
 	check(InNode);
 
 	TArray<TPair<URigVMPin*, URigVMPin*>> Pairs;
@@ -8074,6 +8589,11 @@ bool URigVMController::PromotePinToVariable(const FString& InPinPath, bool bCrea
 		return false;
 	}
 
+	if (!bIsTransacting && !IsGraphEditable())
+	{
+		return false;
+	}
+
 	URigVMGraph* Graph = GetGraph();
 	check(Graph);
 
@@ -8102,6 +8622,11 @@ bool URigVMController::PromotePinToVariable(const FString& InPinPath, bool bCrea
 
 bool URigVMController::PromotePinToVariable(URigVMPin* InPin, bool bCreateVariableNode, const FVector2D& InNodePosition, bool bSetupUndoRedo)
 {
+	if (!bIsTransacting && !IsGraphEditable())
+	{
+		return false;
+	}
+
 	check(InPin);
 
 	if (GetGraph()->IsA<URigVMFunctionLibrary>())
@@ -8188,6 +8713,11 @@ bool URigVMController::AddLink(const FString& InOutputPinPath, const FString& In
 		return false;
 	}
 
+	if (!bIsTransacting && !IsGraphEditable())
+	{
+		return false;
+	}
+
 	URigVMGraph* Graph = GetGraph();
 	check(Graph);
 
@@ -8240,6 +8770,11 @@ bool URigVMController::AddLink(const FString& InOutputPinPath, const FString& In
 
 bool URigVMController::AddLink(URigVMPin* OutputPin, URigVMPin* InputPin, bool bSetupUndoRedo, ERigVMPinDirection InUserDirection)
 {
+	if (!bIsTransacting && !IsGraphEditable())
+	{
+		return false;
+	}
+	
 	if(OutputPin == nullptr)
 	{
 		ReportError(TEXT("OutputPin is nullptr."));
@@ -8450,6 +8985,11 @@ bool URigVMController::BreakLink(const FString& InOutputPinPath, const FString&
 		return false;
 	}
 
+	if (!bIsTransacting && !IsGraphEditable())
+	{
+		return false;
+	}
+
 	URigVMGraph* Graph = GetGraph();
 	check(Graph);
 
@@ -8485,6 +9025,11 @@ bool URigVMController::BreakLink(const FString& InOutputPinPath, const FString&
 
 bool URigVMController::BreakLink(URigVMPin* OutputPin, URigVMPin* InputPin, bool bSetupUndoRedo)
 {
+	if (!bIsTransacting && !IsGraphEditable())
+	{
+		return false;
+	}
+	
 	if(!IsValidPinForGraph(OutputPin) || !IsValidPinForGraph(InputPin))
 	{
 		return false;
@@ -8555,6 +9100,11 @@ bool URigVMController::BreakAllLinks(const FString& InPinPath, bool bAsInput, bo
 		return false;
 	}
 
+	if (!bIsTransacting && !IsGraphEditable())
+	{
+		return false;
+	}
+
 	URigVMGraph* Graph = GetGraph();
 	check(Graph);
 
@@ -8587,6 +9137,11 @@ bool URigVMController::BreakAllLinks(const FString& InPinPath, bool bAsInput, bo
 
 bool URigVMController::BreakAllLinks(URigVMPin* Pin, bool bAsInput, bool bSetupUndoRedo)
 {
+	if (!bIsTransacting && !IsGraphEditable())
+	{
+		return false;
+	}
+	
 	if(!Pin->IsLinked(false))
 	{
 		return false;
@@ -8664,6 +9219,11 @@ FName URigVMController::AddExposedPin(const FName& InPinName, ERigVMPinDirection
 		return NAME_None;
 	}
 
+	if (!bIsTransacting && !IsGraphEditable())
+	{
+		return NAME_None;
+	}
+
 	URigVMGraph* Graph = GetGraph();
 	check(Graph);
 
@@ -8853,6 +9413,11 @@ bool URigVMController::RemoveExposedPin(const FName& InPinName, bool bSetupUndoR
 		return false;
 	}
 
+	if (!bIsTransacting && !IsGraphEditable())
+	{
+		return false;
+	}
+
 	URigVMGraph* Graph = GetGraph();
 	check(Graph);
 
@@ -8955,6 +9520,11 @@ bool URigVMController::RenameExposedPin(const FName& InOldPinName, const FName&
 		return false;
 	}
 
+	if (!bIsTransacting && !IsGraphEditable())
+	{
+		return false;
+	}
+
 	URigVMGraph* Graph = GetGraph();
 	check(Graph);
 
@@ -9127,6 +9697,11 @@ bool URigVMController::ChangeExposedPinType(const FName& InPinName, const FStrin
 		return false;
 	}
 
+	if (!bIsTransacting && !IsGraphEditable())
+	{
+		return false;
+	}
+
 	URigVMGraph* Graph = GetGraph();
 	check(Graph);
 
@@ -9284,6 +9859,11 @@ bool URigVMController::SetExposedPinIndex(const FName& InPinName, int32 InNewInd
 		return false;
 	}
 
+	if (!bIsTransacting && !IsGraphEditable())
+	{
+		return false;
+	}
+
 	URigVMGraph* Graph = GetGraph();
 	check(Graph);
 
@@ -9360,6 +9940,11 @@ URigVMFunctionReferenceNode* URigVMController::AddFunctionReferenceNode(URigVMLi
 		return nullptr;
 	}
 
+	if (!bIsTransacting && !IsGraphEditable())
+	{
+		return nullptr;
+	}
+
 	URigVMGraph* Graph = GetGraph();
 	check(Graph);
 
@@ -9471,6 +10056,11 @@ bool URigVMController::SetRemappedVariable(URigVMFunctionReferenceNode* InFuncti
 		return false;
 	}
 
+	if (!bIsTransacting && !IsGraphEditable())
+	{
+		return false;
+	}
+
 	if(InInnerVariableName.IsNone())
 	{
 		return false;
@@ -9553,6 +10143,11 @@ URigVMLibraryNode* URigVMController::AddFunctionToLibrary(const FName& InFunctio
 		return nullptr;
 	}
 
+	if (!bIsTransacting && !IsGraphEditable())
+	{
+		return nullptr;
+	}
+
 	URigVMGraph* Graph = GetGraph();
 	check(Graph);
 
@@ -9584,6 +10179,7 @@ URigVMLibraryNode* URigVMController::AddFunctionToLibrary(const FName& InFunctio
 
 	{
 		FRigVMControllerGraphGuard GraphGuard(this, CollapseNode->GetContainedGraph(), false);
+		TGuardValue<bool> GuardEditGraph(CollapseNode->ContainedGraph->bEditable, true);
 
 		URigVMFunctionEntryNode* EntryNode = NewObject<URigVMFunctionEntryNode>(CollapseNode->ContainedGraph, TEXT("Entry"));
 		CollapseNode->ContainedGraph->Nodes.Add(EntryNode);
@@ -9635,6 +10231,11 @@ bool URigVMController::RemoveFunctionFromLibrary(const FName& InFunctionName, bo
 		return false;
 	}
 
+	if (!bIsTransacting && !IsGraphEditable())
+	{
+		return false;
+	}
+
 	URigVMGraph* Graph = GetGraph();
 	check(Graph);
 
@@ -9655,6 +10256,11 @@ bool URigVMController::RenameFunction(const FName& InOldFunctionName, const FNam
 		return false;
 	}
 
+	if (!bIsTransacting && !IsGraphEditable())
+	{
+		return false;
+	}
+
 	URigVMGraph* Graph = GetGraph();
 	check(Graph);
 
@@ -9682,6 +10288,11 @@ FRigVMGraphVariableDescription URigVMController::AddLocalVariable(const FName& I
 	{
 		return NewVariable;
 	}
+
+	if (!bIsTransacting && !IsGraphEditable())
+	{
+		return NewVariable;
+	}
 	
 	URigVMGraph* Graph = GetGraph();
 	check(Graph);
@@ -9771,6 +10382,11 @@ FRigVMGraphVariableDescription URigVMController::AddLocalVariableFromObjectPath(
 		return Description;
 	}
 
+	if (!bIsTransacting && !IsGraphEditable())
+	{
+		return Description;
+	}
+
 	UObject* CPPTypeObject = nullptr;
 	if (!InCPPTypeObjectPath.IsEmpty())
 	{
@@ -9792,6 +10408,11 @@ bool URigVMController::RemoveLocalVariable(const FName& InVariableName, bool bSe
 		return false;
 	}
 
+	if (!bIsTransacting && !IsGraphEditable())
+	{
+		return false;
+	}
+
 	URigVMGraph* Graph = GetGraph();
 	check(Graph);
 
@@ -9925,6 +10546,11 @@ bool URigVMController::RenameLocalVariable(const FName& InVariableName, const FN
 		return false;
 	}
 
+	if (!bIsTransacting && !IsGraphEditable())
+	{
+		return false;
+	}
+
 	URigVMGraph* Graph = GetGraph();
 	check(Graph);
 
@@ -10014,6 +10640,11 @@ bool URigVMController::SetLocalVariableType(const FName& InVariableName, const F
 		return false;
 	}
 
+	if (!bIsTransacting && !IsGraphEditable())
+	{
+		return false;
+	}
+
 	URigVMGraph* Graph = GetGraph();
 	check(Graph);
 
@@ -10133,6 +10764,11 @@ bool URigVMController::SetLocalVariableDefaultValue(const FName& InVariableName,
 		return false;
 	}
 
+	if (!bIsTransacting && !IsGraphEditable())
+	{
+		return false;
+	}
+
 	URigVMGraph* Graph = GetGraph();
 	check(Graph);
 
@@ -10380,6 +11016,11 @@ void URigVMController::ExpandPinRecursively(URigVMPin* InPin, bool bSetupUndoRed
 
 bool URigVMController::SetVariableName(URigVMVariableNode* InVariableNode, const FName& InVariableName, bool bSetupUndoRedo)
 {
+	if (!bIsTransacting && !IsGraphEditable())
+	{
+		return false;
+	}
+	
 	if (!IsValidNodeForGraph(InVariableNode))
 	{
 		return false;
@@ -10449,6 +11090,11 @@ URigVMRerouteNode* URigVMController::AddFreeRerouteNode(bool bShowAsFullNode, co
 		return nullptr;
 	}
 
+	if (!bIsTransacting && !IsGraphEditable())
+	{
+		return nullptr;
+	}
+
 	URigVMGraph* Graph = GetGraph();
 	check(Graph);
 
@@ -10513,6 +11159,11 @@ URigVMBranchNode* URigVMController::AddBranchNode(const FVector2D& InPosition, c
 		return nullptr;
 	}
 
+	if (!bIsTransacting && !IsGraphEditable())
+	{
+		return nullptr;
+	}
+
 	URigVMGraph* Graph = GetGraph();
 	check(Graph);
 
@@ -10577,6 +11228,11 @@ URigVMIfNode* URigVMController::AddIfNode(const FString& InCPPType, const FName&
 		return nullptr;
 	}
 
+	if (!bIsTransacting && !IsGraphEditable())
+	{
+		return nullptr;
+	}
+
 	URigVMGraph* Graph = GetGraph();
 	check(Graph);
 
@@ -10693,6 +11349,11 @@ URigVMSelectNode* URigVMController::AddSelectNode(const FString& InCPPType, cons
 		return nullptr;
 	}
 
+	if (!bIsTransacting && !IsGraphEditable())
+	{
+		return nullptr;
+	}
+
 	URigVMGraph* Graph = GetGraph();
 	check(Graph);
 
@@ -10794,6 +11455,11 @@ URigVMTemplateNode* URigVMController::AddTemplateNode(const FName& InNotation, c
 		return nullptr;
 	}
 
+	if (!bIsTransacting && !IsGraphEditable())
+	{
+		return nullptr;
+	}
+
 	URigVMGraph* Graph = GetGraph();
 	check(Graph);
 
@@ -10949,6 +11615,11 @@ URigVMEnumNode* URigVMController::AddEnumNode(const FName& InCPPTypeObjectPath,
 		return nullptr;
 	}
 
+	if (!bIsTransacting && !IsGraphEditable())
+	{
+		return nullptr;
+	}
+
 	URigVMGraph* Graph = GetGraph();
 	check(Graph);
 
@@ -11016,6 +11687,11 @@ URigVMArrayNode* URigVMController::AddArrayNode(ERigVMOpCode InOpCode, const FSt
 		return nullptr;
 	}
 
+	if (!bIsTransacting && !IsGraphEditable())
+	{
+		return nullptr;
+	}
+
 	// validate the op code
 	bool bIsMutable = false;
 	switch(InOpCode)
@@ -11308,6 +11984,11 @@ URigVMArrayNode* URigVMController::AddArrayNodeFromObjectPath(ERigVMOpCode InOpC
 		return nullptr;
 	}
 
+	if (!bIsTransacting && !IsGraphEditable())
+	{
+		return nullptr;
+	}
+
 	UObject* CPPTypeObject = nullptr;
 	if (!InCPPTypeObjectPath.IsEmpty())
 	{
@@ -11367,6 +12048,19 @@ bool URigVMController::IsValidGraph() const
 	return true;
 }
 
+bool URigVMController::IsGraphEditable() const
+{
+	if(!IsValidGraph())
+	{
+		return false;
+	}
+
+	URigVMGraph* Graph = GetGraph();
+	check(Graph);
+
+	return Graph->bEditable;
+}
+
 bool URigVMController::IsValidNodeForGraph(URigVMNode* InNode)
 {
 	if(!IsValidGraph())
@@ -11374,6 +12068,11 @@ bool URigVMController::IsValidNodeForGraph(URigVMNode* InNode)
 		return false;
 	}
 
+	if (!bIsTransacting && !IsGraphEditable())
+	{
+		return false;
+	}
+
 	if (InNode == nullptr)
 	{
 		ReportError(TEXT("InNode is nullptr."));
@@ -11401,6 +12100,11 @@ bool URigVMController::IsValidPinForGraph(URigVMPin* InPin)
 		return false;
 	}
 
+	if (!bIsTransacting && !IsGraphEditable())
+	{
+		return false;
+	}
+
 	if (InPin == nullptr)
 	{
 		ReportError(TEXT("InPin is nullptr."));
@@ -11427,6 +12131,11 @@ bool URigVMController::IsValidLinkForGraph(URigVMLink* InLink)
 		return false;
 	}
 
+	if (!bIsTransacting && !IsGraphEditable())
+	{
+		return false;
+	}
+
 	if (InLink == nullptr)
 	{
 		ReportError(TEXT("InLink is nullptr."));
@@ -11476,6 +12185,11 @@ bool URigVMController::CanAddNode(URigVMNode* InNode, bool bReportErrors, bool b
 		return false;
 	}
 
+	if (!bIsTransacting && !IsGraphEditable())
+	{
+		return false;
+	}
+
 	check(InNode);
 
 	URigVMGraph* Graph = GetGraph();
@@ -11632,6 +12346,11 @@ bool URigVMController::CanAddEventNode(UScriptStruct* InScriptStruct, const bool
 	{
 		return false;
 	}
+
+	if (!bIsTransacting && !IsGraphEditable())
+	{
+		return false;
+	}
 	
 	check(InScriptStruct);
 	
@@ -11667,6 +12386,11 @@ bool URigVMController::CanAddFunctionRefForDefinition(URigVMLibraryNode* InFunct
 		return false;
 	}
 
+	if (!bIsTransacting && !IsGraphEditable())
+	{
+		return false;
+	}
+
 	check(InFunctionDefinition);
 
 	URigVMGraph* Graph = GetGraph();
@@ -12149,6 +12873,7 @@ int32 URigVMController::DetachLinksFromPinObjects(const TArray<URigVMLink*>* InL
 			if (URigVMCollapseNode* CollapseNode = Cast<URigVMCollapseNode>(Node))
 			{
 				FRigVMControllerGraphGuard GraphGuard(this, CollapseNode->GetContainedGraph(), false);
+				TGuardValue<bool> GuardEditGraph(CollapseNode->ContainedGraph->bEditable, true);
 				DetachLinksFromPinObjects(InLinks, bNotify);
 			}
 		}
@@ -12344,6 +13069,7 @@ int32 URigVMController::ReattachLinksToPinObjects(bool bFollowCoreRedirectors, c
 			if (URigVMCollapseNode* CollapseNode = Cast<URigVMCollapseNode>(Node))
 			{
 				FRigVMControllerGraphGuard GraphGuard(this, CollapseNode->GetContainedGraph(), false);
+				TGuardValue<bool> GuardEditGraph(CollapseNode->ContainedGraph->bEditable, true);
 				ReattachLinksToPinObjects(bFollowCoreRedirectors, nullptr, bNotify, bSetupOrphanedPins);
 			}
 		}
@@ -12842,6 +13568,7 @@ void URigVMController::RepopulatePinsOnNode(URigVMNode* InNode, bool bFollowCore
 		}
 
 		FRigVMControllerGraphGuard GraphGuard(this, CollapseNode->GetContainedGraph(), false);
+		TGuardValue<bool> GuardEditGraph(CollapseNode->ContainedGraph->bEditable, true);
 		// need to get a copy of the node array since the following function could remove nodes from the graph
 		// we don't want to remove elements from the array we are iterating over.
 		TArray<URigVMNode*> ContainedNodes = CollapseNode->GetContainedNodes();
@@ -13729,6 +14456,11 @@ bool URigVMController::ResolveWildCardPin(const FString& InPinPath, const FStrin
 		return false;
 	}
 
+	if (!bIsTransacting && !IsGraphEditable())
+	{
+		return false;
+	}
+
 	URigVMGraph* Graph = GetGraph();
 	check(Graph);
 
@@ -14017,6 +14749,11 @@ bool URigVMController::ChangePinType(const FString& InPinPath, const FString& In
 		return false;
 	}
 
+	if (!bIsTransacting && !IsGraphEditable())
+	{
+		return false;
+	}
+
 	URigVMGraph* Graph = GetGraph();
 	check(Graph);
 
@@ -14030,6 +14767,11 @@ bool URigVMController::ChangePinType(const FString& InPinPath, const FString& In
 
 bool URigVMController::ChangePinType(URigVMPin* InPin, const FString& InCPPType, const FName& InCPPTypeObjectPath, bool bSetupUndoRedo, bool bSetupOrphanPins, bool bBreakLinks, bool bRemoveSubPins)
 {
+	if (!bIsTransacting && !IsGraphEditable())
+	{
+		return false;
+	}
+	
 	if (InCPPType == TEXT("None") || InCPPType.IsEmpty())
 	{
 		return false;
@@ -14057,6 +14799,11 @@ bool URigVMController::ChangePinType(URigVMPin* InPin, const FString& InCPPType,
 
 bool URigVMController::ChangePinType(URigVMPin* InPin, const FString& InCPPType, UObject* InCPPTypeObject, bool bSetupUndoRedo, bool bSetupOrphanPins, bool bBreakLinks, bool bRemoveSubPins)
 {
+	if (!bIsTransacting && !IsGraphEditable())
+	{
+		return false;
+	}
+	
 	if (InCPPType == TEXT("None") || InCPPType.IsEmpty())
 	{
 		return false;
diff --git a/Engine/Source/Developer/RigVMDeveloper/Private/RigVMModel/RigVMControllerActions.cpp b/Engine/Source/Developer/RigVMDeveloper/Private/RigVMModel/RigVMControllerActions.cpp
--- a/Engine/Source/Developer/RigVMDeveloper/Private/RigVMModel/RigVMControllerActions.cpp
+++ b/Engine/Source/Developer/RigVMDeveloper/Private/RigVMModel/RigVMControllerActions.cpp
@@ -248,6 +248,7 @@ bool FRigVMBaseAction::Undo(URigVMController* InController)
 	check(Stack);
 	TGuardValue<int32> TabDepthGuard(Stack->LogActionDepth, Stack->LogActionDepth + 1);
 #endif
+	TGuardValue<bool> TransactionGuard(InController->bIsTransacting, true);
 
 	bool Result = true;
 	for (int32 KeyIndex = SubActions.Num() - 1; KeyIndex >= 0; KeyIndex--)
@@ -272,6 +273,7 @@ bool FRigVMBaseAction::Redo(URigVMController* InController)
 	check(Stack);
 	TGuardValue<int32> TabDepthGuard(Stack->LogActionDepth, Stack->LogActionDepth + 1);
 #endif
+	TGuardValue<bool> TransactionGuard(InController->bIsTransacting, true);
 
 	bool Result = true;
 	for (int32 KeyIndex = 0; KeyIndex < SubActions.Num(); KeyIndex++)
diff --git a/Engine/Source/Developer/RigVMDeveloper/Private/RigVMModel/RigVMGraph.cpp b/Engine/Source/Developer/RigVMDeveloper/Private/RigVMModel/RigVMGraph.cpp
--- a/Engine/Source/Developer/RigVMDeveloper/Private/RigVMModel/RigVMGraph.cpp
+++ b/Engine/Source/Developer/RigVMDeveloper/Private/RigVMModel/RigVMGraph.cpp
@@ -12,6 +12,7 @@
 URigVMGraph::URigVMGraph()
 : DiagnosticsAST(nullptr)
 , RuntimeAST(nullptr)
+, bEditable(true)
 {
 }
 
diff --git a/Engine/Source/Developer/RigVMDeveloper/Public/RigVMModel/Nodes/RigVMAggregateNode.h b/Engine/Source/Developer/RigVMDeveloper/Public/RigVMModel/Nodes/RigVMAggregateNode.h
--- a/Engine/Source/Developer/RigVMDeveloper/Public/RigVMModel/Nodes/RigVMAggregateNode.h
+++ b/Engine/Source/Developer/RigVMDeveloper/Public/RigVMModel/Nodes/RigVMAggregateNode.h
@@ -29,10 +29,13 @@ public:
 	virtual  FText GetToolTipText() const override;
 	virtual FText GetToolTipTextForPin(const URigVMPin* InPin) const override;
 
+#if UE_RIGVM_AGGREGATE_NODES_ENABLED
 	virtual bool IsAggregate() const override { return true; };
+#endif
 	virtual bool IsInputAggregate() const override;
 
 	URigVMNode* GetFirstInnerNode() const;
+	URigVMNode* GetLastInnerNode() const;
 
 	UFUNCTION(BlueprintPure, Category = Aggregate)
 	URigVMPin* GetFirstAggregatePin() const;
@@ -48,6 +51,7 @@ private:
 	void InvalidateCache();
 
 	mutable URigVMNode* FirstInnerNodeCache;
+	mutable URigVMNode* LastInnerNodeCache;
 
 	friend class URigVMController;
 };
diff --git a/Engine/Source/Developer/RigVMDeveloper/Public/RigVMModel/Nodes/RigVMUnitNode.h b/Engine/Source/Developer/RigVMDeveloper/Public/RigVMModel/Nodes/RigVMUnitNode.h
--- a/Engine/Source/Developer/RigVMDeveloper/Public/RigVMModel/Nodes/RigVMUnitNode.h
+++ b/Engine/Source/Developer/RigVMDeveloper/Public/RigVMModel/Nodes/RigVMUnitNode.h
@@ -35,6 +35,7 @@ public:
 	
 	TArray<URigVMPin*> GetAggregateInputs() const;
 	TArray<URigVMPin*> GetAggregateOutputs() const;
+	FName GetNextAggregateName(const FName& InLastAggregatePinName)const;
 
 	bool IsDeprecated() const;
 	FString GetDeprecatedMetadata() const;
diff --git a/Engine/Source/Developer/RigVMDeveloper/Public/RigVMModel/RigVMController.h b/Engine/Source/Developer/RigVMDeveloper/Public/RigVMModel/RigVMController.h
--- a/Engine/Source/Developer/RigVMDeveloper/Public/RigVMModel/RigVMController.h
+++ b/Engine/Source/Developer/RigVMDeveloper/Public/RigVMModel/RigVMController.h
@@ -24,10 +24,6 @@
 #define UE_RIGVM_ENABLE_TEMPLATE_NODES 1
 #endif
 
-#ifndef UE_RIGVM_ENABLE_AGGREGATE_NODES
-#define UE_RIGVM_ENABLE_AGGREGATE_NODES 0
-#endif
-
 class URigVMActionStack;
 
 UENUM()
@@ -604,7 +600,7 @@ public:
 	UFUNCTION(BlueprintCallable, Category = RigVMController)
 	bool RemoveAggregatePin(const FString& InPinPath, bool bSetupUndoRedo = true, bool bPrintPythonCommand = false);
 
-#if UE_RIGVM_ENABLE_AGGREGATE_NODES
+#if UE_RIGVM_AGGREGATE_NODES_ENABLED
 	FString AddAggregatePin(URigVMNode* InNode, const FString& InPinName, const FString& InDefaultValue = TEXT(""), bool bSetupUndoRedo = true, bool bPrintPythonCommand = false);
 	bool RemoveAggregatePin(URigVMPin* InPin, bool bSetupUndoRedo = true, bool bPrintPythonCommand = false);
 #endif
@@ -881,6 +877,7 @@ private:
 
 	FString GetValidNodeName(const FString& InPrefix);
 	bool IsValidGraph() const;
+	bool IsGraphEditable() const;
 	bool IsValidNodeForGraph(URigVMNode* InNode);
 	bool IsValidPinForGraph(URigVMPin* InPin);
 	bool IsValidLinkForGraph(URigVMLink* InLink);
@@ -1057,6 +1054,7 @@ private:
 
 	FRigVMUnitNodeCreatedContext UnitNodeCreatedContext;
 
+	bool bIsTransacting; // Performing undo/redo transaction
 	bool bIsRunningUnitTest;
 	bool bIsFullyResolvingTemplateNode;
 
diff --git a/Engine/Source/Developer/RigVMDeveloper/Public/RigVMModel/RigVMGraph.h b/Engine/Source/Developer/RigVMDeveloper/Public/RigVMModel/RigVMGraph.h
--- a/Engine/Source/Developer/RigVMDeveloper/Public/RigVMModel/RigVMGraph.h
+++ b/Engine/Source/Developer/RigVMDeveloper/Public/RigVMModel/RigVMGraph.h
@@ -164,12 +164,16 @@ private:
 	TArray<TSharedPtr<FString>> ParameterNames;
 #endif
 
+	UPROPERTY()
+	bool bEditable;
+
 	UPROPERTY()
 	TArray<FRigVMGraphVariableDescription> LocalVariables;
 
 	bool IsNameAvailable(const FString& InName);
 
 	friend class URigVMController;
+	friend class UControlRigBlueprint;
 	friend class FRigVMControllerCompileBracketScope;
 	friend class URigVMCompiler;
 	friend struct FRigVMControllerObjectFactory;
diff --git a/Engine/Source/Programs/Shared/EpicGames.UHT/Parsers/UhtScriptStructParser.cs b/Engine/Source/Programs/Shared/EpicGames.UHT/Parsers/UhtScriptStructParser.cs
--- a/Engine/Source/Programs/Shared/EpicGames.UHT/Parsers/UhtScriptStructParser.cs
+++ b/Engine/Source/Programs/Shared/EpicGames.UHT/Parsers/UhtScriptStructParser.cs
@@ -217,16 +217,23 @@ namespace EpicGames.UHT.Parsers
 					.RequireIdentifier((ref UhtToken Identifier) => MethodInfo.Name = Identifier.Value.ToString());
 
 				bool bIsGetUpgradeInfo = false;
+				bool bIsGetNextAggregateName = false;
 				if (MethodInfo.ReturnType == "FRigVMStructUpgradeInfo" && MethodInfo.Name == "GetUpgradeInfo")
 				{
 					StructInfo.bHasGetUpgradeInfoMethod = true;
 					bIsGetUpgradeInfo = true;
 				}
+				if (MethodInfo.Name == "GetNextAggregateName")
+				{
+					StructInfo.bHasGetNextAggregateNameMethod = true;
+					bIsGetNextAggregateName = true;
+				}
+				
 
 				TopScope.TokenReader
 					.RequireList('(', ')', ',', false, (IEnumerable<UhtToken> Tokens) =>
 					{
-						if (!bIsGetUpgradeInfo)
+						if (!bIsGetUpgradeInfo && !bIsGetNextAggregateName)
 						{
 							StringViewBuilder Builder = new StringViewBuilder();
 							UhtToken LastToken = new UhtToken();
@@ -251,7 +258,7 @@ namespace EpicGames.UHT.Parsers
 					})
 					.ConsumeUntil(';');
 
-				if (!bIsGetUpgradeInfo)
+				if (!bIsGetUpgradeInfo && !bIsGetNextAggregateName)
 				{
 					StructInfo.Methods.Add(MethodInfo);
 				}
diff --git a/Engine/Source/Programs/Shared/EpicGames.UHT/Types/UhtScriptStruct.cs b/Engine/Source/Programs/Shared/EpicGames.UHT/Types/UhtScriptStruct.cs
--- a/Engine/Source/Programs/Shared/EpicGames.UHT/Types/UhtScriptStruct.cs
+++ b/Engine/Source/Programs/Shared/EpicGames.UHT/Types/UhtScriptStruct.cs
@@ -440,6 +440,11 @@ namespace EpicGames.UHT.Types
 		/// True if the GetUpgradeInfoMethod was found. 
 		/// </summary>
 		public bool bHasGetUpgradeInfoMethod = false;
+		
+		/// <summary>
+		/// True if the GetNextAggregateNameMethod was found. 
+		/// </summary>
+		public bool bHasGetNextAggregateNameMethod = false;
 
 		/// <summary>
 		/// Engine name of the owning script struct
diff --git a/Engine/Source/Programs/UnrealHeaderTool/Private/HeaderParser.cpp b/Engine/Source/Programs/UnrealHeaderTool/Private/HeaderParser.cpp
--- a/Engine/Source/Programs/UnrealHeaderTool/Private/HeaderParser.cpp
+++ b/Engine/Source/Programs/UnrealHeaderTool/Private/HeaderParser.cpp
@@ -6410,6 +6410,15 @@ void FHeaderParser::CompileRigVMMethodDeclaration(FUnrealStructDefinitionInfo& S
 		return;
 	}
 
+	// look out for the next aggregate name method
+	static const FString GetNextAggregateNameString = TEXT("GetNextAggregateName");
+	if(MethodInfo.Name == GetNextAggregateNameString)
+	{
+		FRigVMStructInfo& StructRigVMInfo = StructDef.GetRigVMInfo();
+		StructRigVMInfo.bHasGetNextAggregateNameMethod = true;
+		return;
+	}
+
 	FString ParamString = FString::Join(ParamsContent, TEXT(" "));
 	if (!ParamString.IsEmpty())
 	{
diff --git a/Engine/Source/Programs/UnrealHeaderTool/Private/ParserHelper.h b/Engine/Source/Programs/UnrealHeaderTool/Private/ParserHelper.h
--- a/Engine/Source/Programs/UnrealHeaderTool/Private/ParserHelper.h
+++ b/Engine/Source/Programs/UnrealHeaderTool/Private/ParserHelper.h
@@ -716,6 +716,7 @@ struct FRigVMStructInfo
 {
 	bool bHasRigVM = false;
 	bool bHasGetUpgradeInfoMethod = false;
+	bool bHasGetNextAggregateNameMethod = false;
 	FString Name;
 	FRigVMParameterArray Members;
 	TArray<FRigVMMethodInfo> Methods;
diff --git a/Engine/Source/Programs/UnrealHeaderTool/Resources/UHTDebugging/RigVMStructHeader.h b/Engine/Source/Programs/UnrealHeaderTool/Resources/UHTDebugging/RigVMStructHeader.h
--- a/Engine/Source/Programs/UnrealHeaderTool/Resources/UHTDebugging/RigVMStructHeader.h
+++ b/Engine/Source/Programs/UnrealHeaderTool/Resources/UHTDebugging/RigVMStructHeader.h
@@ -34,6 +34,8 @@ struct FRigVMStructBase
 	UPROPERTY(meta = (Output))
 	float InheritedOutput;
 
+	virtual FName GetNextAggregateName(const FName& InLastAggregateName) const {};
+
 	virtual FRigVMStructUpgradeInfo GetUpgradeInfo() const {};
 };
 
@@ -96,6 +98,9 @@ struct FRigVMMethodStruct : public FRigVMStructBase
 	UPROPERTY()
 	TEnumAsByte<ERigVMTestNameSpaceEnum::Type> HiddenNameSpaceEnum;
 
+	RIGVM_METHOD()
+	virtual FName GetNextAggregateName(const FName& InLastAggregatePinName) const override;
+
 	RIGVM_METHOD()
 	virtual FRigVMStructUpgradeInfo GetUpgradeInfo() const override;
 };
diff --git a/Engine/Source/Runtime/RigVM/Public/RigVMCore/RigVMStruct.h b/Engine/Source/Runtime/RigVM/Public/RigVMCore/RigVMStruct.h
--- a/Engine/Source/Runtime/RigVM/Public/RigVMCore/RigVMStruct.h
+++ b/Engine/Source/Runtime/RigVM/Public/RigVMCore/RigVMStruct.h
@@ -233,6 +233,7 @@ public:
 	static FString ExportToFullyQualifiedText(const FProperty* InMemberProperty, const uint8* InMemberMemoryPtr, bool bUseQuotes = true);
 	static FString ExportToFullyQualifiedText(const UScriptStruct* InStruct, const uint8* InStructMemoryPtr);
 	FString ExportToFullyQualifiedText(const UScriptStruct* InScriptStruct, const FName& InPropertyName, const uint8* InStructMemoryPointer = nullptr) const;
+	virtual FName GetNextAggregateName(const FName& InLastAggregatePinName) const { return FName(); }
 	virtual FRigVMStructUpgradeInfo GetUpgradeInfo() const { return FRigVMStructUpgradeInfo(); }
 
 	static const FName DeprecatedMetaName;
diff --git a/Engine/Source/Runtime/RigVM/Public/RigVMDefines.h b/Engine/Source/Runtime/RigVM/Public/RigVMDefines.h
--- a/Engine/Source/Runtime/RigVM/Public/RigVMDefines.h
+++ b/Engine/Source/Runtime/RigVM/Public/RigVMDefines.h
@@ -13,3 +13,7 @@
 #ifndef UE_RIGVM_UOBJECT_PROPERTIES_ENABLED
 #define UE_RIGVM_UOBJECT_PROPERTIES_ENABLED 0
 #endif
+
+#ifndef UE_RIGVM_AGGREGATE_NODES_ENABLED
+#define UE_RIGVM_AGGREGATE_NODES_ENABLED 0
+#endif
\ No newline at end of file
