RigVM: Nativization First Steps

diff --git a/Engine/Plugins/Animation/ControlRig/Source/ControlRig/Private/ControlRig.cpp b/Engine/Plugins/Animation/ControlRig/Source/ControlRig/Private/ControlRig.cpp
--- a/Engine/Plugins/Animation/ControlRig/Source/ControlRig/Private/ControlRig.cpp
+++ b/Engine/Plugins/Animation/ControlRig/Source/ControlRig/Private/ControlRig.cpp
@@ -48,6 +48,9 @@ static TAutoConsoleVariable<int32> CVarControlRigCreateFloatControlsForCurves(
 // CVar to disable all control rig execution 
 static TAutoConsoleVariable<int32> CVarControlRigDisableExecutionAll(TEXT("ControlRig.DisableExecutionAll"), 0, TEXT("if nonzero we disable all execution of Control Rigs."));
 
+// CVar to disable swapping to nativized vms 
+static TAutoConsoleVariable<int32> CVarControlRigDisableNativizedVMs(TEXT("ControlRig.DisableNativizedVMs"), 1, TEXT("if nonzero we disable swapping to nativized VMs."));
+
 UControlRig::UControlRig(const FObjectInitializer& ObjectInitializer)
 	: Super(ObjectInitializer)
 	, DeltaTime(0.0f)
@@ -81,6 +84,11 @@ UControlRig::UControlRig(const FObjectInitializer& ObjectInitializer)
 	, VMSnapshotBeforeExecution(nullptr)
 #endif
 	, DebugBoneRadiusMultiplier(1.f)
+#if UE_CONTROLRIG_PROFILE_EXECUTE_UNITS_NUM
+	, ProfilingRunsLeft(0)
+	, AccumulatedCycles(0)
+#endif
+
 {
 	EventQueue.Add(FRigUnit_BeginExecution::EventName);
 }
@@ -501,13 +509,18 @@ void UControlRig::InstantiateVMFromCDO()
 
 	if (!HasAnyFlags(RF_ClassDefaultObject))
 	{
+		SwapVMToNativizedIfRequired();
+
 		UControlRig* CDO = GetClass()->GetDefaultObject<UControlRig>();
 		if (VM && CDO && CDO->VM)
 		{
-			// reference the literal memory + byte code
-			// only defer if called from worker thread,
-			// which should be unlikely
-			VM->CopyFrom(CDO->VM, !IsInGameThread(), true);
+			if(!VM->IsNativized())
+			{
+				// reference the literal memory + byte code
+				// only defer if called from worker thread,
+				// which should be unlikely
+				VM->CopyFrom(CDO->VM, !IsInGameThread(), true);
+			}
 		}
 		else if (VM)
 		{
@@ -1169,13 +1182,23 @@ void UControlRig::ExecuteUnits(FRigUnitContext& InOutContext, const FName& InEve
 {
 	if (VM)
 	{
+#if UE_CONTROLRIG_PROFILE_EXECUTE_UNITS_NUM
+		const uint64 StartCycles = FPlatformTime::Cycles64();
+		if(ProfilingRunsLeft <= 0)
+		{
+			ProfilingRunsLeft = UE_CONTROLRIG_PROFILE_EXECUTE_UNITS_NUM;
+			AccumulatedCycles = 0;
+		}
+#endif
+		
+		const bool bUseSnapshots = !VM->IsNativized();
 		TArray<URigVMMemoryStorage*> LocalMemory = VM->GetLocalMemoryArray();
 		TArray<void*> AdditionalArguments;
 		AdditionalArguments.Add(&InOutContext);
 
 		if (InOutContext.State == EControlRigState::Init)
 		{
-			if(IsInGameThread())
+			if(IsInGameThread() && bUseSnapshots)
 			{
 				const uint32 SnapshotHash = GetHashForInitializeVMSnapShot();
 				UControlRig* CDO = Cast<UControlRig>(GetClass()->GetDefaultObject());
@@ -1219,15 +1242,18 @@ void UControlRig::ExecuteUnits(FRigUnitContext& InOutContext, const FName& InEve
 		else
 		{
 #if WITH_EDITOR
-			if(URigVM* SnapShotVM = GetSnapshotVM(false)) // don't create it for normal runs
+			if(bUseSnapshots)
 			{
-				if (VM->GetHaltedAtBreakpoint() != nullptr)
-				{
-					VM->CopyFrom(SnapShotVM, false, false, false, true, true);	
-				}
-				else
+				if(URigVM* SnapShotVM = GetSnapshotVM(false)) // don't create it for normal runs
 				{
-					SnapShotVM->CopyFrom(VM, false, false, false, true, true);
+					if (VM->GetHaltedAtBreakpoint() != nullptr)
+					{
+						VM->CopyFrom(SnapShotVM, false, false, false, true, true);	
+					}
+					else
+					{
+						SnapShotVM->CopyFrom(VM, false, false, false, true, true);
+					}
 				}
 			}
 #endif
@@ -1246,6 +1272,19 @@ void UControlRig::ExecuteUnits(FRigUnitContext& InOutContext, const FName& InEve
 
 			VM->Execute(LocalMemory, AdditionalArguments, InEventName);
 		}
+
+#if UE_CONTROLRIG_PROFILE_EXECUTE_UNITS_NUM
+		const uint64 EndCycles = FPlatformTime::Cycles64();
+		const uint64 Cycles = EndCycles - StartCycles;
+		AccumulatedCycles += Cycles;
+		ProfilingRunsLeft--;
+		if(ProfilingRunsLeft == 0)
+		{
+			const double Milliseconds = FPlatformTime::ToMilliseconds64(AccumulatedCycles);
+			UE_LOG(LogControlRig, Display, TEXT("%s: %d runs took %.03lfms."), *GetClass()->GetName(), UE_CONTROLRIG_PROFILE_EXECUTE_UNITS_NUM, Milliseconds);
+		}
+#endif
+
 	}
 }
 
@@ -2796,6 +2835,80 @@ void UControlRig::PostInitInstanceIfRequired()
 	}
 }
 
+void UControlRig::SwapVMToNativizedIfRequired(UClass* InNativizedClass)
+{
+	if (HasAnyFlags(RF_NeedPostLoad))
+	{
+		return;
+	}
+	if(VM == nullptr)
+	{
+		return;
+	}
+
+	const bool bNativizedVMDisabled = AreNativizedVMsDisabled();
+
+	if(InNativizedClass == nullptr)
+	{
+		if(!HasAnyFlags(RF_ClassDefaultObject))
+		{
+			if(UControlRig* CDO = GetClass()->GetDefaultObject<UControlRig>())
+			{
+				if(CDO->VM)
+				{
+					if(CDO->VM->IsNativized())
+					{
+						InNativizedClass = CDO->VM->GetClass();
+					}
+					else
+					{
+						InNativizedClass = CDO->VM->GetNativizedClass(GetExternalVariables());
+					}
+				}
+			}
+		}
+		else
+		{
+			InNativizedClass = VM->GetNativizedClass(GetExternalVariablesImpl(true));
+		}
+	}	
+
+	if(VM->IsNativized())
+	{
+		if((InNativizedClass == nullptr) || bNativizedVMDisabled)
+		{
+			const EObjectFlags PreviousFlags = VM->GetFlags();
+			VM->Rename(nullptr, GetTransientPackage(), REN_ForceNoResetLoaders | REN_DoNotDirty | REN_DontCreateRedirectors | REN_NonTransactional);
+			VM->MarkAsGarbage();
+			VM = NewObject<URigVM>(this, TEXT("VM"), PreviousFlags);
+#if UE_CONTROLRIG_PROFILE_EXECUTE_UNITS_NUM
+			ProfilingRunsLeft = 0;
+			AccumulatedCycles = 0;
+#endif
+		}
+	}
+	else
+	{
+		if(InNativizedClass && !bNativizedVMDisabled)
+		{
+			const EObjectFlags PreviousFlags = VM->GetFlags();
+			VM->Rename(nullptr, GetTransientPackage(), REN_ForceNoResetLoaders | REN_DoNotDirty | REN_DontCreateRedirectors | REN_NonTransactional);
+			VM->MarkAsGarbage();
+			VM = NewObject<URigVM>(this, InNativizedClass, TEXT("VM"), PreviousFlags);
+			VM->ExecutionReachedExit().AddUObject(this, &UControlRig::HandleExecutionReachedExit);
+#if UE_CONTROLRIG_PROFILE_EXECUTE_UNITS_NUM
+			ProfilingRunsLeft = 0;
+			AccumulatedCycles = 0;
+#endif
+		}
+	}
+}
+
+bool UControlRig::AreNativizedVMsDisabled()
+{
+	return (CVarControlRigDisableNativizedVMs->GetInt() != 0);
+}
+
 void UControlRig::PostInitInstance(UControlRig* InCDO)
 {
 	const EObjectFlags SubObjectFlags =
diff --git a/Engine/Plugins/Animation/ControlRig/Source/ControlRig/Public/ControlRig.h b/Engine/Plugins/Animation/ControlRig/Source/ControlRig/Public/ControlRig.h
--- a/Engine/Plugins/Animation/ControlRig/Source/ControlRig/Public/ControlRig.h
+++ b/Engine/Plugins/Animation/ControlRig/Source/ControlRig/Public/ControlRig.h
@@ -45,6 +45,15 @@ struct FRigControl;
 
 CONTROLRIG_API DECLARE_LOG_CATEGORY_EXTERN(LogControlRig, Log, All);
 
+// set this to something larger than 0 to profile N runs
+#ifndef UE_CONTROLRIG_PROFILE_EXECUTE_UNITS_NUM
+#define UE_CONTROLRIG_PROFILE_EXECUTE_UNITS_NUM 0
+#endif
+
+#if UE_CONTROLRIG_PROFILE_EXECUTE_UNITS_NUM
+#include "HAL/PlatformTime.h"
+#endif
+
 /** Runs logic for mapping input data to transforms (the "Rig") */
 UCLASS(Blueprintable, Abstract, editinlinenew)
 class CONTROLRIG_API UControlRig : public UObject, public INodeMappingProviderInterface, public IInterface_AssetUserData
@@ -432,6 +441,8 @@ public:
 	void SetControlCustomization(const FRigElementKey& InControl, const FRigControlElementCustomization& InCustomization);
 
 	void PostInitInstanceIfRequired();
+	void SwapVMToNativizedIfRequired(UClass* InNativizedClass = nullptr);
+	static bool AreNativizedVMsDisabled() ;
 
 private:
 
@@ -843,8 +854,13 @@ public:
 		FRigPose CachedPose;	
 	};	
 
-private:
+#endif
 	
+private:
+
+#if UE_CONTROLRIG_PROFILE_EXECUTE_UNITS_NUM
+	int32 ProfilingRunsLeft;
+	uint64 AccumulatedCycles;
 #endif
 
 	friend class FControlRigBlueprintCompilerContext;
diff --git a/Engine/Plugins/Animation/ControlRig/Source/ControlRigEditor/Private/ControlRigCompilerDetails.cpp b/Engine/Plugins/Animation/ControlRig/Source/ControlRigEditor/Private/ControlRigCompilerDetails.cpp
--- a/Engine/Plugins/Animation/ControlRig/Source/ControlRigEditor/Private/ControlRigCompilerDetails.cpp
+++ b/Engine/Plugins/Animation/ControlRig/Source/ControlRigEditor/Private/ControlRigCompilerDetails.cpp
@@ -15,6 +15,7 @@
 #include "RigVMPythonUtils.h"
 #include "ControlRigVisualGraphUtils.h"
 #include "ControlRig/Private/Units/Execution/RigUnit_BeginExecution.h"
+#include "RigVMCompiler/RigVMCodeGenerator.h"
 
 #define LOCTEXT_NAMESPACE "ControlRigCompilerDetails"
 
@@ -126,6 +127,22 @@ void FRigVMCompileSettingsDetails::CustomizeChildren(TSharedRef<IPropertyHandle>
 					]
 				]
 				+ SVerticalBox::Slot()
+				[
+					SNew(SButton)
+					.OnClicked(this, &FRigVMCompileSettingsDetails::OnCopyGeneratedCodeClicked)
+					.ContentPadding(FMargin(2))
+					.Visibility_Lambda([]()
+					{
+						return UControlRig::AreNativizedVMsDisabled() ? EVisibility::Collapsed : EVisibility::Visible;
+					})
+					.Content()
+					[
+						SNew(STextBlock)
+						.Justification(ETextJustify::Center)
+						.Text(LOCTEXT("CopyGeneratedCodeToClipboard", "Copy Nativized C++ Code"))
+					]
+				]
+				+ SVerticalBox::Slot()
 				[
 					SNew(SButton)
 					.OnClicked(this, &FRigVMCompileSettingsDetails::OnCopyHierarchyGraphClicked)
@@ -205,4 +222,33 @@ FReply FRigVMCompileSettingsDetails::OnCopyHierarchyGraphClicked()
 	return FReply::Handled();
 }
 
+FReply FRigVMCompileSettingsDetails::OnCopyGeneratedCodeClicked()
+{
+	if (BlueprintBeingCustomized)
+	{
+		const FString ClassName = FString::Printf(TEXT("%sVM"), *BlueprintBeingCustomized->GetName());
+		if (BlueprintBeingCustomized->GeneratedClass)
+		{
+			if (UControlRig* CDO = Cast<UControlRig>(BlueprintBeingCustomized->GeneratedClass->GetDefaultObject()))
+			{
+				if(CDO->GetVM())
+				{
+					CDO->GetVM()->ClearExternalVariables();
+					TArray<FRigVMExternalVariable> ExternalVariables = CDO->GetExternalVariables();
+					for(const FRigVMExternalVariable& ExternalVariable : ExternalVariables)
+					{
+						CDO->GetVM()->AddExternalVariable(ExternalVariable);
+					}
+					
+					FRigVMCodeGenerator CodeGenerator(ClassName,
+						TEXT("TestModule"), BlueprintBeingCustomized->GetModel(), CDO->GetVM(), BlueprintBeingCustomized->PinToOperandMap);
+					const FString Content = CodeGenerator.DumpHeader() + TEXT("\r\n\r\n") + CodeGenerator.DumpSource();
+					FPlatformApplicationMisc::ClipboardCopy(*Content);
+				}
+			}
+		}
+	}
+	return FReply::Handled();
+}
+
 #undef LOCTEXT_NAMESPACE
diff --git a/Engine/Plugins/Animation/ControlRig/Source/ControlRigEditor/Private/ControlRigCompilerDetails.h b/Engine/Plugins/Animation/ControlRig/Source/ControlRigEditor/Private/ControlRigCompilerDetails.h
--- a/Engine/Plugins/Animation/ControlRig/Source/ControlRigEditor/Private/ControlRigCompilerDetails.h
+++ b/Engine/Plugins/Animation/ControlRig/Source/ControlRigEditor/Private/ControlRigCompilerDetails.h
@@ -28,6 +28,7 @@ protected:
 	FReply OnCopyASTClicked();
 	FReply OnCopyByteCodeClicked();
 	FReply OnCopyHierarchyGraphClicked();
+	FReply OnCopyGeneratedCodeClicked();
 	
 	UControlRigBlueprint* BlueprintBeingCustomized;
 };
diff --git a/Engine/Source/Developer/RigVMDeveloper/Private/RigVMCompiler/RigVMCodeGenerator.cpp b/Engine/Source/Developer/RigVMDeveloper/Private/RigVMCompiler/RigVMCodeGenerator.cpp
new file mode 100644
--- /dev/null
+++ b/Engine/Source/Developer/RigVMDeveloper/Private/RigVMCompiler/RigVMCodeGenerator.cpp
@@ -0,0 +1,2154 @@
+// Copyright Epic Games, Inc. All Rights Reserved.
+
+#include "RigVMCompiler/RigVMCodeGenerator.h"
+#include "RigVMCore/RigVMExecuteContext.h"
+#include "RigVMDeveloperModule.h"
+#include "RigVMDeveloperTypeUtils.h"
+#include "Algo/Count.h"
+#include "Animation/Rig.h"
+
+static constexpr TCHAR RigVM_NewLineFormat[] = TEXT("\r\n");
+static constexpr TCHAR RigVM_IncludeBracketFormat[] = TEXT("#include <{0}>");
+static constexpr TCHAR RigVM_IncludeQuoteFormat[] = TEXT("#include \"{0}.h\"");
+static constexpr TCHAR RigVM_WrappedArrayTypeFormat[] = TEXT("struct {0}_API {1}\r\n{\r\n\tTArray<{2}> Array;\r\n};");
+static constexpr TCHAR RigVM_WrappedTypeNameFormat[] = TEXT("{0}Array_{1}");
+static constexpr TCHAR RigVM_DeclareExternalVariableFormat[] = TEXT("\t{0}* {1} = nullptr;");
+static constexpr TCHAR RigVM_UpdateExternalVariableFormat[] = TEXT("\t{0} = &GetExternalVariableRef<{1}>(TEXT(\"{2}\"), TEXT(\"{1}\"));");
+static constexpr TCHAR RigVM_MemberPropertyFormat[] = TEXT("\t{0} {1};");
+static constexpr TCHAR RigVM_DeclareEntryNameFormat[] = TEXT("\tstatic const FName EntryName_{0};");
+static constexpr TCHAR RigVM_DefineEntryNameFormat[] = TEXT("const FName U{0}::EntryName_{1} = TEXT(\"{1}\");");
+static constexpr TCHAR RigVM_DefineConstFormatNoDefault[] = TEXT("\tstatic const {0} {1};");
+static constexpr TCHAR RigVM_DefineStructConstFormat[] = TEXT("\tstatic const {0} {1} = URigVMNativized::GetStructConstant<{0}>(TEXT(\"{2}\"));");
+static constexpr TCHAR RigVM_DefineStructArrayConstFormat[] = TEXT("\tstatic const {0} {1} = URigVMNativized::GetStructArrayConstant<{2}>(TEXT(\"{3}\"));");
+static constexpr TCHAR RigVM_DefineConstFormat[] = TEXT("\tstatic const {0} {1} = {2};");
+static constexpr TCHAR RigVM_NameNoneFormat[] = TEXT("FName(NAME_None)");
+static constexpr TCHAR RigVM_EmptyStringFormat[] = TEXT("FString()");
+static constexpr TCHAR RigVM_SingleStringFormat[] = TEXT("%s");
+static constexpr TCHAR RigVM_TextFormat[] = TEXT("TEXT({0})");
+static constexpr TCHAR RigVM_QuotedTextFormat[] = TEXT("TEXT(\"{0}\")");
+static constexpr TCHAR RigVM_CurlyBracesFormat[] = TEXT("{{0}}");
+static constexpr TCHAR RigVM_BracesFormat[] = TEXT("({0})");
+static constexpr TCHAR RigVM_TemplateOneArgFormat[] = TEXT("<{0}>");
+static constexpr TCHAR RigVM_CallExternOpFormat[] = TEXT("\t{0}::Static{1}({2});");
+static constexpr TCHAR RigVM_ZeroOpFormat[] = TEXT("\t{0} = 0;");
+static constexpr TCHAR RigVM_BoolFalseOpFormat[] = TEXT("\t{0} = false;");
+static constexpr TCHAR RigVM_BoolTrueFormat[] = TEXT("\t{0} = true;");
+static constexpr TCHAR RigVM_CopyUnrelatedArraysFormat[] = TEXT("\tCopyUnrelatedArrays<{0}, {1}>({2}, {3});");
+static constexpr TCHAR RigVM_CopyOpMethodFormat[] = TEXT("\t{0}{1}{2});");
+static constexpr TCHAR RigVM_CopyOpAssignFormat[] = TEXT("\t{0} = {1}{2};");
+static constexpr TCHAR RigVM_IncrementOpFormat[] = TEXT("\t{0}++;");
+static constexpr TCHAR RigVM_DecrementOpFormat[] = TEXT("\t{0}--;");
+static constexpr TCHAR RigVM_EqualsOpFormat[] = TEXT("\t{0} = {1} == {2};");
+static constexpr TCHAR RigVM_InvokeEntryFormat[] = TEXT("\tif (InEntryName == EntryName_{0}) return ExecuteEntry_{0}({1});");
+static constexpr TCHAR RigVM_EntryNameFormat[] = TEXT("EntryName_{0}");
+static constexpr TCHAR RigVM_UpdateContextFormat[] = TEXT("\tconst FRigVMExecuteContext& PublicContext = UpdateContext(AdditionalArguments, {0}, InEntryName);");
+static constexpr TCHAR RigVM_TrueFormat[] = TEXT("true");  
+static constexpr TCHAR RigVM_FalseFormat[] = TEXT("false");
+static constexpr TCHAR RigVM_SingleUnderscoreFormat[] = TEXT("_");
+static constexpr TCHAR RigVM_DoubleUnderscoreFormat[] = TEXT("__");
+static constexpr TCHAR RigVM_BoolPropertyPrefix[] = TEXT("b");
+static constexpr TCHAR RigVM_EnumTypeSuffixFormat[] = TEXT("::Type");
+static constexpr TCHAR RigVM_IsValidArraySizeFormat[] = TEXT("IsValidArraySize({0})");
+static constexpr TCHAR RigVM_IsValidArrayIndexFormat[] = TEXT("IsValidArrayIndex<{0}>(TemporaryArrayIndex, {1})");
+static constexpr TCHAR RigVM_TemporaryArrayIndexFormat[] = TEXT("\tTemporaryArrayIndex = {0};");
+static constexpr TCHAR RigVM_ExecuteReachedExitFormat[] = TEXT("\tBroadcastExecutionReachedExit();");
+static constexpr TCHAR RigVM_AdditionalArgumentNameFormat[] = TEXT("AdditionalArgument_{0}_{1}");
+static constexpr TCHAR RigVM_FullOpaqueArgumentFormat[] = TEXT("\t{0} {1} = *({2}*)Context.OpaqueArguments[{3}];");
+static constexpr TCHAR RigVM_ParameterOpaqueArgumentFormat[] = TEXT("{0} {1}");
+static constexpr TCHAR RigVM_InstructionLabelFormat[] = TEXT("\tInstruction{0}Label:");
+static constexpr TCHAR RigVM_SetInstructionIndexFormat[] = TEXT("\tSetInstructionIndex({0});");
+static constexpr TCHAR RigVM_ContextPublicFormat[] = TEXT("PublicContext");
+static constexpr TCHAR RigVM_ContextPublicParameterFormat[] = TEXT("const FRigVMExecuteContext& PublicContext");
+static constexpr TCHAR RigVM_NotEqualsOpFormat[] = TEXT("\t{0} = {1} != {2};");
+static constexpr TCHAR RigVM_JumpOpFormat[] = TEXT("\tgoto Instruction{0}Label;");
+static constexpr TCHAR RigVM_JumpIfOpFormat[] = TEXT("\tif ({0} == {1}) { goto Instruction{2}Label; }");
+static constexpr TCHAR RigVM_BeginBlockOpFormat[] = TEXT("\tBeginSlice({0}, {1});");
+static constexpr TCHAR RigVM_EndBlockOpFormat[] = TEXT("\tEndSlice();");
+static constexpr TCHAR RigVM_ArrayResetOpFormat[] = TEXT("\t{0}.Reset();");
+static constexpr TCHAR RigVM_ArrayGetNumOpFormat[] = TEXT("\t{0} = {1}.Num();");
+static constexpr TCHAR RigVM_ArraySetNumOpFormat[] = TEXT("\tif ({0}) { {1}.SetNum({2}); }");
+static constexpr TCHAR RigVM_ArrayGetAtIndexOpFormat[] = TEXT("\tif ({0}) { {1} = {2}[TemporaryArrayIndex]; }");
+static constexpr TCHAR RigVM_ArraySetAtIndexOpFormat[] = TEXT("\tif ({0}) { {1}[TemporaryArrayIndex] = {2}; }");
+static constexpr TCHAR RigVM_ArrayAddOpFormat[] = TEXT("\tif ({0}) { {1} = {2}.Add({3}); }");
+static constexpr TCHAR RigVM_ArrayNumPlusOneFormat[] = TEXT("{0}.Num() + 1");
+static constexpr TCHAR RigVM_ArrayInsertOpFormat[] = TEXT("\tif ({0}) { {1}.Insert({2}, {3}); }");
+static constexpr TCHAR RigVM_ArrayRemoveOpFormat[] = TEXT("\t{0}.RemoveAt({1});");
+static constexpr TCHAR RigVM_ArrayFindOpFormat[] = TEXT("\t{0} = ArrayOp_Find{1}({2}, {3}, {4});");
+static constexpr TCHAR RigVM_ArrayAppendOpFormat[] = TEXT("\tif ({0} { {1}.Append({2}); }");
+static constexpr TCHAR RigVM_ArrayAppendOpNumFormat[] = TEXT("{0}.Num() + {1}.Num()");
+static constexpr TCHAR RigVM_ArrayIteratorOpFormat[] = TEXT("\t{0} = ArrayOp_Iterator{1}({2}, {3}, {4}, {5}, {6});");
+static constexpr TCHAR RigVM_ArrayUnionOpFormat[] = TEXT("\tArrayOp_Union{0}({1}, {2});");
+static constexpr TCHAR RigVM_ArrayDifferenceOpFormat[] = TEXT("\tArrayOp_Difference{0}({1}, {2});");
+static constexpr TCHAR RigVM_ArrayIntersectionOpFormat[] = TEXT("\tArrayOp_Intersection{0}({1}, {2});");
+static constexpr TCHAR RigVM_ArrayReverseOpFormat[] = TEXT("\tArrayOp_Reverse{0}({1});");
+static constexpr TCHAR RigVM_ReturnFalseFormat[] = TEXT("\treturn false;");
+static constexpr TCHAR RigVM_ReturnTrueFormat[] = TEXT("\treturn true;");
+static constexpr TCHAR RigVM_CopyrightFormat[] = TEXT("// Copyright Epic Games, Inc. All Rights Reserved.");
+static constexpr TCHAR RigVM_AutoGeneratedFormat[] = TEXT("// THIS FILE HAS BEEN AUTO-GENERATED. PLEASE DO NOT MANUALLY EDIT THIS FILE FURTHER.");
+static constexpr TCHAR RigVM_PragmaOnceFormat[] = TEXT("#pragma once");
+static constexpr TCHAR RigVM_GeneratedIncludeFormat[] = TEXT("#include \"{0}.generated.h\"");
+static constexpr TCHAR RigVM_UClassDefinitionFormat[] = TEXT("UCLASS()\r\nclass {0}_API U{1} : public URigVMNativized\r\n{\r\n\tGENERATED_BODY()\r\npublic:\r\n\tU{1}() {}\r\n\tvirtual ~U{1}() override {}\r\n");
+static constexpr TCHAR RigVM_ProtectedFormat[] = TEXT("protected:");
+static constexpr TCHAR RigVM_GetVMHashFormat[] = TEXT("\tvirtual uint32 GetVMHash() const override { return {0}; }");
+static constexpr TCHAR RigVM_GetEntryNamesFormat[] = TEXT("\tvirtual TArray<FName> GetEntryNames() const override { return { {0} }; }");
+static constexpr TCHAR RigVM_DeclareUpdateExternalVariablesFormat[] = TEXT("\tvirtual void UpdateExternalVariables() override;");
+static constexpr TCHAR RigVM_DeclareExecuteFormat[] = TEXT("\tvirtual bool Execute(TArrayView<URigVMMemoryStorage*> Memory, TArrayView<void*> AdditionalArguments, const FName& InEntryName) override;");
+static constexpr TCHAR RigVM_DefineUpdateExternalVariablesFormat[] = TEXT("void U{0}::UpdateExternalVariables()\r\n{");
+static constexpr TCHAR RigVM_DefineExecuteFormat[] = TEXT("bool U{0}::Execute(TArrayView<URigVMMemoryStorage*> Memory, TArrayView<void*> AdditionalArguments, const FName& InEntryName)\r\n{");
+static constexpr TCHAR RigVM_DeclareExecuteEntryFormat[] = TEXT("\tbool ExecuteEntry_{0}({1});");
+static constexpr TCHAR RigVM_DefineExecuteEntryFormat[] = TEXT("bool U{0}::ExecuteEntry_{1}({2})\r\n{");
+static constexpr TCHAR RigVM_DeclareExecuteGroupFormat[] = TEXT("\tbool ExecuteGroup_{0}_{1}({2});");
+static constexpr TCHAR RigVM_DefineExecuteGroupFormat[] = TEXT("bool U{0}::ExecuteGroup_{1}_{2}({3})\r\n{");
+static constexpr TCHAR RigVM_InvokeExecuteGroupFormat[] = TEXT("\tif(!ExecuteGroup_{0}_{1}({2})) return false;");
+static constexpr TCHAR RigVM_RigVMCoreIncludeFormat[] = TEXT("RigVMCore/RigVMCore.h");
+static constexpr TCHAR RigVM_RigVMCoreLibraryFormat[] = TEXT("RigVM");
+static constexpr TCHAR RigVM_JoinFilePathFormat[] = TEXT("{0}/{1}");
+static constexpr TCHAR RigVM_GetOperandSliceFormat[] = TEXT("GetOperandSlice<{0}>({1}){2}");
+static constexpr TCHAR RigVM_ExternalVariableFormat[] = TEXT("(*External_{0})");
+static constexpr TCHAR RigVM_JoinSegmentPathFormat[] = TEXT("{0}.{1}");
+static constexpr TCHAR RigVM_GetArrayElementSafeFormat[] = TEXT("GetArrayElementSafe<{0}>({1}, {2})");
+
+FString FRigVMCodeGenerator::DumpIncludes(bool bLog)
+{
+	FStringArray Lines;
+	for(const FString& Include : Includes)
+	{
+		Lines.Add(Format(RigVM_IncludeBracketFormat, Include));
+	}
+	return DumpLines(Lines, bLog);
+}
+
+FString FRigVMCodeGenerator::DumpWrappedArrayTypes(bool bLog)
+{
+	FStringArray Lines;
+	for(const FString& WrappedArrayType : WrappedArrayCPPTypes)
+	{
+		if (!Lines.IsEmpty())
+		{
+			Lines.Add(FString());
+		}
+
+		const FString TypeName = GetMappedArrayTypeName(WrappedArrayType);
+		const FString Line = Format(RigVM_WrappedArrayTypeFormat, *ModuleName.ToUpper(), *TypeName, *WrappedArrayType);
+		Lines.Add(Line);
+	}
+
+	return DumpLines(Lines, bLog);
+}
+
+FString FRigVMCodeGenerator::DumpExternalVariables(bool bForHeader, bool bLog)
+{
+
+	FStringArray Lines;
+	for(int32 ExternalVariableIndex = 0; ExternalVariableIndex < VM->GetExternalVariables().Num(); ExternalVariableIndex++)
+	{
+		const FRigVMOperand ExternalVarOperand(ERigVMMemoryType::External, ExternalVariableIndex, INDEX_NONE);
+		const FRigVMExternalVariable& ExternalVariable = VM->GetExternalVariables()[ExternalVariableIndex];
+		const FString ExternalVarCPPType = ExternalVariable.GetExtendedCPPType().ToString();
+		FString OperandName = *GetOperandName(ExternalVarOperand, false);
+		if(OperandName.StartsWith(TEXT("(*")) && OperandName.EndsWith(TEXT(")")))
+		{
+			OperandName = OperandName.Mid(2, OperandName.Len() - 3);
+		}
+
+		if(bForHeader)
+		{
+			Lines.Add(Format(RigVM_DeclareExternalVariableFormat, *ExternalVarCPPType, *OperandName));
+		}
+		else
+		{
+			Lines.Add(Format(RigVM_UpdateExternalVariableFormat, *OperandName, *ExternalVarCPPType, *ExternalVariable.Name.ToString(), *ExternalVarCPPType));
+		}
+	}
+
+	return DumpLines(Lines, bLog);
+}
+
+FString FRigVMCodeGenerator::DumpEntries(bool bForHeader, bool bLog)
+{
+	FStringArray Lines;
+
+	const FRigVMByteCode& ByteCode = VM->GetByteCode();
+	for(int32 EntryIndex = 0; EntryIndex < ByteCode.NumEntries(); EntryIndex++)
+	{
+		const FRigVMByteCodeEntry& Entry = ByteCode.GetEntry(EntryIndex);
+		const FString EntryName = Entry.Name.ToString();
+		if(bForHeader)
+		{
+			Lines.Add(Format(RigVM_DeclareEntryNameFormat, *EntryName));
+		}
+		else
+		{
+			Lines.Add(Format(RigVM_DefineEntryNameFormat, *ClassName, *EntryName));
+		}
+	}
+
+	return DumpLines(Lines, bLog);
+}
+
+FString FRigVMCodeGenerator::DumpProperties(bool bForHeader, int32 InOperationGroup, bool bLog)
+{
+	if(bForHeader)
+	{
+		// for headers we show all properties
+		check(InOperationGroup == INDEX_NONE);
+	}
+
+	FStringArray Lines;
+	for(int32 Index = 0; Index < Properties.Num(); Index++)
+	{
+		if(!IsPropertyPartOfGroup(Index, InOperationGroup))
+		{
+			continue;
+		}
+
+		const FPropertyInfo& PropertyInfo = Properties[Index];
+		
+		// in headers we only dump the work / sliced properties,
+		// and for source files we only dump the non-sliced
+		if(bForHeader != (PropertyInfo.PropertyType == ERigVMNativizedPropertyType::Sliced ||
+							(PropertyInfo.PropertyType == ERigVMNativizedPropertyType::Work &&
+								PropertyInfo.Groups.Num() > 1)))
+		{
+			continue;
+		}
+
+		const FRigVMPropertyDescription& Property = PropertyInfo.Description;
+		check(Property.IsValid());
+
+		if(PropertyInfo.PropertyType == ERigVMNativizedPropertyType::Literal)
+		{
+			FRigVMOperand Operand(ERigVMMemoryType::Literal, PropertyInfo.MemoryPropertyIndex, INDEX_NONE);
+			FString OperandName = GetOperandName(Operand, false);
+			FString CPPType = GetOperandCPPType(Operand);
+
+			FString BaseCPPType = CPPType;
+			const bool bIsArray = RigVMTypeUtils::IsArrayType(CPPType);
+			if (bIsArray)
+			{
+				BaseCPPType = RigVMTypeUtils::BaseTypeFromArrayType(CPPType);
+			}
+			
+			FString DefaultValue = Property.DefaultValue;
+
+			if (UScriptStruct* ScriptStruct = Cast<UScriptStruct>(Property.CPPTypeObject))
+			{
+				DefaultValue = DefaultValue.ReplaceCharWithEscapedChar();
+
+				if(DefaultValue.IsEmpty())
+				{
+					Lines.Add(Format(
+						RigVM_DefineConstFormatNoDefault,
+						*CPPType,
+						*GetOperandName(Operand, false)
+					));
+				}
+				else
+				{
+					if(const TStructConstGenerator* StructConstGenerator = GetStructConstGenerators().Find(*BaseCPPType))
+					{
+						if(bIsArray)
+						{
+							FStringArray DefaultValues = URigVMPin::SplitDefaultValue(DefaultValue);
+							for(int32 DefaultValueIndex = 0; DefaultValueIndex < DefaultValues.Num(); DefaultValueIndex++)
+							{
+								DefaultValues[DefaultValueIndex] = (*StructConstGenerator)(DefaultValues[DefaultValueIndex]); 
+							}
+							DefaultValue = Format(RigVM_CurlyBracesFormat, FString::Join(DefaultValues, TEXT(", ")));
+						}
+						else
+						{
+							DefaultValue = (*StructConstGenerator)(DefaultValue);
+						}
+						
+						Lines.Add(Format(
+							RigVM_DefineConstFormat,
+							*CPPType,
+							*OperandName,
+							*DefaultValue
+						));
+					}
+					else if (bIsArray)
+					{
+						Lines.Add(Format(
+							RigVM_DefineStructArrayConstFormat,
+							*CPPType,
+							*OperandName,
+							*BaseCPPType,
+							*DefaultValue
+						));
+					}
+					else
+					{
+						Lines.Add(Format(
+							RigVM_DefineStructConstFormat,
+							*CPPType,
+							*OperandName,
+							*DefaultValue
+						));
+					}
+				}
+			}
+			else if (const UEnum* Enum = Cast<UEnum>(Property.CPPTypeObject))
+			{
+				BaseCPPType = Enum->GetName();
+				if (Enum->GetCppForm() == UEnum::ECppForm::Namespaced)
+				{
+					BaseCPPType += RigVM_EnumTypeSuffixFormat;
+				}
+				
+				CPPType = BaseCPPType;
+				if (bIsArray)
+				{
+					CPPType = RigVMTypeUtils::ArrayTypeFromBaseType(BaseCPPType);
+				}
+
+				FStringArray DefaultValues;
+				if (bIsArray)
+				{
+					DefaultValues = URigVMPin::SplitDefaultValue(DefaultValue);
+				}
+				else
+				{
+					DefaultValues.Add(DefaultValue);
+				}
+
+				for(int32 DefaultValueIndex = 0; DefaultValueIndex < DefaultValues.Num(); DefaultValueIndex++)
+				{
+					if (DefaultValues[DefaultValueIndex].IsNumeric())
+					{
+						const int64 EnumIndex = FCString::Atoi64(*DefaultValues[DefaultValueIndex]);
+						const FString EnumName = Enum->GetNameStringByValue(EnumIndex);
+						DefaultValues[DefaultValueIndex] = Enum->GenerateFullEnumName(*EnumName);
+					}
+					else if (!UEnum::IsFullEnumName(*DefaultValues[DefaultValueIndex]))
+					{
+						DefaultValues[DefaultValueIndex] = Enum->GenerateFullEnumName(*DefaultValues[DefaultValueIndex]);
+					}
+				}
+				
+				if (bIsArray)
+				{
+					DefaultValue = Format(RigVM_CurlyBracesFormat, *FString::Join(DefaultValues, TEXT(",")));
+				}
+				else
+				{
+					DefaultValue = DefaultValues[0];
+				}
+
+
+				if(DefaultValue.IsEmpty())
+				{
+					Lines.Add(Format(
+						RigVM_DefineConstFormatNoDefault,
+						*CPPType,
+						*OperandName
+					));
+				}
+				else
+				{
+					Lines.Add(Format(
+						RigVM_DefineConstFormat,
+						*CPPType,
+						*OperandName,
+						*DefaultValue
+					));
+				}
+			}
+			else
+			{
+				bool bUseConstExpr = true;
+
+				if (bIsArray)
+				{
+					bUseConstExpr = false;
+					
+					if (DefaultValue.StartsWith(TEXT("(")) && DefaultValue.EndsWith(TEXT(")")))
+					{
+						DefaultValue = Format(RigVM_CurlyBracesFormat, *DefaultValue.Mid(1, DefaultValue.Len() - 2));
+					}
+				}
+
+				if (BaseCPPType == RigVMTypeUtils::BoolType)
+				{
+					DefaultValue.ToLowerInline();
+				}
+
+				if (CPPType == RigVMTypeUtils::FNameType || CPPType == RigVMTypeUtils::FStringType)
+				{
+					if(DefaultValue.IsEmpty())
+					{
+						if(CPPType == RigVMTypeUtils::FNameType)
+						{
+							DefaultValue = RigVM_NameNoneFormat;
+						}
+						else
+						{
+							DefaultValue = RigVM_EmptyStringFormat;
+						}
+					}
+					else
+					{
+						if (DefaultValue.StartsWith(TEXT("\"")) && DefaultValue.EndsWith(TEXT("\"")))
+						{
+							DefaultValue = Format(RigVM_TextFormat, *DefaultValue);
+						}
+						else
+						{
+							DefaultValue = Format(RigVM_QuotedTextFormat, *DefaultValue);
+						}
+					}
+					bUseConstExpr = false;
+				}
+
+				if(DefaultValue.IsEmpty())
+				{
+					Lines.Add(Format(
+						RigVM_DefineConstFormatNoDefault, 
+						*CPPType,
+						*OperandName
+					));
+				}
+				else
+				{
+					Lines.Add(Format(
+						RigVM_DefineConstFormat, 
+						*CPPType,
+						*OperandName,
+						*DefaultValue
+					));
+				}
+			}
+		}
+		else // work and slice look the same in the file
+		{
+			FRigVMOperand Operand(ERigVMMemoryType::Work, PropertyInfo.MemoryPropertyIndex, INDEX_NONE);
+			FString OperandName = GetOperandName(Operand, false);
+			FString CPPType = GetOperandCPPType(Operand);
+
+			const FString MappedType = GetMappedType(Property.CPPType);
+			const FString Line = Format(RigVM_MemberPropertyFormat, *MappedType, *SanitizeName(Property.Name.ToString(), Property.CPPType));
+			Lines.Add(Line);
+		}
+	}
+
+	return DumpLines(Lines, bLog);
+}
+
+FString FRigVMCodeGenerator::DumpInstructions(int32 InOperationGroup, bool bLog)
+{
+	const FRigVMByteCode& ByteCode = VM->GetByteCode();
+	const TArray<FName>& Functions = VM->GetFunctionNames();
+	const FRigVMInstructionArray Operations = ByteCode.GetInstructions();
+
+	FStringArray Lines;
+
+	const FOperationGroup& Group = GetGroup(InOperationGroup);
+
+	if(Group.Entry.IsNone())
+	{
+		Lines.Add(Format(RigVM_UpdateContextFormat, Operations.Num()));
+
+		// let's get the additional arguments for all sub groups
+		TArray<FString> OpaqueArgumentHit;
+		for(const FOperationGroup& EntryGroup : OperationGroups)
+		{
+			if(EntryGroup.Entry.IsNone() || EntryGroup.Depth > 0)
+			{
+				continue;
+			}
+
+			for(const FOpaqueArgument& OpaqueArgument : EntryGroup.OpaqueArguments)
+			{
+				if(OpaqueArgumentHit.Contains(OpaqueArgument.Get<0>()))
+				{
+					continue;
+				}
+				const int32 ArgumentIndex = OpaqueArgumentHit.Add(OpaqueArgument.Get<0>());
+				FString ArgumentTypeNoRef = OpaqueArgument.Get<2>();
+				if (ArgumentTypeNoRef.EndsWith(TEXT("&")))
+				{
+					ArgumentTypeNoRef.LeftChopInline(1);
+				}
+				Lines.Add(Format(RigVM_FullOpaqueArgumentFormat, *OpaqueArgument.Get<2>(), *OpaqueArgument.Get<1>(), *ArgumentTypeNoRef, ArgumentIndex));
+			}
+		}
+
+		Lines.Emplace();
+		for(int32 EntryIndex = 0; EntryIndex < ByteCode.NumEntries(); EntryIndex++)
+		{
+			const FRigVMByteCodeEntry& Entry = ByteCode.GetEntry(EntryIndex);
+			const FString EntryName = Entry.Name.ToString();
+			FString OpaqueArguments;
+
+			// find the entry's group and provide the needed arguments
+			for(const FOperationGroup& EntryGroup : OperationGroups)
+			{
+				if(EntryGroup.Depth <= 0 && EntryGroup.Entry == Entry.Name)
+				{
+					TArray<FString> ArgumentNames = {RigVM_ContextPublicFormat};
+					for(const FOpaqueArgument& OpaqueArgument : EntryGroup.OpaqueArguments)
+					{
+						ArgumentNames.Add(OpaqueArgument.Get<1>());
+					}
+					OpaqueArguments = FString::Join(ArgumentNames, TEXT(", "));
+					break;
+				}
+			}
+			Lines.Add(Format(RigVM_InvokeEntryFormat, *EntryName, *OpaqueArguments));
+		}
+	}
+
+	if(Group.Entry.IsNone() && OperationGroups.Num() > 0)
+	{
+		Lines.Add(RigVM_ReturnFalseFormat);
+		return DumpLines(Lines, bLog);
+	}
+
+	if(Group.ChildGroups.IsEmpty())
+	{
+		// loop over all of the operations once more
+		for(int32 OperationIndex = Group.First; OperationIndex <= Group.Last; OperationIndex++)
+		{
+			// inject a label if required
+			if (Group.RequiredLabels.Contains(OperationIndex))
+			{
+				Lines.Add(Format(RigVM_InstructionLabelFormat, OperationIndex));
+			}
+
+			Lines.Add(Format(RigVM_SetInstructionIndexFormat, OperationIndex));
+
+			const FRigVMInstruction& Operation = Operations[OperationIndex];
+			switch(Operation.OpCode)
+			{
+				case ERigVMOpCode::Execute_0_Operands:
+				case ERigVMOpCode::Execute_1_Operands:
+				case ERigVMOpCode::Execute_2_Operands:
+				case ERigVMOpCode::Execute_3_Operands:
+				case ERigVMOpCode::Execute_4_Operands:
+				case ERigVMOpCode::Execute_5_Operands:
+				case ERigVMOpCode::Execute_6_Operands:
+				case ERigVMOpCode::Execute_7_Operands:
+				case ERigVMOpCode::Execute_8_Operands:
+				case ERigVMOpCode::Execute_9_Operands:
+				case ERigVMOpCode::Execute_10_Operands:
+				case ERigVMOpCode::Execute_11_Operands:
+				case ERigVMOpCode::Execute_12_Operands:
+				case ERigVMOpCode::Execute_13_Operands:
+				case ERigVMOpCode::Execute_14_Operands:
+				case ERigVMOpCode::Execute_15_Operands:
+				case ERigVMOpCode::Execute_16_Operands:
+				case ERigVMOpCode::Execute_17_Operands:
+				case ERigVMOpCode::Execute_18_Operands:
+				case ERigVMOpCode::Execute_19_Operands:
+				case ERigVMOpCode::Execute_20_Operands:
+				case ERigVMOpCode::Execute_21_Operands:
+				case ERigVMOpCode::Execute_22_Operands:
+				case ERigVMOpCode::Execute_23_Operands:
+				case ERigVMOpCode::Execute_24_Operands:
+				case ERigVMOpCode::Execute_25_Operands:
+				case ERigVMOpCode::Execute_26_Operands:
+				case ERigVMOpCode::Execute_27_Operands:
+				case ERigVMOpCode::Execute_28_Operands:
+				case ERigVMOpCode::Execute_29_Operands:
+				case ERigVMOpCode::Execute_30_Operands:
+				case ERigVMOpCode::Execute_31_Operands:
+				case ERigVMOpCode::Execute_32_Operands:
+				case ERigVMOpCode::Execute_33_Operands:
+				case ERigVMOpCode::Execute_34_Operands:
+				case ERigVMOpCode::Execute_35_Operands:
+				case ERigVMOpCode::Execute_36_Operands:
+				case ERigVMOpCode::Execute_37_Operands:
+				case ERigVMOpCode::Execute_38_Operands:
+				case ERigVMOpCode::Execute_39_Operands:
+				case ERigVMOpCode::Execute_40_Operands:
+				case ERigVMOpCode::Execute_41_Operands:
+				case ERigVMOpCode::Execute_42_Operands:
+				case ERigVMOpCode::Execute_43_Operands:
+				case ERigVMOpCode::Execute_44_Operands:
+				case ERigVMOpCode::Execute_45_Operands:
+				case ERigVMOpCode::Execute_46_Operands:
+				case ERigVMOpCode::Execute_47_Operands:
+				case ERigVMOpCode::Execute_48_Operands:
+				case ERigVMOpCode::Execute_49_Operands:
+				case ERigVMOpCode::Execute_50_Operands:
+				case ERigVMOpCode::Execute_51_Operands:
+				case ERigVMOpCode::Execute_52_Operands:
+				case ERigVMOpCode::Execute_53_Operands:
+				case ERigVMOpCode::Execute_54_Operands:
+				case ERigVMOpCode::Execute_55_Operands:
+				case ERigVMOpCode::Execute_56_Operands:
+				case ERigVMOpCode::Execute_57_Operands:
+				case ERigVMOpCode::Execute_58_Operands:
+				case ERigVMOpCode::Execute_59_Operands:
+				case ERigVMOpCode::Execute_60_Operands:
+				case ERigVMOpCode::Execute_61_Operands:
+				case ERigVMOpCode::Execute_62_Operands:
+				case ERigVMOpCode::Execute_63_Operands:
+				case ERigVMOpCode::Execute_64_Operands:
+				{
+					const FRigVMExecuteOp& Op = ByteCode.GetOpAt<FRigVMExecuteOp>(Operation);
+					FRigVMOperandArray Operands = ByteCode.GetOperandsForExecuteOp(Operation);
+
+					const FRigVMFunction* Function = FRigVMRegistry::Get().FindFunction(*Functions[Op.FunctionIndex].ToString());
+					check(Function);
+						
+					FStringArray Arguments;
+					Arguments.Add(RigVM_ContextPublicFormat);
+					for(int32 OperandIndex = 0; OperandIndex < Operands.Num(); OperandIndex++)
+					{
+						const FRigVMOperand& Operand = Operands[OperandIndex];
+						bool bSliced = false;
+						const FRigVMPropertyDescription& Property = GetPropertyForOperand(Operand);
+						if (RigVMTypeUtils::IsArrayType(Property.CPPType))
+						{
+							const FRigVMFunctionArgument& FunctionArgument = Function->GetArguments()[OperandIndex];
+							bSliced = FunctionArgument.Type != Property.CPPType;
+						}
+						Arguments.Add(GetOperandName(Operand, bSliced));
+					}
+
+					for(const FRigVMFunctionArgument& Argument : Function->GetArguments())
+					{
+						if (Function->IsAdditionalArgument(Argument))
+						{
+							const FOpaqueArgument* OpaqueArgument = Group.OpaqueArguments.FindByPredicate([Argument](const FOpaqueArgument& InArgument) -> bool
+							{
+								return InArgument.Get<0>() == Argument.Name;
+							});
+
+							check(OpaqueArgument);
+							Arguments.Add(OpaqueArgument->Get<1>());
+						}
+					}
+
+					const FString JoinedArguments = FString::Join(Arguments, TEXT(", "));
+					Lines.Add(Format(RigVM_CallExternOpFormat, *Function->Struct->GetStructCPPName(), *Function->GetMethodName().ToString(), *JoinedArguments));
+					break;
+				}
+				case ERigVMOpCode::Zero:
+				{
+					const FRigVMUnaryOp& Op = ByteCode.GetOpAt<FRigVMUnaryOp>(Operation);
+					Lines.Add(Format(RigVM_ZeroOpFormat, *GetOperandName(Op.Arg, false)));
+					break;
+				}
+				case ERigVMOpCode::BoolFalse:
+				{
+					const FRigVMUnaryOp& Op = ByteCode.GetOpAt<FRigVMUnaryOp>(Operation);
+					Lines.Add(Format(RigVM_BoolFalseOpFormat, *GetOperandName(Op.Arg, false)));
+					break;
+				}
+				case ERigVMOpCode::BoolTrue:
+				{
+					const FRigVMUnaryOp& Op = ByteCode.GetOpAt<FRigVMUnaryOp>(Operation);
+					Lines.Add(Format(RigVM_BoolTrueFormat, *GetOperandName(Op.Arg, false)));
+					break;
+				}
+				case ERigVMOpCode::Copy:
+				case ERigVMOpCode::ArrayClone:
+				{
+					const FRigVMBinaryOp& Op = ByteCode.GetOpAt<FRigVMBinaryOp>(Operation);
+					const FString TargetOperand = GetOperandName(Op.ArgB, false, false);
+					const FString SourceOperand = GetOperandName(Op.ArgA, false, true);
+					const FString TargetCPPType = GetOperandCPPType(Op.ArgB);
+					const FString SourceCPPType = GetOperandCPPType(Op.ArgA);
+
+					if(RigVMTypeUtils::IsArrayType(TargetCPPType) &&
+						RigVMTypeUtils::IsArrayType(SourceCPPType) &&
+						TargetCPPType != SourceCPPType)
+					{
+						const FString TargetBaseCPPType = RigVMTypeUtils::BaseTypeFromArrayType(TargetCPPType);
+						const FString SourceBaseCPPType = RigVMTypeUtils::BaseTypeFromArrayType(SourceCPPType);
+						Lines.Add(Format(RigVM_CopyUnrelatedArraysFormat, *TargetBaseCPPType, *SourceBaseCPPType, *TargetOperand, *SourceOperand));
+					}
+					else
+					{
+						const FString CastPrefix = TargetCPPType != SourceCPPType ? Format(RigVM_BracesFormat, *TargetCPPType) : FString();  
+						
+						if(TargetOperand.EndsWith(TEXT("(")))
+						{
+							Lines.Add(Format(RigVM_CopyOpMethodFormat, *TargetOperand, *CastPrefix, *SourceOperand));
+						}
+						else
+						{
+							Lines.Add(Format(RigVM_CopyOpAssignFormat, *TargetOperand, *CastPrefix, *SourceOperand));
+						}
+					}
+					break;
+				}
+				case ERigVMOpCode::Increment:
+				{
+					const FRigVMUnaryOp& Op = ByteCode.GetOpAt<FRigVMUnaryOp>(Operation);
+					Lines.Add(Format(RigVM_IncrementOpFormat, *GetOperandName(Op.Arg, false)));
+					break;
+				}
+				case ERigVMOpCode::Decrement:
+				{
+					const FRigVMUnaryOp& Op = ByteCode.GetOpAt<FRigVMUnaryOp>(Operation);
+					Lines.Add(Format(RigVM_DecrementOpFormat, *GetOperandName(Op.Arg, false)));
+					break;
+				}
+				case ERigVMOpCode::Equals:
+				{
+					const FRigVMComparisonOp& Op = ByteCode.GetOpAt<FRigVMComparisonOp>(Operation);
+					Lines.Add(Format(RigVM_EqualsOpFormat, *GetOperandName(Op.Result, false), *GetOperandName(Op.B, false), *GetOperandName(Op.B, false)));
+					break;
+				}
+				case ERigVMOpCode::NotEquals:
+				{
+					const FRigVMComparisonOp& Op = ByteCode.GetOpAt<FRigVMComparisonOp>(Operation);
+					Lines.Add(Format(RigVM_NotEqualsOpFormat, *GetOperandName(Op.Result, false), *GetOperandName(Op.B, false), *GetOperandName(Op.B, false)));
+					break;
+				}
+				case ERigVMOpCode::JumpAbsolute:
+				{
+					const FRigVMJumpOp& Op = ByteCode.GetOpAt<FRigVMJumpOp>(Operation);
+					Lines.Add(Format(RigVM_JumpOpFormat, Op.InstructionIndex));
+					break;
+				}
+				case ERigVMOpCode::JumpForward:
+				{
+					const FRigVMJumpOp& Op = ByteCode.GetOpAt<FRigVMJumpOp>(Operation);
+					Lines.Add(Format(RigVM_JumpOpFormat, OperationIndex + Op.InstructionIndex));
+					break;
+				}
+				case ERigVMOpCode::JumpBackward:
+				{
+					const FRigVMJumpOp& Op = ByteCode.GetOpAt<FRigVMJumpOp>(Operation);
+					Lines.Add(Format(RigVM_JumpOpFormat, OperationIndex - Op.InstructionIndex));
+					break;
+				}
+				case ERigVMOpCode::JumpAbsoluteIf:
+				{
+					const FRigVMJumpIfOp& Op = ByteCode.GetOpAt<FRigVMJumpIfOp>(Operation);
+					const FString& Condition = Op.Condition ? RigVM_TrueFormat : RigVM_FalseFormat;
+					Lines.Add(Format(RigVM_JumpIfOpFormat, *GetOperandName(Op.Arg, false), *Condition, Op.InstructionIndex));
+					break;
+				}
+				case ERigVMOpCode::JumpForwardIf:
+				{
+					const FRigVMJumpIfOp& Op = ByteCode.GetOpAt<FRigVMJumpIfOp>(Operation);
+					const FString& Condition = Op.Condition ? RigVM_TrueFormat : RigVM_FalseFormat;
+					Lines.Add(Format(RigVM_JumpIfOpFormat, *GetOperandName(Op.Arg, false), *Condition, OperationIndex + Op.InstructionIndex));
+					break;
+				}
+				case ERigVMOpCode::JumpBackwardIf:
+				{
+					const FRigVMJumpIfOp& Op = ByteCode.GetOpAt<FRigVMJumpIfOp>(Operation);
+					const FString& Condition = Op.Condition ? RigVM_TrueFormat : RigVM_FalseFormat;
+					Lines.Add(Format(RigVM_JumpIfOpFormat, *GetOperandName(Op.Arg, false), *Condition, OperationIndex - Op.InstructionIndex));
+					break;
+				}
+				case ERigVMOpCode::Exit:
+				{
+					if(OperationIndex != Group.Last)
+					{
+						Lines.Add(Format(RigVM_JumpOpFormat, Operations.Num()));
+					}
+					break;
+				}
+				case ERigVMOpCode::BeginBlock:
+				{
+					const FRigVMBinaryOp& Op = ByteCode.GetOpAt<FRigVMBinaryOp>(Operation);
+					Lines.Add(Format(RigVM_BeginBlockOpFormat, *GetOperandName(Op.ArgA, false), *GetOperandName(Op.ArgB, false)));
+					break;
+				}
+				case ERigVMOpCode::EndBlock:
+				{
+					Lines.Add(RigVM_EndBlockOpFormat);
+					break;
+				}
+				case ERigVMOpCode::ArrayReset:
+				{
+					const FRigVMUnaryOp& Op = ByteCode.GetOpAt<FRigVMUnaryOp>(Operation);
+					Lines.Add(Format(RigVM_ArrayResetOpFormat, *GetOperandName(Op.Arg, false)));
+					break;
+				}
+				case ERigVMOpCode::ArrayGetNum:
+				{
+					const FRigVMBinaryOp& Op = ByteCode.GetOpAt<FRigVMBinaryOp>(Operation);
+					Lines.Add(Format(RigVM_ArrayGetNumOpFormat, *GetOperandName(Op.ArgB, false), *GetOperandName(Op.ArgA, false)));
+					break;
+				} 
+				case ERigVMOpCode::ArraySetNum:
+				{
+					const FRigVMBinaryOp& Op = ByteCode.GetOpAt<FRigVMBinaryOp>(Operation);
+					const FString ConditionString = Format(RigVM_IsValidArraySizeFormat, *GetOperandName(Op.ArgB, false));
+					Lines.Add(Format(RigVM_ArraySetNumOpFormat, *ConditionString, *GetOperandName(Op.ArgA, false), *GetOperandName(Op.ArgB, false)));
+					break;
+				}
+				case ERigVMOpCode::ArrayGetAtIndex:
+				{
+					const FRigVMTernaryOp& Op = ByteCode.GetOpAt<FRigVMTernaryOp>(Operation);
+					const FString ConditionString = Format(RigVM_IsValidArrayIndexFormat, *GetOperandCPPBaseType(Op.ArgA), *GetOperandName(Op.ArgA, false));
+					Lines.Add(Format(RigVM_TemporaryArrayIndexFormat, *GetOperandName(Op.ArgB, false)));
+					Lines.Add(Format(RigVM_ArrayGetAtIndexOpFormat, *ConditionString, *GetOperandName(Op.ArgC, false), *GetOperandName(Op.ArgA, false)));
+					break;
+				}  
+				case ERigVMOpCode::ArraySetAtIndex:
+				{
+					const FRigVMTernaryOp& Op = ByteCode.GetOpAt<FRigVMTernaryOp>(Operation);
+					const FString ConditionString = Format(RigVM_IsValidArrayIndexFormat,  *GetOperandCPPBaseType(Op.ArgA), *GetOperandName(Op.ArgA, false));
+					Lines.Add(Format(RigVM_TemporaryArrayIndexFormat, *GetOperandName(Op.ArgB, false)));
+					Lines.Add(Format(RigVM_ArraySetAtIndexOpFormat, *ConditionString, *GetOperandName(Op.ArgA, false), *GetOperandName(Op.ArgC, false), *GetOperandName(Op.ArgC, false)));
+					break;
+				}
+				case ERigVMOpCode::ArrayAdd:
+				{
+					const FRigVMTernaryOp& Op = ByteCode.GetOpAt<FRigVMTernaryOp>(Operation);
+					const FString NumString = Format(RigVM_ArrayNumPlusOneFormat, *GetOperandName(Op.ArgA, false));
+					const FString ConditionString = Format(RigVM_IsValidArraySizeFormat, *NumString);
+					Lines.Add(Format(RigVM_ArrayAddOpFormat, *ConditionString, *GetOperandName(Op.ArgC, false), *GetOperandName(Op.ArgA, false), *GetOperandName(Op.ArgB, false)));
+					break;
+				}
+				case ERigVMOpCode::ArrayInsert:
+				{
+					const FRigVMTernaryOp& Op = ByteCode.GetOpAt<FRigVMTernaryOp>(Operation);
+					const FString NumString = Format(RigVM_ArrayNumPlusOneFormat, *GetOperandName(Op.ArgA, false));
+					const FString ConditionString = Format(RigVM_IsValidArraySizeFormat, *NumString);
+					Lines.Add(Format(RigVM_ArrayInsertOpFormat, *ConditionString, *GetOperandName(Op.ArgA, false), *GetOperandName(Op.ArgC, false), *GetOperandName(Op.ArgB, false)));
+					break;
+				}
+				case ERigVMOpCode::ArrayRemove:
+				{
+					const FRigVMBinaryOp& Op = ByteCode.GetOpAt<FRigVMBinaryOp>(Operation);
+					Lines.Add(Format(RigVM_ArrayRemoveOpFormat, *GetOperandName(Op.ArgA, false), *GetOperandName(Op.ArgB, false)));
+					break;
+				}
+				case ERigVMOpCode::ArrayFind:
+				{
+					const FRigVMQuaternaryOp& Op = ByteCode.GetOpAt<FRigVMQuaternaryOp>(Operation);
+					const FString ArrayElementCPPType = RigVMTypeUtils::BaseTypeFromArrayType(GetOperandCPPBaseType(Op.ArgA));
+					const FString ElementCPPType = GetOperandCPPType(Op.ArgB);
+					const FString TemplateSuffix = ArrayElementCPPType == ElementCPPType ? Format(RigVM_TemplateOneArgFormat, *ElementCPPType) : FString();
+					Lines.Add(Format(RigVM_ArrayFindOpFormat, *GetOperandName(Op.ArgD, false), *TemplateSuffix, *GetOperandName(Op.ArgA, false), *GetOperandName(Op.ArgB, false), *GetOperandName(Op.ArgC, false)));
+					break;
+				}
+				case ERigVMOpCode::ArrayAppend:
+				{
+					const FRigVMBinaryOp& Op = ByteCode.GetOpAt<FRigVMBinaryOp>(Operation);
+					const FString NumString = Format(RigVM_ArrayAppendOpNumFormat, *GetOperandName(Op.ArgA, false), *GetOperandName(Op.ArgB, false));
+					const FString ConditionString = Format(RigVM_IsValidArraySizeFormat, *NumString);
+					Lines.Add(Format(RigVM_ArrayAppendOpFormat, *ConditionString, *GetOperandName(Op.ArgA, false), *GetOperandName(Op.ArgB, false)));
+					break;
+				}
+				case ERigVMOpCode::ArrayIterator:
+				{
+					const FRigVMSenaryOp& Op = ByteCode.GetOpAt<FRigVMSenaryOp>(Operation);
+					const FString ArrayElementCPPType = RigVMTypeUtils::BaseTypeFromArrayType(GetOperandCPPType(Op.ArgA));
+					const FString ElementCPPType = GetOperandCPPType(Op.ArgB);
+					const FString TemplateSuffix = ArrayElementCPPType == ElementCPPType ? Format(RigVM_TemplateOneArgFormat, *ElementCPPType) : FString();
+					Lines.Add(Format(RigVM_ArrayIteratorOpFormat, *GetOperandName(Op.ArgF, false), *TemplateSuffix, *GetOperandName(Op.ArgA, false), *GetOperandName(Op.ArgB, false), *GetOperandName(Op.ArgC, false), *GetOperandName(Op.ArgD, false), *GetOperandName(Op.ArgE, false)));
+					break;
+				}
+				case ERigVMOpCode::ArrayUnion:
+				{
+					const FRigVMBinaryOp& Op = ByteCode.GetOpAt<FRigVMBinaryOp>(Operation);
+					const FString ElementCPPTypeA = RigVMTypeUtils::BaseTypeFromArrayType(GetOperandCPPType(Op.ArgA));
+					const FString ElementCPPTypeB = RigVMTypeUtils::BaseTypeFromArrayType(GetOperandCPPType(Op.ArgA));
+					const FString TemplateSuffix = ElementCPPTypeA == ElementCPPTypeB ? Format(RigVM_TemplateOneArgFormat, *ElementCPPTypeA) : FString();
+					Lines.Add(Format(RigVM_ArrayUnionOpFormat, *TemplateSuffix, *GetOperandName(Op.ArgA, false, false), *GetOperandName(Op.ArgB, false, false)));
+					break;
+				}
+				case ERigVMOpCode::ArrayDifference:
+				{
+					const FRigVMBinaryOp& Op = ByteCode.GetOpAt<FRigVMBinaryOp>(Operation);
+					const FString ElementCPPTypeA = RigVMTypeUtils::BaseTypeFromArrayType(GetOperandCPPType(Op.ArgA));
+					const FString ElementCPPTypeB = RigVMTypeUtils::BaseTypeFromArrayType(GetOperandCPPType(Op.ArgA));
+					const FString TemplateSuffix = ElementCPPTypeA == ElementCPPTypeB ? Format(RigVM_TemplateOneArgFormat, *ElementCPPTypeA) : FString();
+					Lines.Add(Format(RigVM_ArrayDifferenceOpFormat, *TemplateSuffix, *GetOperandName(Op.ArgA, false), *GetOperandName(Op.ArgB, false)));
+					break;
+				}
+				case ERigVMOpCode::ArrayIntersection:
+				{
+					const FRigVMBinaryOp& Op = ByteCode.GetOpAt<FRigVMBinaryOp>(Operation);
+					const FString ElementCPPTypeA = RigVMTypeUtils::BaseTypeFromArrayType(GetOperandCPPType(Op.ArgA));
+					const FString ElementCPPTypeB = RigVMTypeUtils::BaseTypeFromArrayType(GetOperandCPPType(Op.ArgA));
+					const FString TemplateSuffix = ElementCPPTypeA == ElementCPPTypeB ? Format(RigVM_TemplateOneArgFormat, *ElementCPPTypeA) : FString();
+					Lines.Add(Format(RigVM_ArrayIntersectionOpFormat, *TemplateSuffix, *GetOperandName(Op.ArgA, false), *GetOperandName(Op.ArgB, false)));
+					break;
+				}
+				case ERigVMOpCode::ArrayReverse:
+				{
+					const FRigVMUnaryOp& Op = ByteCode.GetOpAt<FRigVMUnaryOp>(Operation);
+					const FString ElementCPPType = RigVMTypeUtils::BaseTypeFromArrayType(GetOperandCPPType(Op.Arg));
+					const FString TemplateSuffix = Format(RigVM_TemplateOneArgFormat, *ElementCPPType);
+					Lines.Add(Format(RigVM_ArrayReverseOpFormat, *TemplateSuffix, *GetOperandName(Op.Arg, false)));
+					break;
+				}
+				case ERigVMOpCode::Invalid:
+				case ERigVMOpCode::ChangeType:
+				default:
+				{
+					// we expect to cover all op types
+					checkNoEntry();
+				}
+			}
+		}
+	}
+	else
+	{
+		// we have child groups - we need to invoke those
+		for(int32 ChildGroupIndex : Group.ChildGroups)
+		{
+			const FOperationGroup& ChildGroup = OperationGroups[ChildGroupIndex];
+
+			FString Parameters;
+			TArray<FString> ParameterArray = {RigVM_ContextPublicFormat};
+			if(!ChildGroup.OpaqueArguments.IsEmpty())
+			{
+				for(const FOpaqueArgument& OpaqueArgument : ChildGroup.OpaqueArguments)
+				{
+					ParameterArray.Add(OpaqueArgument.Get<1>());
+				}
+			}
+			Parameters = FString::Join(ParameterArray, TEXT(", "));
+			Lines.Add(Format(RigVM_InvokeExecuteGroupFormat, *ChildGroup.Entry.ToString(), ChildGroupIndex, *Parameters));
+		}
+	}
+	
+	Lines.Emplace();
+	if(Group.Depth <= 0)
+	{
+		Lines.Add(RigVM_ExecuteReachedExitFormat);
+	}
+	Lines.Add(RigVM_ReturnTrueFormat);
+
+	return DumpLines(Lines, bLog);
+}
+
+FString FRigVMCodeGenerator::DumpHeader(bool bLog)
+{
+	const FRigVMByteCode& ByteCode = VM->GetByteCode();
+
+	FStringArray FormattedEntries;
+	for(int32 EntryIndex = 0; EntryIndex < ByteCode.NumEntries(); EntryIndex++)
+	{
+		const FRigVMByteCodeEntry& Entry = ByteCode.GetEntry(EntryIndex);
+		FormattedEntries.Add(Format(RigVM_EntryNameFormat, *Entry.Name.ToString()));
+	}
+
+	FStringArray Lines;
+	Lines.Add(RigVM_CopyrightFormat);
+	Lines.Emplace();
+	Lines.Add(RigVM_AutoGeneratedFormat);
+	Lines.Emplace();
+	Lines.Add(RigVM_PragmaOnceFormat);
+	Lines.Emplace();
+	Lines.Add(DumpIncludes());
+	Lines.Add(Format(RigVM_GeneratedIncludeFormat, *ClassName));
+	Lines.Emplace();
+	if(WrappedArrayCPPTypes.Num() > 0)
+	{
+		Lines.Add(DumpWrappedArrayTypes());
+		Lines.Emplace();
+	}
+	Lines.Add(Format(RigVM_UClassDefinitionFormat, *ModuleName.ToUpper(), *ClassName));
+	Lines.Add(Format(RigVM_GetVMHashFormat, VM->GetVMHash()));
+	Lines.Add(Format(RigVM_GetEntryNamesFormat, *FString::Join(FormattedEntries, TEXT(", "))));
+	Lines.Emplace();
+	Lines.Add(FString(RigVM_DeclareExecuteFormat));
+	Lines.Emplace();
+	Lines.Add(RigVM_ProtectedFormat);
+	if(!VM->GetExternalVariables().IsEmpty())
+	{
+		Lines.Add(RigVM_DeclareUpdateExternalVariablesFormat);
+	}
+
+	for(int32 GroupIndex = 0; GroupIndex < OperationGroups.Num(); GroupIndex++)
+	{
+		const FOperationGroup& Group = OperationGroups[GroupIndex];
+
+		FString Parameters;
+		TArray<FString> ParameterArray = {RigVM_ContextPublicParameterFormat};
+		if(!Group.OpaqueArguments.IsEmpty())
+		{
+			for(const FOpaqueArgument& OpaqueArgument : Group.OpaqueArguments)
+			{
+				ParameterArray.Add(Format(RigVM_ParameterOpaqueArgumentFormat, *OpaqueArgument.Get<2>(), *OpaqueArgument.Get<1>()));
+			}
+		}
+		Parameters = FString::Join(ParameterArray, TEXT(", "));
+		
+		if(Group.Depth == 0)
+		{
+			Lines.Add(Format(RigVM_DeclareExecuteEntryFormat, *Group.Entry.ToString(), *Parameters));
+		}
+		else
+		{
+			Lines.Add(Format(RigVM_DeclareExecuteGroupFormat, *Group.Entry.ToString(), GroupIndex, *Parameters));
+		}
+	}
+
+	Lines.Emplace();
+	Lines.Add(DumpEntries(true));
+
+	Lines.Emplace();
+	Lines.Add(DumpProperties(true, INDEX_NONE));
+	if(!VM->GetExternalVariables().IsEmpty())
+	{
+		Lines.Add(DumpExternalVariables(true));
+	}
+	Lines.Add(TEXT("};"));
+	Lines.Emplace();
+	
+	return DumpLines(Lines, bLog);
+}
+
+FString FRigVMCodeGenerator::DumpSource(bool bLog)
+{
+	const FRigVMByteCode& ByteCode = VM->GetByteCode();
+
+	FStringArray Lines;
+	Lines.Add(RigVM_CopyrightFormat);
+	Lines.Emplace();
+	Lines.Add(RigVM_AutoGeneratedFormat);
+	Lines.Emplace();
+	Lines.Add(Format(RigVM_IncludeQuoteFormat, *ClassName));
+	Lines.Emplace();
+	Lines.Add(DumpEntries(false));
+
+	Lines.Emplace();
+	Lines.Add(Format(RigVM_DefineExecuteFormat, *ClassName));
+	Lines.Add(DumpInstructions(INDEX_NONE));
+	Lines.Add(TEXT("}"));
+
+	if(!VM->GetExternalVariables().IsEmpty())
+	{
+		Lines.Emplace();
+		Lines.Add(Format(RigVM_DefineUpdateExternalVariablesFormat, *ClassName));
+		Lines.Add(DumpExternalVariables(false));
+		Lines.Add(TEXT("}"));
+	}
+
+	for(int32 GroupIndex = 0; GroupIndex < OperationGroups.Num(); GroupIndex++)
+	{
+		const FOperationGroup& Group = OperationGroups[GroupIndex];
+
+		FString Parameters;
+		TArray<FString> ParameterArray = {RigVM_ContextPublicParameterFormat};
+		if(!Group.OpaqueArguments.IsEmpty())
+		{
+			for(const FOpaqueArgument& OpaqueArgument : Group.OpaqueArguments)
+			{
+				ParameterArray.Add(Format(RigVM_ParameterOpaqueArgumentFormat, *OpaqueArgument.Get<2>(), *OpaqueArgument.Get<1>()));
+			}
+		}
+		Parameters = FString::Join(ParameterArray, TEXT(", "));
+
+		Lines.Emplace();
+		if(Group.Depth == 0)
+		{
+			Lines.Add(Format(RigVM_DefineExecuteEntryFormat, *ClassName, *Group.Entry.ToString(), *Parameters));
+		}
+		else
+		{
+			Lines.Add(Format(RigVM_DefineExecuteGroupFormat, *ClassName, *Group.Entry.ToString(), GroupIndex, *Parameters));
+		}
+
+		const FString DumpedProperties = DumpProperties(false, GroupIndex);
+		if(!DumpedProperties.IsEmpty())
+		{
+			Lines.Add(DumpedProperties);
+			Lines.Emplace();
+		}
+		Lines.Add(DumpInstructions(GroupIndex));
+		Lines.Add(TEXT("}"));
+	}
+	
+	return DumpLines(Lines, bLog);
+}
+
+FString FRigVMCodeGenerator::DumpLines(const TArray<FString>& InLines, bool bLog)
+{
+	if (bLog)
+	{
+		for(const FString& Line : InLines)
+		{
+			UE_LOG(LogRigVMDeveloper, Display, RigVM_SingleStringFormat, *Line);
+		}
+	}
+	return FString::Join(InLines, RigVM_NewLineFormat);
+}
+
+void FRigVMCodeGenerator::Reset()
+{
+	ClassName.Reset();
+	ModuleName.Reset();
+	Libraries.Reset();
+	Includes.Reset();
+	OperationGroups.Reset();
+	WrappedArrayCPPTypes.Reset();
+	MappedCPPTypes.Reset();
+	Properties.Reset();
+	PropertyNameToIndex.Reset();
+}
+
+void FRigVMCodeGenerator::ParseVM(const FString& InClassName, const FString& InModuleName,
+	URigVMGraph* InModelToNativize, URigVM* InVMToNativize,
+	TMap<FString,FRigVMOperand> InPinToOperandMap, int32 InMaxOperationsPerFunction)
+{
+	check(InVMToNativize);
+
+	Reset();
+
+	Model = TStrongObjectPtr<URigVMGraph>(InModelToNativize);
+	VM = TStrongObjectPtr<URigVM>(InVMToNativize);
+	PinToOperandMap = InPinToOperandMap;
+	MaxOperationsPerFunction = InMaxOperationsPerFunction;
+	ClassName = InClassName;
+	ModuleName = InModuleName;
+
+	// create an inverted map to lookup pins from operands
+	OperandToPinMap.Reset();
+	for(const TPair<FString,FRigVMOperand>& Pair : PinToOperandMap)
+	{
+		OperandToPinMap.FindOrAdd(Pair.Value) = Pair.Key;
+	}
+
+	// default includes
+	ParseInclude(URigVM::StaticClass());
+
+	Properties.Reserve(
+		InVMToNativize->GetLiteralMemory(true)->Num() +
+		InVMToNativize->GetWorkMemory(true)->Num());
+	ParseMemory(InVMToNativize->GetLiteralMemory(true));
+	ParseMemory(InVMToNativize->GetWorkMemory(true));
+
+	ParseOperationGroups();
+}
+
+void FRigVMCodeGenerator::ParseInclude(UStruct* InDependency, const FName& InMethodName)
+{
+	if (InDependency->IsChildOf(FRigVMStruct::StaticStruct()))
+	{
+		if (const FRigVMFunction* Function = FRigVMRegistry::Get().FindFunction(Cast<UScriptStruct>(InDependency), *InMethodName.ToString()))
+		{
+			FString FunctionModuleName = Function->GetModuleName();
+			if (FunctionModuleName.Contains(TEXT("/")))
+			{
+				FunctionModuleName.Split(TEXT("/"), nullptr, &FunctionModuleName, ESearchCase::IgnoreCase, ESearchDir::FromEnd);
+			}
+			Libraries.AddUnique(FunctionModuleName);
+			Includes.AddUnique(Format(RigVM_JoinFilePathFormat, *FunctionModuleName, *Function->GetModuleRelativeHeaderPath()));
+			return;
+		}
+	}
+
+	if (InDependency == URigVM::StaticClass() ||
+		InDependency == FRigVMByteCode::StaticStruct() || 
+		InDependency->IsChildOf(FRigVMExecuteContext::StaticStruct()) ||
+		InDependency->IsChildOf(FRigVMBaseOp::StaticStruct()) ||
+		InDependency->IsChildOf(URigVMMemoryStorage::StaticClass()))
+	{
+		Libraries.AddUnique(RigVM_RigVMCoreLibraryFormat);
+		Includes.AddUnique(RigVM_RigVMCoreIncludeFormat);
+	}
+}
+
+void FRigVMCodeGenerator::ParseMemory(URigVMMemoryStorage* InMemory)
+{
+	if (InMemory == nullptr)
+	{
+		return;
+	}
+	if (InMemory->GetClass() == URigVMMemoryStorage::StaticClass())
+	{
+		return;
+	}
+
+	for (TFieldIterator<FProperty> PropertyIt(InMemory->GetClass()); PropertyIt; ++PropertyIt)
+	{
+		const FProperty* Property = *PropertyIt;
+		ParseProperty(InMemory->GetMemoryType(), Property, InMemory);
+	}
+}
+
+void FRigVMCodeGenerator::ParseProperty(ERigVMMemoryType InMemoryType, const FProperty* InProperty, URigVMMemoryStorage* InMemory)
+{
+	if (URigVMMemoryStorageGeneratorClass* MemoryClass = Cast<URigVMMemoryStorageGeneratorClass>(InMemory->GetClass()))
+	{
+		const int32 PropertyIndex = InMemory->GetPropertyIndex(InProperty);
+		const FRigVMOperand Operand(InMemoryType, PropertyIndex);
+		const FRigVMPropertyDescription& PropertyDescription = GetPropertyForOperand(Operand);
+
+		FPropertyInfo Info;
+		Info.MemoryPropertyIndex = PropertyIndex;
+		Info.Description = PropertyDescription;
+
+		check(InMemoryType == ERigVMMemoryType::Literal || InMemoryType == ERigVMMemoryType::Work);
+		if(InMemoryType == ERigVMMemoryType::Literal)
+		{
+			Info.PropertyType = ERigVMNativizedPropertyType::Literal;
+		}
+		else
+		{
+			Info.PropertyType = ERigVMNativizedPropertyType::Work;
+
+			// work state may be a hidden / sliced property.
+			if(RigVMTypeUtils::IsArrayType(PropertyDescription.CPPType))
+			{
+				if(const FString* PinPath = OperandToPinMap.Find(Operand))
+				{
+					if(const URigVMPin* Pin = Model->FindPin(*PinPath))
+					{
+						if(Pin->GetDirection() == ERigVMPinDirection::Hidden)
+						{
+							Info.PropertyType = ERigVMNativizedPropertyType::Sliced;
+						}
+					}
+				}
+			}
+		}
+
+		const int32 LookupIndex = Properties.Add(Info);
+		PropertyNameToIndex.Add(Info.Description.Name, LookupIndex);
+
+		// if this property is a nested array, we need to inject out wrapped types
+		if(PropertyDescription.Containers.Num() > 1)
+		{
+			check(PropertyDescription.Containers[0] == EPinContainerType::Array);
+			check(PropertyDescription.Containers[1] == EPinContainerType::Array);
+
+			FString BaseElementType = RigVMTypeUtils::BaseTypeFromArrayType(PropertyDescription.CPPType);
+			BaseElementType = RigVMTypeUtils::BaseTypeFromArrayType(BaseElementType);
+			WrappedArrayCPPTypes.AddUnique(BaseElementType);
+			
+			const FString ArrayArray = RigVMTypeUtils::ArrayTypeFromBaseType(RigVMTypeUtils::ArrayTypeFromBaseType(BaseElementType));
+			const FString MappedType = RigVMTypeUtils::ArrayTypeFromBaseType(GetMappedArrayTypeName(BaseElementType));
+			MappedCPPTypes.Add(ArrayArray, FMappedType(MappedType, TEXT(".Array")));
+		}
+	}
+}
+
+void FRigVMCodeGenerator::ParseOperationGroups()
+{
+	const TArray<FName>& Functions = VM->GetFunctionNames();
+	const FRigVMByteCode& ByteCode = VM->GetByteCode();
+	const FRigVMInstructionArray Operations = ByteCode.GetInstructions();
+
+	for(int32 EntryIndex = 0; EntryIndex < ByteCode.NumEntries(); EntryIndex++)
+	{
+		const FRigVMByteCodeEntry& Entry = ByteCode.GetEntry(EntryIndex);
+
+		FOperationGroup Group;
+		Group.Entry = Entry.Name;
+		Group.Depth = 0;
+		Group.First = Entry.InstructionIndex;
+		Group.Last = Operations.Num() - 1;
+
+		for(int32 OperationIndex = Group.First+1; OperationIndex < Operations.Num(); OperationIndex++)
+		{
+			if(Operations[OperationIndex].OpCode == ERigVMOpCode::Exit)
+			{
+				Group.Last = OperationIndex;
+				break;
+			}
+		}
+
+		OperationGroups.Add(Group);
+	}
+
+	for(int32 GroupIndex = 0; GroupIndex < OperationGroups.Num(); GroupIndex++)
+	{
+		// copy the group here since it will be invalid soon
+		FOperationGroup Group = OperationGroups[GroupIndex];
+		if(Group.Last - Group.First + 1 > MaxOperationsPerFunction)
+		{
+			// find all of the constraints / jumps
+			TArray<TTuple<int32,int32>> Constraints;
+			for(int32 OperationIndex = Group.First; OperationIndex <= Group.Last; OperationIndex++)
+			{
+				int32 OtherIndex = INDEX_NONE;
+				
+				const FRigVMInstruction& Operation = Operations[OperationIndex];
+				switch(Operation.OpCode)
+				{
+					case ERigVMOpCode::JumpAbsolute:
+					{
+						const FRigVMJumpOp& Op = ByteCode.GetOpAt<FRigVMJumpOp>(Operation);
+						OtherIndex = Op.InstructionIndex;
+						break;
+					}
+					case ERigVMOpCode::JumpForward:
+					{
+						const FRigVMJumpOp& Op = ByteCode.GetOpAt<FRigVMJumpOp>(Operation);
+						OtherIndex = OperationIndex + Op.InstructionIndex;
+						break;
+					}
+					case ERigVMOpCode::JumpBackward:
+					{
+						const FRigVMJumpOp& Op = ByteCode.GetOpAt<FRigVMJumpOp>(Operation);
+						OtherIndex = OperationIndex - Op.InstructionIndex;
+						break;
+					}
+					case ERigVMOpCode::JumpAbsoluteIf:
+					{
+						const FRigVMJumpIfOp& Op = ByteCode.GetOpAt<FRigVMJumpIfOp>(Operation);
+						OtherIndex = OperationIndex;
+						break;
+					}
+					case ERigVMOpCode::JumpForwardIf:
+					{
+						const FRigVMJumpIfOp& Op = ByteCode.GetOpAt<FRigVMJumpIfOp>(Operation);
+						OtherIndex = OperationIndex + Op.InstructionIndex;
+						break;
+					}
+					case ERigVMOpCode::JumpBackwardIf:
+					{
+						const FRigVMJumpIfOp& Op = ByteCode.GetOpAt<FRigVMJumpIfOp>(Operation);
+						OtherIndex = OperationIndex - Op.InstructionIndex;
+						break;
+					}
+					default:
+					{
+						break;
+					}
+				}
+
+				if(OtherIndex != INDEX_NONE)
+				{
+					Constraints.Add(TTuple<int32,int32>(
+						FMath::Min(OperationIndex, OtherIndex),
+						FMath::Max(OperationIndex, OtherIndex)));
+				}
+			}
+
+			const int32 First = Group.First;
+			const int32 Last = Group.Last;
+			int32 Middle = (Group.Last + Group.First) / 2;
+			int32 MiddleHead = Middle;
+			int32 MiddleTail = Middle;
+			bool bMoveTowardsHead = true;
+
+			auto IsWithinConstraint = [&bMoveTowardsHead, &MiddleHead, &MiddleTail](const TTuple<int32,int32>& Constraint) -> bool
+			{
+				int32& Middle = bMoveTowardsHead ? MiddleHead : MiddleTail;
+				return FMath::IsWithin(Middle, Constraint.Get<0>(), Constraint.Get<1>());
+			};
+
+			// while the middle instruction is within a constraint - move it  
+			while(Constraints.FindByPredicate(IsWithinConstraint))
+			{
+				MiddleHead--;
+			}
+			bMoveTowardsHead = false;
+			while(Constraints.FindByPredicate(IsWithinConstraint))
+			{
+				MiddleTail++;
+			}
+
+			Middle = ((MiddleHead - First) > (Last - MiddleTail)) ? MiddleHead : MiddleTail;
+			if(Middle == First || Middle == Last)
+			{
+				continue;
+			}
+
+			FOperationGroup HeadGroup = Group;
+			HeadGroup.ParentGroup = GroupIndex;
+			HeadGroup.ChildGroups.Reset();
+			HeadGroup.First = First;
+			HeadGroup.Last = Middle;
+			HeadGroup.Depth++;
+
+			FOperationGroup TailGroup = Group;
+			TailGroup.ParentGroup = GroupIndex;
+			TailGroup.ChildGroups.Reset();
+			TailGroup.First = Middle+1;
+			TailGroup.Last = Last;
+			TailGroup.Depth++;
+			
+			const int32 HeadGroupIndex = OperationGroups.Add(HeadGroup);
+			const int32 TailGroupIndex = OperationGroups.Add(TailGroup);
+			OperationGroups[GroupIndex].ChildGroups.Add(HeadGroupIndex);
+			OperationGroups[GroupIndex].ChildGroups.Add(TailGroupIndex);
+		}
+	}
+
+	for(int32 GroupIndex = 0; GroupIndex < OperationGroups.Num(); GroupIndex++)
+	{
+		FOperationGroup& Group = OperationGroups[GroupIndex];
+		
+		// determine all the necessary labels and optional arguments
+		for(int32 OperationIndex = Group.First; OperationIndex <= Group.Last; OperationIndex++)
+		{
+			const FRigVMInstruction& Operation = Operations[OperationIndex];
+			switch(Operation.OpCode)
+			{
+				case ERigVMOpCode::Execute_0_Operands:
+				case ERigVMOpCode::Execute_1_Operands:
+				case ERigVMOpCode::Execute_2_Operands:
+				case ERigVMOpCode::Execute_3_Operands:
+				case ERigVMOpCode::Execute_4_Operands:
+				case ERigVMOpCode::Execute_5_Operands:
+				case ERigVMOpCode::Execute_6_Operands:
+				case ERigVMOpCode::Execute_7_Operands:
+				case ERigVMOpCode::Execute_8_Operands:
+				case ERigVMOpCode::Execute_9_Operands:
+				case ERigVMOpCode::Execute_10_Operands:
+				case ERigVMOpCode::Execute_11_Operands:
+				case ERigVMOpCode::Execute_12_Operands:
+				case ERigVMOpCode::Execute_13_Operands:
+				case ERigVMOpCode::Execute_14_Operands:
+				case ERigVMOpCode::Execute_15_Operands:
+				case ERigVMOpCode::Execute_16_Operands:
+				case ERigVMOpCode::Execute_17_Operands:
+				case ERigVMOpCode::Execute_18_Operands:
+				case ERigVMOpCode::Execute_19_Operands:
+				case ERigVMOpCode::Execute_20_Operands:
+				case ERigVMOpCode::Execute_21_Operands:
+				case ERigVMOpCode::Execute_22_Operands:
+				case ERigVMOpCode::Execute_23_Operands:
+				case ERigVMOpCode::Execute_24_Operands:
+				case ERigVMOpCode::Execute_25_Operands:
+				case ERigVMOpCode::Execute_26_Operands:
+				case ERigVMOpCode::Execute_27_Operands:
+				case ERigVMOpCode::Execute_28_Operands:
+				case ERigVMOpCode::Execute_29_Operands:
+				case ERigVMOpCode::Execute_30_Operands:
+				case ERigVMOpCode::Execute_31_Operands:
+				case ERigVMOpCode::Execute_32_Operands:
+				case ERigVMOpCode::Execute_33_Operands:
+				case ERigVMOpCode::Execute_34_Operands:
+				case ERigVMOpCode::Execute_35_Operands:
+				case ERigVMOpCode::Execute_36_Operands:
+				case ERigVMOpCode::Execute_37_Operands:
+				case ERigVMOpCode::Execute_38_Operands:
+				case ERigVMOpCode::Execute_39_Operands:
+				case ERigVMOpCode::Execute_40_Operands:
+				case ERigVMOpCode::Execute_41_Operands:
+				case ERigVMOpCode::Execute_42_Operands:
+				case ERigVMOpCode::Execute_43_Operands:
+				case ERigVMOpCode::Execute_44_Operands:
+				case ERigVMOpCode::Execute_45_Operands:
+				case ERigVMOpCode::Execute_46_Operands:
+				case ERigVMOpCode::Execute_47_Operands:
+				case ERigVMOpCode::Execute_48_Operands:
+				case ERigVMOpCode::Execute_49_Operands:
+				case ERigVMOpCode::Execute_50_Operands:
+				case ERigVMOpCode::Execute_51_Operands:
+				case ERigVMOpCode::Execute_52_Operands:
+				case ERigVMOpCode::Execute_53_Operands:
+				case ERigVMOpCode::Execute_54_Operands:
+				case ERigVMOpCode::Execute_55_Operands:
+				case ERigVMOpCode::Execute_56_Operands:
+				case ERigVMOpCode::Execute_57_Operands:
+				case ERigVMOpCode::Execute_58_Operands:
+				case ERigVMOpCode::Execute_59_Operands:
+				case ERigVMOpCode::Execute_60_Operands:
+				case ERigVMOpCode::Execute_61_Operands:
+				case ERigVMOpCode::Execute_62_Operands:
+				case ERigVMOpCode::Execute_63_Operands:
+				case ERigVMOpCode::Execute_64_Operands:
+				{
+					const FRigVMExecuteOp& Op = ByteCode.GetOpAt<FRigVMExecuteOp>(Operation);
+					const FRigVMFunction* Function = FRigVMRegistry::Get().FindFunction(*Functions[Op.FunctionIndex].ToString());
+					check(Function);
+
+					// make sure to include the required header
+					ParseInclude(Function->Struct, Function->GetMethodName());
+
+					// check all of the arguments that may not be part of the struct definition
+					for(const FRigVMFunctionArgument& Argument : Function->GetArguments())
+					{
+						if (Function->IsAdditionalArgument(Argument))
+						{
+							if(!Group.OpaqueArguments.ContainsByPredicate([Argument](const FOpaqueArgument& OpaqueArgument) -> bool
+							{
+								return Argument.Name == OpaqueArgument.Get<0>();
+							}))
+							{
+								const FString ArgumentName = Format(RigVM_AdditionalArgumentNameFormat, Group.OpaqueArguments.Num(), Argument.Name);
+								FOpaqueArgument OpaqueArgument(Argument.Name, ArgumentName, Argument.Type);
+								Group.OpaqueArguments.AddUnique(OpaqueArgument);
+							}
+						}
+					}
+					break;
+				}
+				case ERigVMOpCode::JumpAbsolute:
+				{
+					const FRigVMJumpOp& Op = ByteCode.GetOpAt<FRigVMJumpOp>(Operation);
+					Group.RequiredLabels.AddUnique(Op.InstructionIndex);break;
+					break;
+				}
+				case ERigVMOpCode::JumpForward:
+				{
+					const FRigVMJumpOp& Op = ByteCode.GetOpAt<FRigVMJumpOp>(Operation);
+					Group.RequiredLabels.AddUnique(OperationIndex + Op.InstructionIndex);break;
+					break;
+				}
+				case ERigVMOpCode::JumpBackward:
+				{
+					const FRigVMJumpOp& Op = ByteCode.GetOpAt<FRigVMJumpOp>(Operation);
+					Group.RequiredLabels.AddUnique(OperationIndex - Op.InstructionIndex);break;
+					break;
+				}
+				case ERigVMOpCode::JumpAbsoluteIf:
+				{
+					const FRigVMJumpIfOp& Op = ByteCode.GetOpAt<FRigVMJumpIfOp>(Operation);
+					Group.RequiredLabels.AddUnique(Op.InstructionIndex);break;
+					break;
+				}
+				case ERigVMOpCode::JumpForwardIf:
+				{
+					const FRigVMJumpIfOp& Op = ByteCode.GetOpAt<FRigVMJumpIfOp>(Operation);
+					Group.RequiredLabels.AddUnique(OperationIndex + Op.InstructionIndex);break;
+					break;
+				}
+				case ERigVMOpCode::JumpBackwardIf:
+				{
+					const FRigVMJumpIfOp& Op = ByteCode.GetOpAt<FRigVMJumpIfOp>(Operation);
+					Group.RequiredLabels.AddUnique(OperationIndex - Op.InstructionIndex);break;
+					break;
+				}
+				default:
+				{
+					break;
+				}
+			}
+
+			// update the usage table to know which property is used where
+			FRigVMOperandArray Operands = ByteCode.GetOperandsForOp(Operation);
+			for(const FRigVMOperand& Operand : Operands)
+			{
+				if(Operand.GetMemoryType() == ERigVMMemoryType::Literal ||
+					Operand.GetMemoryType() == ERigVMMemoryType::Work)
+				{
+					const int32 PropertyIndex = GetPropertyIndex(Operand);
+					Properties[PropertyIndex].Groups.AddUnique(GroupIndex);
+				}
+			}
+		}
+	}
+
+	// now that we know the groups that properties belong to, we need to filter them.
+	// we only want to declare properties once - and pass them as parameters to the enclosed
+	// groups. parent groups which only pass them to one sub group can be removed from the
+	// property's group table. the top level group which still contains the property is the
+	// one to declare it.
+	const TArray<FOperationGroup>& Groups = OperationGroups;
+	for(int32 PropertyIndex = 0; PropertyIndex < Properties.Num(); PropertyIndex++)
+	{
+		// work / sliced properties are always defined at the top
+		FPropertyInfo& Property = Properties[PropertyIndex];
+
+		if(Property.PropertyType == ERigVMNativizedPropertyType::Sliced)
+		{
+			Property.Groups.Reset();
+			continue;
+		}
+		
+		int32 GroupCount;
+		do
+		{
+			GroupCount = Property.Groups.Num();
+			TArray<int32> FilteredGroups = Property.Groups;
+			FilteredGroups.RemoveAll([Groups, Property](int32 Group) -> bool
+			{
+				return !Groups[Group].ChildGroups.IsEmpty();
+			});
+			Property.Groups = FilteredGroups;
+			
+		} while(GroupCount != Property.Groups.Num());
+	}
+}
+
+FString FRigVMCodeGenerator::GetOperandName(const FRigVMOperand& InOperand, bool bPerSlice, bool bAsInput) const
+{
+	const FRigVMPropertyDescription& Property = GetPropertyForOperand(InOperand);
+	const FRigVMPropertyPathDescription& PropertyPath = GetPropertyPathForOperand(InOperand);
+
+	FString OperandName;
+	if (Property.IsValid())
+	{
+		OperandName = Property.Name.ToString();
+	}
+
+	if (InOperand.GetMemoryType() != ERigVMMemoryType::External)
+	{
+		OperandName = SanitizeName(OperandName, Property.CPPType);
+	}
+
+	// if we are an array on work memory this indicates that we'll be sliced.
+	if (bPerSlice && InOperand.GetMemoryType() == ERigVMMemoryType::Work && RigVMTypeUtils::IsArrayType(Property.CPPType))
+	{
+		const FString MappedType = GetMappedType(Property.CPPType);
+		const FString MappedTypeSuffix = GetMappedTypeSuffix(Property.CPPType);
+		const FString BaseCPPType = RigVMTypeUtils::IsArrayType(MappedType) ? RigVMTypeUtils::BaseTypeFromArrayType(MappedType) : MappedType;
+		OperandName = Format(RigVM_GetOperandSliceFormat, *BaseCPPType, *OperandName, *MappedTypeSuffix);
+	}
+
+	if (InOperand.GetMemoryType() == ERigVMMemoryType::External)
+	{
+		OperandName = Format(RigVM_ExternalVariableFormat, *OperandName);
+	}
+
+	if (PropertyPath.IsValid())
+	{
+		check(Property.IsValid());
+
+		FString RemainingSegmentPath = PropertyPath.SegmentPath;
+		FString SegmentPath = OperandName;
+
+		static TMap<UScriptStruct*, TMap<FString, TTuple<FString, FString>>> MappedProperties;
+		if (MappedProperties.IsEmpty())
+		{
+			TMap<FString, TTuple<FString, FString>>& TransformMap = MappedProperties.Add(TBaseStructure<FTransform>::Get());
+			TransformMap.Add(TEXT("Translation"), TTuple<FString, FString>(TEXT("{0}.GetTranslation()"), TEXT("{0}.SetTranslation(")));
+			TransformMap.Add(TEXT("Translation.X"), TTuple<FString, FString>(TEXT("{0}.GetTranslation().X"), TEXT("FTransformSetter({0}).SetTranslationX(")));
+			TransformMap.Add(TEXT("Translation.Y"), TTuple<FString, FString>(TEXT("{0}.GetTranslation().Y"), TEXT("FTransformSetter({0}).SetTranslationY(")));
+			TransformMap.Add(TEXT("Translation.Z"), TTuple<FString, FString>(TEXT("{0}.GetTranslation().Z"), TEXT("FTransformSetter({0}).SetTranslationZ(")));
+			TransformMap.Add(TEXT("Rotation"), TTuple<FString, FString>(TEXT("{0}.GetRotation()"), TEXT("{0}.SetRotation(")));
+			TransformMap.Add(TEXT("Rotation.X"), TTuple<FString, FString>(TEXT("{0}.GetRotation().X"), TEXT("FTransformSetter({0}).SetRotationX(")));
+			TransformMap.Add(TEXT("Rotation.Y"), TTuple<FString, FString>(TEXT("{0}.GetRotation().Y"), TEXT("FTransformSetter({0}).SetRotationY(")));
+			TransformMap.Add(TEXT("Rotation.Z"), TTuple<FString, FString>(TEXT("{0}.GetRotation().Z"), TEXT("FTransformSetter({0}).SetRotationZ(")));
+			TransformMap.Add(TEXT("Rotation.W"), TTuple<FString, FString>(TEXT("{0}.GetRotation().W"), TEXT("FTransformSetter({0}).SetRotationW(")));
+			TransformMap.Add(TEXT("Scale3D"), TTuple<FString, FString>(TEXT("{0}.GetScale3D()"), TEXT("{0}.SetScale3D(")));
+			TransformMap.Add(TEXT("Scale3D.X"), TTuple<FString, FString>(TEXT("{0}.GetScale3D().X"), TEXT("FTransformSetter({0}).SetScaleX(")));
+			TransformMap.Add(TEXT("Scale3D.Y"), TTuple<FString, FString>(TEXT("{0}.GetScale3D().Y"), TEXT("FTransformSetter({0}).SetScaleY(")));
+			TransformMap.Add(TEXT("Scale3D.Z"), TTuple<FString, FString>(TEXT("{0}.GetScale3D().Z"), TEXT("FTransformSetter({0}).SetScaleZ(")));
+		}
+
+		const FProperty* CurrentProperty = Property.Property;
+
+		while(!RemainingSegmentPath.IsEmpty() && (CurrentProperty != nullptr))
+		{
+			FString Left, Right;
+			if(!URigVMPin::SplitPinPathAtStart(RemainingSegmentPath, Left, Right))
+			{
+				Left = RemainingSegmentPath;
+				Right.Reset();
+			}
+			
+			if (const FStructProperty* StructProperty = CastField<FStructProperty>(CurrentProperty))
+			{
+				CurrentProperty = StructProperty->Struct->FindPropertyByName(*Left);
+				
+				if (const TMap<FString, TTuple<FString, FString>>* PropertyMap = MappedProperties.Find(StructProperty->Struct))
+				{
+					if (const TTuple<FString, FString>* NewSegmentName = PropertyMap->Find(RemainingSegmentPath))
+					{
+						FStringFormatOrderedArguments FormatArguments;
+						FormatArguments.Add(SegmentPath);
+
+						SegmentPath = FString::Format(
+							bAsInput ? *NewSegmentName->Get<0>() : *NewSegmentName->Get<1>(),
+							FormatArguments);
+						break;
+					}
+				}
+			}
+			else if(const FArrayProperty* ArrayProperty = CastField<FArrayProperty>(CurrentProperty))
+			{
+				CurrentProperty = ArrayProperty->Inner;
+
+				Left = Left.TrimChar(TEXT('['));
+				Left = Left.TrimChar(TEXT(']'));
+
+				FString ExtendedType;
+				FString CPPType = CurrentProperty->GetCPPType(&ExtendedType);
+				CPPType += ExtendedType;
+				if(CPPType == RigVMTypeUtils::UInt8Type && CastField<FBoolProperty>(CurrentProperty))
+				{
+					CPPType = RigVMTypeUtils::BoolType;
+				}
+
+				SegmentPath = Format(RigVM_GetArrayElementSafeFormat, *CPPType, *SegmentPath, *Left);
+				RemainingSegmentPath = Right;
+				continue;
+			}
+			else
+			{
+				CurrentProperty = nullptr;
+			}
+
+			SegmentPath = Format(RigVM_JoinSegmentPathFormat, *SegmentPath, *Left);
+			RemainingSegmentPath = Right;
+		}
+
+		return SegmentPath;
+	}
+
+	return OperandName;
+}
+
+FString FRigVMCodeGenerator::GetOperandCPPType(const FRigVMOperand& InOperand) const
+{
+	const FRigVMPropertyDescription& Property = GetPropertyForOperand(InOperand);
+	const FRigVMPropertyPathDescription& PropertyPath = GetPropertyPathForOperand(InOperand);
+
+	if (PropertyPath.IsValid())
+	{
+		check(Property.IsValid());
+
+		const FRigVMPropertyPath Path(Property.Property, PropertyPath.SegmentPath);
+		const FProperty* TailProperty = Path.GetTailProperty();
+		check(TailProperty);
+
+		FString ExtendedType;
+		const FString CPPType = TailProperty->GetCPPType(&ExtendedType);
+		return CPPType + ExtendedType;
+	}
+	
+	if (Property.IsValid())
+	{
+		return Property.CPPType;
+	}
+
+	return FString(); 
+}
+
+FString FRigVMCodeGenerator::GetOperandCPPBaseType(const FRigVMOperand& InOperand) const
+{
+	FString CPPType = GetOperandCPPType(InOperand);
+	while(RigVMTypeUtils::IsArrayType(CPPType))
+	{
+		CPPType = RigVMTypeUtils::BaseTypeFromArrayType(CPPType);
+	}
+	return CPPType;
+}
+
+
+FString FRigVMCodeGenerator::SanitizeName(const FString& InName, const FString& CPPType)
+{
+	FString Name = InName;
+	while(Name.Contains(RigVM_DoubleUnderscoreFormat))
+	{
+		Name.ReplaceInline(RigVM_DoubleUnderscoreFormat, RigVM_SingleUnderscoreFormat);
+	}
+
+	if(CPPType == RigVMTypeUtils::BoolType && !Name.StartsWith(RigVM_BoolPropertyPrefix, ESearchCase::CaseSensitive))
+	{
+		Name = RigVM_BoolPropertyPrefix + Name;
+	}
+
+	return Name;
+}
+
+FRigVMPropertyDescription FRigVMCodeGenerator::GetPropertyForOperand(const FRigVMOperand& InOperand) const
+{
+	CheckOperand(InOperand);
+
+	const FProperty* Property = nullptr;
+	const uint8* Memory = nullptr;
+
+	if (InOperand.GetMemoryType() == ERigVMMemoryType::External)
+	{
+		const FRigVMExternalVariable& ExternalVariable = VM->ExternalVariables[InOperand.GetRegisterIndex()];
+		Property = ExternalVariable.Property;
+		Memory = ExternalVariable.Memory;
+	}
+	else
+	{
+		if (URigVMMemoryStorage* MemoryStorage = VM->GetMemoryByType(InOperand.GetMemoryType()))
+		{
+			Property = MemoryStorage->GetProperty(InOperand.GetRegisterIndex());
+			if (Property)
+			{
+				Memory = Property->ContainerPtrToValuePtr<uint8>(MemoryStorage);
+			}
+		}
+	}
+
+	if (Property)
+	{
+		FString DefaultValue;
+		
+		if (const UClass* OwningClass = Property->GetOwnerClass())
+		{
+			if (const UObject* CDO = OwningClass->GetDefaultObject())
+			{
+				Memory = Property->ContainerPtrToValuePtr<uint8>(CDO);
+			}
+		}
+
+		if (Memory)
+		{
+			DefaultValue = FRigVMStruct::ExportToFullyQualifiedText(Property, Memory, true);
+		}
+		
+		return FRigVMPropertyDescription(Property, DefaultValue);
+	}
+
+	return FRigVMPropertyDescription();
+}
+
+const FRigVMPropertyPathDescription& FRigVMCodeGenerator::GetPropertyPathForOperand(const FRigVMOperand& InOperand) const
+{
+	CheckOperand(InOperand);
+
+	const int32 RegisterOffsetIndex = InOperand.GetRegisterOffset();
+	if (RegisterOffsetIndex != INDEX_NONE)
+	{
+		if (InOperand.GetMemoryType() == ERigVMMemoryType::External)
+		{
+			if (VM->ExternalPropertyPathDescriptions.IsValidIndex(RegisterOffsetIndex))
+			{
+				return VM->ExternalPropertyPathDescriptions[RegisterOffsetIndex];
+			}
+		}
+		else
+		{
+			if (const URigVMMemoryStorage* MemoryStorage = VM->GetMemoryByType(InOperand.GetMemoryType()))
+			{
+				if (const URigVMMemoryStorageGeneratorClass* MemoryClass = Cast<URigVMMemoryStorageGeneratorClass>(MemoryStorage->GetClass()))
+				{
+					if (MemoryClass->PropertyPathDescriptions.IsValidIndex(RegisterOffsetIndex))
+					{
+						return MemoryClass->PropertyPathDescriptions[RegisterOffsetIndex];
+					}
+				}
+			}
+		}
+	}
+
+	static const FRigVMPropertyPathDescription EmptyPropertyPath;
+	return EmptyPropertyPath;
+}
+
+int32 FRigVMCodeGenerator::GetPropertyIndex(const FRigVMPropertyDescription& InProperty) const
+{
+	if(const int32* Index = PropertyNameToIndex.Find(InProperty.Name))
+	{
+		return *Index;
+	}
+	return INDEX_NONE;
+}
+
+int32 FRigVMCodeGenerator::GetPropertyIndex(const FRigVMOperand& InOperand) const
+{
+	return GetPropertyIndex(GetPropertyForOperand(InOperand));
+}
+
+FRigVMCodeGenerator::ERigVMNativizedPropertyType FRigVMCodeGenerator::GetPropertyType(
+	const FRigVMPropertyDescription& InProperty) const
+{
+	const int32 PropertyIndex = GetPropertyIndex(InProperty);
+	if(Properties.IsValidIndex(PropertyIndex))
+	{
+		return Properties[PropertyIndex].PropertyType;
+	}
+	return ERigVMNativizedPropertyType::Invalid;
+}
+
+FRigVMCodeGenerator::ERigVMNativizedPropertyType FRigVMCodeGenerator::GetPropertyType(const FRigVMOperand& InOperand) const
+{
+	return GetPropertyType(GetPropertyForOperand(InOperand));
+}
+
+void FRigVMCodeGenerator::CheckOperand(const FRigVMOperand& InOperand) const
+{
+	ensure(InOperand.IsValid());
+	ensure(InOperand.GetMemoryType() != ERigVMMemoryType::Invalid);
+	ensure(InOperand.GetMemoryType() != ERigVMMemoryType::Debug);
+
+	if (InOperand.GetMemoryType() == ERigVMMemoryType::External)
+	{
+		ensure(VM->ExternalVariables.IsValidIndex(InOperand.GetRegisterIndex()));
+		if (InOperand.GetRegisterOffset() != INDEX_NONE)
+		{
+			ensure(VM->ExternalPropertyPaths.IsValidIndex(InOperand.GetRegisterOffset()));
+			ensure(VM->ExternalPropertyPathDescriptions.IsValidIndex(InOperand.GetRegisterOffset()));
+		}
+	}
+	else
+	{
+		URigVMMemoryStorage* Memory = VM->GetMemoryByType(InOperand.GetMemoryType(), false);
+		check(Memory);
+
+		ensure(Memory->GetProperties().IsValidIndex(InOperand.GetRegisterIndex()));
+		if (InOperand.GetRegisterOffset() != INDEX_NONE)
+		{
+			ensure(Memory->GetPropertyPaths().IsValidIndex(InOperand.GetRegisterOffset()));
+			URigVMMemoryStorageGeneratorClass* Class = CastChecked<URigVMMemoryStorageGeneratorClass>(Memory->GetClass()); 
+			ensure(Class->PropertyPathDescriptions.IsValidIndex(InOperand.GetRegisterOffset()));
+		}
+	}
+}
+
+FString FRigVMCodeGenerator::GetMappedType(const FString& InCPPType) const
+{
+	FString CPPType = InCPPType;
+	CPPType.RemoveSpacesInline();
+	
+	if(const FMappedType* MappedType = MappedCPPTypes.Find(CPPType))
+	{
+		return MappedType->Get<0>();
+	}
+
+	if(RigVMTypeUtils::IsArrayType(CPPType))
+	{
+		const FString BaseType = RigVMTypeUtils::BaseTypeFromArrayType(CPPType);
+		return RigVMTypeUtils::ArrayTypeFromBaseType(GetMappedType(BaseType));
+	}
+
+	return CPPType;
+}
+
+const FString& FRigVMCodeGenerator::GetMappedTypeSuffix(const FString& InCPPType) const
+{
+	FString CPPType = InCPPType;
+	CPPType.RemoveSpacesInline();
+	
+	if(const FMappedType* MappedType = MappedCPPTypes.Find(CPPType))
+	{
+		return MappedType->Get<1>();
+	}
+
+	static const FString EmptyString;
+	return EmptyString;
+}
+
+FString FRigVMCodeGenerator::GetMappedArrayTypeName(const FString InBaseElementType) const
+{
+	return Format(RigVM_WrappedTypeNameFormat, *InBaseElementType, *ClassName);
+}
+
+const FRigVMCodeGenerator::FOperationGroup& FRigVMCodeGenerator::GetGroup(int32 InGroupIndex) const
+{
+	if(OperationGroups.IsValidIndex(InGroupIndex))
+	{
+		return OperationGroups[InGroupIndex];
+	}
+	
+	static FOperationGroup CompleteGroup;
+	if(CompleteGroup.First == INDEX_NONE)
+	{
+		CompleteGroup.First = 0;
+		CompleteGroup.Last = VM->GetByteCode().GetNumInstructions() - 1; 
+	}
+	return CompleteGroup;
+}
+
+bool FRigVMCodeGenerator::IsOperationPartOfGroup(int32 InOperationIndex, int32 InGroupIndex, bool bIncludeChildGroups) const
+{
+	if(!OperationGroups.IsValidIndex(InGroupIndex) && OperationGroups.IsEmpty())
+	{
+		return true;
+	}
+
+	const FOperationGroup& Group = OperationGroups[InGroupIndex];
+	for(const int32 ChildGroupIndex : Group.ChildGroups)
+	{
+		if(IsOperationPartOfGroup(InOperationIndex, ChildGroupIndex, true))
+		{
+			// if we found the operation in a child group
+			// we'll return true if we are also looking within the child groups,
+			// and false if we are not. this means that an operation that's part of a
+			// child group but not the main group will return false here
+			return bIncludeChildGroups;
+		}
+	}
+
+	return FMath::IsWithin(InOperationIndex, Group.First, Group.Last);;
+}
+
+bool FRigVMCodeGenerator::IsPropertyPartOfGroup(int32 InPropertyIndex, int32 InGroupIndex) const
+{
+	if(InGroupIndex == INDEX_NONE)
+	{
+		return true;
+	}
+
+	if(!OperationGroups.IsValidIndex(InGroupIndex))
+	{
+		return false;
+	}
+
+	if(Properties.IsValidIndex(InPropertyIndex))
+	{
+		const FPropertyInfo& Property = Properties[InPropertyIndex];
+		const FOperationGroup& Group = OperationGroups[InGroupIndex];
+
+		// properties are only defines in the leaf groups
+		if(!Group.ChildGroups.IsEmpty())
+		{
+			return false;
+		}
+
+		if(Property.PropertyType == ERigVMNativizedPropertyType::Work)
+		{
+			if(Property.Groups.Num() > 1)
+			{
+				return false;
+			}
+		}
+
+		return Property.Groups.Contains(InGroupIndex);
+	}
+	
+	return false;
+}
+
+const TMap<FName, FRigVMCodeGenerator::TStructConstGenerator>& FRigVMCodeGenerator::GetStructConstGenerators()
+{
+	static TMap<FName, FRigVMCodeGenerator::TStructConstGenerator> StructConstGenerators;
+	if(!StructConstGenerators.IsEmpty())
+	{
+		return StructConstGenerators;
+	}
+
+	struct DefaultValueHelpers
+	{
+		static FStringArray SplitIntoArray(const FString& InValue)
+		{
+			return URigVMPin::SplitDefaultValue(InValue);
+		}
+		
+		static FStringMap SplitIntoMap(const FString& InValue)
+		{
+			FStringArray Pairs = SplitIntoArray(InValue);
+			FStringMap Map;
+			for(const FString& Pair : Pairs)
+			{
+				FString Key, Value;
+				if(Pair.Split(TEXT("="), &Key, &Value))
+				{
+					Map.Add(Key, Value);
+				}
+			}
+			return Map;
+		}
+
+		static FString RemoveQuotes(const FString& InValue)
+		{
+			FString Value = InValue;
+			if(Value.StartsWith(TEXT("\\\"")))
+			{
+				Value = Value.Mid(2, Value.Len()  - 4);
+			}
+			else
+			{
+				Value = Value.TrimQuotes();
+			}
+			return Value;
+		}
+	};
+
+	static constexpr TCHAR X[] = TEXT("X");
+	static constexpr TCHAR Y[] = TEXT("Y");
+	static constexpr TCHAR Z[] = TEXT("Z");
+	static constexpr TCHAR W[] = TEXT("W");
+	static constexpr TCHAR R[] = TEXT("R");
+	static constexpr TCHAR G[] = TEXT("G");
+	static constexpr TCHAR B[] = TEXT("B");
+	static constexpr TCHAR A[] = TEXT("A");
+	static constexpr TCHAR Translation[] = TEXT("Translation");
+	static constexpr TCHAR Rotation[] = TEXT("Rotation");
+	static constexpr TCHAR Scale[] = TEXT("Scale");
+	static constexpr TCHAR Scale3D[] = TEXT("Scale3D");
+	static constexpr TCHAR Pitch[] = TEXT("Pitch");
+	static constexpr TCHAR Yaw[] = TEXT("Yaw");
+	static constexpr TCHAR Roll[] = TEXT("Roll");
+	static constexpr TCHAR Type[] = TEXT("Type");
+	static constexpr TCHAR Name[] = TEXT("Name");
+
+	StructConstGenerators.Add(TEXT("FVector2D"), [](const FString& InDefault) -> FString
+	{
+		static constexpr TCHAR Constructor[] = TEXT("FVector2D({0}, {1})");
+		const FStringMap Defaults = DefaultValueHelpers::SplitIntoMap(InDefault);
+		if(Defaults.Contains(X) && Defaults.Contains(Y))
+		{
+			return Format(Constructor, Defaults.FindChecked(X), Defaults.FindChecked(Y));
+		}
+		return FString();
+	});
+
+	StructConstGenerators.Add(TEXT("FVector"), [](const FString& InDefault) -> FString
+	{
+		static constexpr TCHAR Constructor[] = TEXT("FVector({0}, {1}, {2})");
+		const FStringMap Defaults = DefaultValueHelpers::SplitIntoMap(InDefault);
+		if(Defaults.Contains(X) && Defaults.Contains(Y) && Defaults.Contains(Z))
+		{
+			return Format(Constructor, Defaults.FindChecked(X), Defaults.FindChecked(Y), Defaults.FindChecked(Z));
+		}
+		return FString();
+	});
+
+	StructConstGenerators.Add(TEXT("FVector4"), [](const FString& InDefault) -> FString
+	{
+		static constexpr TCHAR Constructor[] = TEXT("FVector4({0}, {1}, {2}, {3})");
+		const FStringMap Defaults = DefaultValueHelpers::SplitIntoMap(InDefault);
+		if(Defaults.Contains(X) && Defaults.Contains(Y) && Defaults.Contains(Z) && Defaults.Contains(W))
+		{
+			return Format(Constructor, Defaults.FindChecked(X), Defaults.FindChecked(Y), Defaults.FindChecked(Z), Defaults.FindChecked(W));
+		}
+		return FString();
+	});
+
+	StructConstGenerators.Add(TEXT("FLinearColor"), [](const FString& InDefault) -> FString
+	{
+		static constexpr TCHAR Constructor[] = TEXT("FLinearColor({0}, {1}, {2}, {3})");
+		const FStringMap Defaults = DefaultValueHelpers::SplitIntoMap(InDefault);
+		if(Defaults.Contains(R) && Defaults.Contains(G) && Defaults.Contains(B) && Defaults.Contains(A))
+		{
+			return Format(Constructor, Defaults.FindChecked(R), Defaults.FindChecked(G), Defaults.FindChecked(B), Defaults.FindChecked(A));
+		}
+		return FString();
+	});
+
+	StructConstGenerators.Add(TEXT("FQuat"), [](const FString& InDefault) -> FString
+	{
+		static constexpr TCHAR Constructor[] = TEXT("FQuat({0}, {1}, {2}, {3})");
+		const FStringMap Defaults = DefaultValueHelpers::SplitIntoMap(InDefault);
+		if(Defaults.Contains(X) && Defaults.Contains(Y) && Defaults.Contains(Z) && Defaults.Contains(W))
+		{
+			return Format(Constructor, Defaults.FindChecked(X), Defaults.FindChecked(Y), Defaults.FindChecked(Z), Defaults.FindChecked(W));
+		}
+		return FString();
+	});
+
+	StructConstGenerators.Add(TEXT("FRotator"), [](const FString& InDefault) -> FString
+	{
+		static constexpr TCHAR Constructor[] = TEXT("FRotator({0}, {1}, {2})");
+		const FStringMap Defaults = DefaultValueHelpers::SplitIntoMap(InDefault);
+		if(Defaults.Contains(Pitch) && Defaults.Contains(Yaw) && Defaults.Contains(Roll))
+		{
+			return Format(Constructor, Defaults.FindChecked(Pitch), Defaults.FindChecked(Yaw), Defaults.FindChecked(Roll));
+		}
+		return FString();
+	});
+
+	StructConstGenerators.Add(TEXT("FTransform"), [](const FString& InDefault) -> FString
+	{
+		static constexpr TCHAR Constructor[] = TEXT("FTransform({0}, {1}, {2})");
+		const FStringMap Defaults = DefaultValueHelpers::SplitIntoMap(InDefault);
+		if(Defaults.Contains(Rotation) && Defaults.Contains(Translation) && Defaults.Contains(Scale3D))
+		{
+			const FString RotationValue = StructConstGenerators.FindChecked(TEXT("FQuat"))(Defaults.FindChecked(Rotation));
+			const FString TranslationValue = StructConstGenerators.FindChecked(TEXT("FVector"))(Defaults.FindChecked(Translation));
+			const FString Scale3DValue = StructConstGenerators.FindChecked(TEXT("FVector"))(Defaults.FindChecked(Scale3D));
+			return Format(Constructor, RotationValue, TranslationValue, Scale3DValue);
+		}
+		return FString();
+	});
+
+	StructConstGenerators.Add(TEXT("FEulerTransform"), [](const FString& InDefault) -> FString
+	{
+		static constexpr TCHAR Constructor[] = TEXT("FEulerTransform({0}, {1}, {2})");
+		const FStringMap Defaults = DefaultValueHelpers::SplitIntoMap(InDefault);
+		if(Defaults.Contains(Rotation) && Defaults.Contains(Translation) && Defaults.Contains(Scale3D))
+		{
+			const FString RotationValue = StructConstGenerators.FindChecked(TEXT("FRotator"))(Defaults.FindChecked(Rotation));
+			const FString TranslationValue = StructConstGenerators.FindChecked(TEXT("FVector"))(Defaults.FindChecked(Translation));
+			const FString Scale3DValue = StructConstGenerators.FindChecked(TEXT("FVector"))(Defaults.FindChecked(Scale3D));
+			return Format(Constructor, RotationValue, TranslationValue, Scale3DValue);
+		}
+		return FString();
+	});
+
+	StructConstGenerators.Add(TEXT("FRigElementKey"), [](const FString& InDefault) -> FString
+	{
+		static constexpr TCHAR Constructor[] = TEXT("FRigElementKey(TEXT(\"{0}\"), ERigElementType::{1})");
+		const FStringMap Defaults = DefaultValueHelpers::SplitIntoMap(InDefault);
+		if(Defaults.Contains(Type) && Defaults.Contains(Name))
+		{
+			const FString NameValue = DefaultValueHelpers::RemoveQuotes(Defaults.FindChecked(Name));
+			return Format(Constructor, NameValue, Defaults.FindChecked(Type));
+		}
+		return FString();
+	});
+
+	return StructConstGenerators;
+}
diff --git a/Engine/Source/Developer/RigVMDeveloper/Private/RigVMCompiler/RigVMCompiler.cpp b/Engine/Source/Developer/RigVMDeveloper/Private/RigVMCompiler/RigVMCompiler.cpp
--- a/Engine/Source/Developer/RigVMDeveloper/Private/RigVMCompiler/RigVMCompiler.cpp
+++ b/Engine/Source/Developer/RigVMDeveloper/Private/RigVMCompiler/RigVMCompiler.cpp
@@ -3,6 +3,7 @@
 #include "RigVMCompiler/RigVMCompiler.h"
 #include "RigVMModel/RigVMController.h"
 #include "RigVMCore/RigVMExecuteContext.h"
+#include "RigVMCore/RigVMNativized.h"
 #include "RigVMDeveloperModule.h"
 #include "UObject/PropertyPortFlags.h"
 #include "UObject/Interface.h"
@@ -141,6 +142,12 @@ bool URigVMCompiler::Compile(URigVMGraph* InGraph, URigVMController* InControlle
 		return false;
 	}
 
+	if(OutVM->GetClass()->IsChildOf(URigVMNativized::StaticClass()))
+	{
+		ReportError(TEXT("Provided vm is nativized."));
+		return false;
+	}
+
 	DECLARE_SCOPE_HIERARCHICAL_COUNTER_FUNC()
 
 	TArray<FRigVMUserDataArray> UserData = InRigVMUserData;
diff --git a/Engine/Source/Developer/RigVMDeveloper/Public/RigVMCompiler/RigVMCodeGenerator.h b/Engine/Source/Developer/RigVMDeveloper/Public/RigVMCompiler/RigVMCodeGenerator.h
new file mode 100644
--- /dev/null
+++ b/Engine/Source/Developer/RigVMDeveloper/Public/RigVMCompiler/RigVMCodeGenerator.h
@@ -0,0 +1,241 @@
+// Copyright Epic Games, Inc. All Rights Reserved.
+
+#pragma once
+
+#include "CoreMinimal.h"
+
+#include "RigVMCore/RigVM.h"
+#include "RigVMCore/RigVMStruct.h"
+#include "RigVMModel/RigVMGraph.h"
+#include "RigVMCompiler/RigVMCompiler.h"
+
+struct RIGVMDEVELOPER_API FRigVMCodeGenerator
+{
+public:
+
+	FRigVMCodeGenerator(const FString& InClassName, const FString& InModuleName,
+		URigVMGraph* InModelToNativize, URigVM* InVMToNativize, TMap<FString,FRigVMOperand> InPinToOperandMap,
+		int32 InMaxOperationsPerFunction = 100)
+	{
+		ParseVM(InClassName, InModuleName, InModelToNativize, InVMToNativize, InPinToOperandMap, InMaxOperationsPerFunction);
+	}
+
+	FString DumpIncludes(bool bLog = false);
+	FString DumpWrappedArrayTypes(bool bLog = false);
+	FString DumpExternalVariables(bool bForHeader, bool bLog = false);
+	FString DumpEntries(bool bForHeader, bool bLog = false);
+	FString DumpProperties(bool bForHeader, int32 InOperationGroup, bool bLog = false);
+	FString DumpInstructions(int32 InOperationGroup, bool bLog = false);
+	FString DumpHeader(bool bLog = false);
+	FString DumpSource(bool bLog = false);
+	FString DumpLines(const TArray<FString>& InLines, bool bLog = false);
+	
+private:
+
+	typedef TArray<FString> FStringArray;
+	typedef TMap<FString, FString> FStringMap;
+	typedef TArray<FRigVMPropertyDescription> FRigVMPropertyDescriptionArray;
+	typedef TTuple<FString, FString> FMappedType;
+	typedef TTuple<FString, FString, FString> FOpaqueArgument;
+	typedef TFunction<FString(const FString&)> TStructConstGenerator;
+
+	enum ERigVMNativizedPropertyType
+	{
+		Literal,
+		Work,
+		Sliced,
+		Invalid
+	};
+
+	struct FPropertyInfo
+	{
+		FRigVMPropertyDescription Description;
+		int32 MemoryPropertyIndex;
+		ERigVMNativizedPropertyType PropertyType;
+		TArray<int32> Groups;
+	};
+
+	struct FOperationGroup
+	{
+		FName Entry;
+		int32 First;
+		int32 Last;
+		int32 Depth;
+		int32 ParentGroup;
+		TArray<int32> ChildGroups;
+		TArray<int32> RequiredLabels;
+		TArray<FOpaqueArgument> OpaqueArguments;
+
+		FOperationGroup()
+			: Entry(NAME_None)
+			, First(INDEX_NONE)
+			, Last(INDEX_NONE)
+			, Depth(INDEX_NONE)
+			, ParentGroup(INDEX_NONE)
+		{}
+
+		FOperationGroup(const FOperationGroup& InOther)
+			: Entry(InOther.Entry)
+			, First(InOther.First)
+			, Last(InOther.Last)
+			, Depth(InOther.Depth)
+			, ParentGroup(InOther.ParentGroup)
+			, ChildGroups(InOther.ChildGroups)
+			, RequiredLabels(InOther.RequiredLabels)
+			, OpaqueArguments(InOther.OpaqueArguments)
+		{}
+	};
+
+	void Reset();
+	void ParseVM(const FString& InClassName, const FString& InModuleName,
+		URigVMGraph* InModelToNativize, URigVM* InVMToNativize,
+		TMap<FString,FRigVMOperand> InPinToOperandMap, int32 InMaxOperationsPerFunction);
+	void ParseInclude(UStruct* InDependency, const FName& InMethodName = NAME_None);
+	void ParseMemory(URigVMMemoryStorage* InMemory);
+	void ParseProperty(ERigVMMemoryType InMemoryType, const FProperty* InProperty, URigVMMemoryStorage* InMemory);
+	void ParseOperationGroups();
+	FString GetOperandName(const FRigVMOperand& InOperand, bool bPerSlice = true, bool bAsInput = true) const;
+	FString GetOperandCPPType(const FRigVMOperand& InOperand) const;
+	FString GetOperandCPPBaseType(const FRigVMOperand& InOperand) const;
+	static FString SanitizeName(const FString& InName, const FString& CPPType = FString()); 
+	FRigVMPropertyDescription GetPropertyForOperand(const FRigVMOperand& InOperand) const;
+	const FRigVMPropertyPathDescription& GetPropertyPathForOperand(const FRigVMOperand& InOperand) const;
+	int32 GetPropertyIndex(const FRigVMPropertyDescription& InProperty) const;
+	int32 GetPropertyIndex(const FRigVMOperand& InOperand) const;
+	ERigVMNativizedPropertyType GetPropertyType(const FRigVMPropertyDescription& InProperty) const;
+	ERigVMNativizedPropertyType GetPropertyType(const FRigVMOperand& InOperand) const;
+	void CheckOperand(const FRigVMOperand& InOperand) const;
+	FString GetMappedType(const FString& InCPPType) const;
+	const FString& GetMappedTypeSuffix(const FString& InCPPType) const;
+	FString GetMappedArrayTypeName(const FString InBaseElementType) const;
+	const FOperationGroup& GetGroup(int32 InGroupIndex) const;
+	bool IsOperationPartOfGroup(int32 InOperationIndex, int32 InGroupIndex, bool bIncludeChildGroups) const;
+	bool IsPropertyPartOfGroup(int32 InPropertyIndex, int32 InGroupIndex) const;
+
+	FORCEINLINE static FString FormatArgs(const TCHAR* InFormatString, const FStringFormatOrderedArguments& InArgs)
+	{
+		return FString::Format(InFormatString, InArgs);
+	}
+
+	template<typename TypeA>
+	FORCEINLINE static FString Format(const TCHAR* InFormatString, const TypeA& InArgA)
+	{
+		return FormatArgs(InFormatString, {
+			FStringFormatArg(InArgA)
+		});
+	}
+
+	template<typename TypeA, typename TypeB>
+	FORCEINLINE static FString Format(const TCHAR* InFormatString, const TypeA& InArgA, const TypeB& InArgB)
+	{
+		return FormatArgs(InFormatString, {
+			FStringFormatArg(InArgA),
+			FStringFormatArg(InArgB)
+		});
+	}
+
+	template<typename TypeA, typename TypeB, typename TypeC>
+	FORCEINLINE static FString Format(const TCHAR* InFormatString, const TypeA& InArgA, const TypeB& InArgB, const TypeC& InArgC)
+	{
+		return FormatArgs(InFormatString, {
+			FStringFormatArg(InArgA),
+			FStringFormatArg(InArgB),
+			FStringFormatArg(InArgC)
+		});
+	}
+
+	template<typename TypeA, typename TypeB, typename TypeC, typename TypeD>
+	FORCEINLINE static FString Format(const TCHAR* InFormatString,
+		const TypeA& InArgA,
+		const TypeB& InArgB,
+		const TypeC& InArgC,
+		const TypeD& InArgD
+	)
+	{
+		return FormatArgs(InFormatString, {
+			FStringFormatArg(InArgA),
+			FStringFormatArg(InArgB),
+			FStringFormatArg(InArgC),
+			FStringFormatArg(InArgD)
+		});
+	}
+
+	template<typename TypeA, typename TypeB, typename TypeC, typename TypeD, typename TypeE>
+	FORCEINLINE static FString Format(const TCHAR* InFormatString,
+		const TypeA& InArgA,
+		const TypeB& InArgB,
+		const TypeC& InArgC,
+		const TypeD& InArgD,
+		const TypeE& InArgE
+	)
+	{
+		return FormatArgs(InFormatString, {
+			FStringFormatArg(InArgA),
+			FStringFormatArg(InArgB),
+			FStringFormatArg(InArgC),
+			FStringFormatArg(InArgD),
+			FStringFormatArg(InArgE)
+		});
+	}
+
+	template<typename TypeA, typename TypeB, typename TypeC, typename TypeD, typename TypeE, typename TypeF>
+	FORCEINLINE static FString Format(const TCHAR* InFormatString,
+		const TypeA& InArgA,
+		const TypeB& InArgB,
+		const TypeC& InArgC,
+		const TypeD& InArgD,
+		const TypeE& InArgE,
+		const TypeF& InArgF
+	)
+	{
+		return FormatArgs(InFormatString, {
+			FStringFormatArg(InArgA),
+			FStringFormatArg(InArgB),
+			FStringFormatArg(InArgC),
+			FStringFormatArg(InArgD),
+			FStringFormatArg(InArgE),
+			FStringFormatArg(InArgF)
+		});
+	}
+
+	template<typename TypeA, typename TypeB, typename TypeC, typename TypeD, typename TypeE, typename TypeF, typename TypeG>
+	FORCEINLINE static FString Format(const TCHAR* InFormatString,
+		const TypeA& InArgA,
+		const TypeB& InArgB,
+		const TypeC& InArgC,
+		const TypeD& InArgD,
+		const TypeE& InArgE,
+		const TypeF& InArgF,
+		const TypeG& InArgG
+	)
+	{
+		return FormatArgs(InFormatString, {
+			FStringFormatArg(InArgA),
+			FStringFormatArg(InArgB),
+			FStringFormatArg(InArgC),
+			FStringFormatArg(InArgD),
+			FStringFormatArg(InArgE),
+			FStringFormatArg(InArgF),
+			FStringFormatArg(InArgG)
+		});
+	}
+
+	TStrongObjectPtr<URigVMGraph> Model;
+	TStrongObjectPtr<URigVM> VM;
+	TMap<FString,FRigVMOperand> PinToOperandMap;
+	TMap<FRigVMOperand, FString> OperandToPinMap;
+	int32 MaxOperationsPerFunction;
+	FString ClassName;
+	FString ModuleName;
+	
+	FStringArray Libraries;
+	FStringArray Includes;
+	TArray<FOperationGroup> OperationGroups;
+	FStringArray WrappedArrayCPPTypes;
+	TMap<FString, FMappedType> MappedCPPTypes;
+	TArray<FPropertyInfo> Properties;
+	TMap<FName, int32> PropertyNameToIndex;
+	TMap<FName, FName> MappedPropertyNames;
+
+	static const TMap<FName, TStructConstGenerator>& GetStructConstGenerators();
+};
diff --git a/Engine/Source/Runtime/RigVM/Private/RigVMCore/RigVM.cpp b/Engine/Source/Runtime/RigVM/Private/RigVMCore/RigVM.cpp
--- a/Engine/Source/Runtime/RigVM/Private/RigVMCore/RigVM.cpp
+++ b/Engine/Source/Runtime/RigVM/Private/RigVMCore/RigVM.cpp
@@ -1,6 +1,7 @@
 // Copyright Epic Games, Inc. All Rights Reserved.
 
 #include "RigVMCore/RigVM.h"
+#include "RigVMCore/RigVMNativized.h"
 #include "UObject/Package.h"
 #include "UObject/AnimObjectVersion.h"
 #include "UObject/UE5MainStreamObjectVersion.h"
@@ -254,6 +255,61 @@ void URigVM::PostLoad()
 	}
 }
 
+uint32 URigVM::GetVMHash() const
+{
+	uint32 Hash = 0;
+	for(const FName& FunctionName : GetFunctionNames())
+	{
+		Hash = HashCombine(Hash, GetTypeHash(FunctionName.ToString()));
+	}
+
+	Hash = HashCombine(Hash, GetTypeHash(GetByteCode()));
+
+	for(const FRigVMExternalVariable& ExternalVariable : ExternalVariables)
+	{
+		Hash = HashCombine(Hash, GetTypeHash(ExternalVariable.Name.ToString()));
+		Hash = HashCombine(Hash, GetTypeHash(ExternalVariable.TypeName.ToString()));
+	}
+
+	if(LiteralMemoryStorageObject)
+	{
+		Hash = HashCombine(Hash, LiteralMemoryStorageObject->GetMemoryHash());
+	}
+	if(WorkMemoryStorageObject)
+	{
+		Hash = HashCombine(Hash, WorkMemoryStorageObject->GetMemoryHash());
+	}
+	
+	return Hash;
+}
+
+UClass* URigVM::GetNativizedClass(const TArray<FRigVMExternalVariable>& InExternalVariables)
+{
+	TSharedPtr<TGuardValue<TArray<FRigVMExternalVariable>>> GuardPtr;
+	if(!InExternalVariables.IsEmpty())
+	{
+		GuardPtr = MakeShareable(new TGuardValue<TArray<FRigVMExternalVariable>>(ExternalVariables, InExternalVariables));
+	}
+
+	const uint32 VMHash = GetVMHash();
+	
+	for (TObjectIterator<UClass> ClassIterator; ClassIterator; ++ClassIterator)
+	{
+		if (ClassIterator->IsChildOf(URigVMNativized::StaticClass()) && (*ClassIterator != URigVMNativized::StaticClass()))
+		{
+			if(const URigVM* NativizedVMCDO = ClassIterator->GetDefaultObject<URigVM>())
+			{
+				if(NativizedVMCDO->GetVMHash() == VMHash)
+				{
+					return *ClassIterator;
+				}
+			}
+		}
+	}
+
+	return nullptr;
+}
+
 bool URigVM::ValidateAllOperandsDuringLoad()
 {
 	// check all operands on all ops for validity
@@ -910,7 +966,7 @@ void URigVM::CacheMemoryHandlesIfRequired(TArrayView<URigVMMemoryStorage*> InMem
 		return;
 	}
 
-	if (Instructions.Num() != FirstHandleForInstruction.Num())
+	if ((Instructions.Num() + 1) != FirstHandleForInstruction.Num())
 	{
 		InvalidateCachedMemory();
 	}
@@ -930,7 +986,7 @@ void URigVM::CacheMemoryHandlesIfRequired(TArrayView<URigVMMemoryStorage*> InMem
 		}
 	}
 
-	if (Instructions.Num() == FirstHandleForInstruction.Num())
+	if ((Instructions.Num() + 1) == FirstHandleForInstruction.Num())
 	{
 		return;
 	}
@@ -1577,7 +1633,7 @@ bool URigVM::Execute(TArrayView<URigVMMemoryStorage*> Memory, TArrayView<void*>
 		{
 			return false;
 		}
-		Context.PublicData.InstructionIndex = (uint16)ByteCode.GetEntry(EntryIndex).InstructionIndex;
+		SetInstructionIndex((uint16)ByteCode.GetEntry(EntryIndex).InstructionIndex);
 	}
 
 #if WITH_EDITOR
@@ -1999,6 +2055,7 @@ bool URigVM::Execute(TArrayView<URigVMMemoryStorage*> Memory, TArrayView<void*>
 				const FArrayProperty* ArrayProperty = CastFieldChecked<FArrayProperty>(ArrayHandle.GetProperty());
 				FScriptArrayHelper ArrayHelper(ArrayProperty, ArrayHandle.GetData());
 				int32 Index = (*((int32*)CachedMemoryHandles[FirstHandleForInstruction[Context.PublicData.InstructionIndex] + 1].GetData()));
+
 				if(Context.IsValidArrayIndex(Index, ArrayHelper))
 				{
 					FRigVMMemoryHandle& ElementHandle = CachedMemoryHandles[FirstHandleForInstruction[Context.PublicData.InstructionIndex] + 2];
@@ -2228,11 +2285,11 @@ bool URigVM::Execute(TArrayView<URigVMMemoryStorage*> Memory, TArrayView<void*>
 			case ERigVMOpCode::ArrayIntersection:
 			{
 				FRigVMMemoryHandle& ArrayHandleA = CachedMemoryHandles[FirstHandleForInstruction[Context.PublicData.InstructionIndex]]; 					
-				FRigVMMemoryHandle& ArrayHandleB = CachedMemoryHandles[FirstHandleForInstruction[Context.PublicData.InstructionIndex] + 1]; 					
-				FScriptArrayHelper ArrayHelperA(CastFieldChecked<FArrayProperty>(ArrayHandleA.GetProperty()), ArrayHandleA.GetData());
-				FScriptArrayHelper ArrayHelperB(CastFieldChecked<FArrayProperty>(ArrayHandleB.GetProperty()), ArrayHandleB.GetData());
-				const FArrayProperty* ArrayPropertyA = CastFieldChecked<FArrayProperty>(ArrayHandleA.GetProperty());
-				const FArrayProperty* ArrayPropertyB = CastFieldChecked<FArrayProperty>(ArrayHandleB.GetProperty());
+				FRigVMMemoryHandle& ArrayHandleB = CachedMemoryHandles[FirstHandleForInstruction[Context.PublicData.InstructionIndex] + 1];
+				const FArrayProperty* ArrayPropertyA = CastFieldChecked<FArrayProperty>(ArrayHandleA.GetProperty()); 					
+				const FArrayProperty* ArrayPropertyB = CastFieldChecked<FArrayProperty>(ArrayHandleB.GetProperty()); 					
+				FScriptArrayHelper ArrayHelperA(ArrayPropertyA, ArrayHandleA.GetData());
+				FScriptArrayHelper ArrayHelperB(ArrayPropertyB, ArrayHandleB.GetData());
 				const FProperty* ElementPropertyA = ArrayPropertyA->Inner;
 				const FProperty* ElementPropertyB = ArrayPropertyB->Inner;
 
diff --git a/Engine/Source/Runtime/RigVM/Private/RigVMCore/RigVMByteCode.cpp b/Engine/Source/Runtime/RigVM/Private/RigVMCore/RigVMByteCode.cpp
--- a/Engine/Source/Runtime/RigVM/Private/RigVMCore/RigVMByteCode.cpp
+++ b/Engine/Source/Runtime/RigVM/Private/RigVMCore/RigVMByteCode.cpp
@@ -690,6 +690,185 @@ void FRigVMByteCode::Empty()
 #endif
 }
 
+uint32 FRigVMByteCode::GetByteCodeHash() const
+{
+	uint32 Hash = 0;
+
+	for(int32 EntryIndex = 0; EntryIndex < NumEntries(); EntryIndex++)
+	{
+		Hash = HashCombine(Hash, GetTypeHash(GetEntry(EntryIndex).Name.ToString()));
+		Hash = HashCombine(Hash, GetTypeHash(GetEntry(EntryIndex).InstructionIndex));
+	}
+
+	const FRigVMInstructionArray Instructions = GetInstructions();
+	for(const FRigVMInstruction& Instruction : Instructions)
+	{
+		Hash = HashCombine(Hash, GetOperatorHash(Instruction));
+	}
+
+	return Hash;
+}
+
+uint32 FRigVMByteCode::GetOperatorHash(const FRigVMInstruction& InInstruction) const
+{
+	const ERigVMOpCode OpCode = GetOpCodeAt(InInstruction.ByteCodeIndex);
+	switch (OpCode)
+	{
+		case ERigVMOpCode::Execute_0_Operands:
+		case ERigVMOpCode::Execute_1_Operands:
+		case ERigVMOpCode::Execute_2_Operands:
+		case ERigVMOpCode::Execute_3_Operands:
+		case ERigVMOpCode::Execute_4_Operands:
+		case ERigVMOpCode::Execute_5_Operands:
+		case ERigVMOpCode::Execute_6_Operands:
+		case ERigVMOpCode::Execute_7_Operands:
+		case ERigVMOpCode::Execute_8_Operands:
+		case ERigVMOpCode::Execute_9_Operands:
+		case ERigVMOpCode::Execute_10_Operands:
+		case ERigVMOpCode::Execute_11_Operands:
+		case ERigVMOpCode::Execute_12_Operands:
+		case ERigVMOpCode::Execute_13_Operands:
+		case ERigVMOpCode::Execute_14_Operands:
+		case ERigVMOpCode::Execute_15_Operands:
+		case ERigVMOpCode::Execute_16_Operands:
+		case ERigVMOpCode::Execute_17_Operands:
+		case ERigVMOpCode::Execute_18_Operands:
+		case ERigVMOpCode::Execute_19_Operands:
+		case ERigVMOpCode::Execute_20_Operands:
+		case ERigVMOpCode::Execute_21_Operands:
+		case ERigVMOpCode::Execute_22_Operands:
+		case ERigVMOpCode::Execute_23_Operands:
+		case ERigVMOpCode::Execute_24_Operands:
+		case ERigVMOpCode::Execute_25_Operands:
+		case ERigVMOpCode::Execute_26_Operands:
+		case ERigVMOpCode::Execute_27_Operands:
+		case ERigVMOpCode::Execute_28_Operands:
+		case ERigVMOpCode::Execute_29_Operands:
+		case ERigVMOpCode::Execute_30_Operands:
+		case ERigVMOpCode::Execute_31_Operands:
+		case ERigVMOpCode::Execute_32_Operands:
+		case ERigVMOpCode::Execute_33_Operands:
+		case ERigVMOpCode::Execute_34_Operands:
+		case ERigVMOpCode::Execute_35_Operands:
+		case ERigVMOpCode::Execute_36_Operands:
+		case ERigVMOpCode::Execute_37_Operands:
+		case ERigVMOpCode::Execute_38_Operands:
+		case ERigVMOpCode::Execute_39_Operands:
+		case ERigVMOpCode::Execute_40_Operands:
+		case ERigVMOpCode::Execute_41_Operands:
+		case ERigVMOpCode::Execute_42_Operands:
+		case ERigVMOpCode::Execute_43_Operands:
+		case ERigVMOpCode::Execute_44_Operands:
+		case ERigVMOpCode::Execute_45_Operands:
+		case ERigVMOpCode::Execute_46_Operands:
+		case ERigVMOpCode::Execute_47_Operands:
+		case ERigVMOpCode::Execute_48_Operands:
+		case ERigVMOpCode::Execute_49_Operands:
+		case ERigVMOpCode::Execute_50_Operands:
+		case ERigVMOpCode::Execute_51_Operands:
+		case ERigVMOpCode::Execute_52_Operands:
+		case ERigVMOpCode::Execute_53_Operands:
+		case ERigVMOpCode::Execute_54_Operands:
+		case ERigVMOpCode::Execute_55_Operands:
+		case ERigVMOpCode::Execute_56_Operands:
+		case ERigVMOpCode::Execute_57_Operands:
+		case ERigVMOpCode::Execute_58_Operands:
+		case ERigVMOpCode::Execute_59_Operands:
+		case ERigVMOpCode::Execute_60_Operands:
+		case ERigVMOpCode::Execute_61_Operands:
+		case ERigVMOpCode::Execute_62_Operands:
+		case ERigVMOpCode::Execute_63_Operands:
+		case ERigVMOpCode::Execute_64_Operands:
+		{
+			const FRigVMExecuteOp& Op = GetOpAt<FRigVMExecuteOp>(InInstruction);
+			uint32 Hash = GetTypeHash(Op);
+			const FRigVMOperandArray Operands = GetOperandsForExecuteOp(InInstruction);
+			for(const FRigVMOperand& Operand : Operands)
+			{
+				Hash = HashCombine(Hash, GetTypeHash(Operand));
+			}
+			return Hash;
+		}
+		case ERigVMOpCode::Copy:
+		{
+			return GetTypeHash(GetOpAt<FRigVMCopyOp>(InInstruction));
+		}
+		case ERigVMOpCode::Zero:
+		case ERigVMOpCode::BoolFalse:
+		case ERigVMOpCode::BoolTrue:
+		case ERigVMOpCode::Increment:
+		case ERigVMOpCode::Decrement:
+		case ERigVMOpCode::ArrayReset:
+		case ERigVMOpCode::ArrayReverse:
+		{
+			return GetTypeHash(GetOpAt<FRigVMUnaryOp>(InInstruction));
+		}
+		case ERigVMOpCode::Equals:
+		case ERigVMOpCode::NotEquals:
+		{
+			return GetTypeHash(GetOpAt<FRigVMComparisonOp>(InInstruction));
+		}
+		case ERigVMOpCode::JumpAbsolute:
+		case ERigVMOpCode::JumpForward:
+		case ERigVMOpCode::JumpBackward:
+		{
+			return GetTypeHash(GetOpAt<FRigVMJumpOp>(InInstruction));
+		}
+		case ERigVMOpCode::JumpAbsoluteIf:
+		case ERigVMOpCode::JumpForwardIf:
+		case ERigVMOpCode::JumpBackwardIf:
+		{
+			return GetTypeHash(GetOpAt<FRigVMJumpIfOp>(InInstruction));
+		}
+		case ERigVMOpCode::ChangeType:
+		{
+			checkNoEntry();
+			return 0;
+		}
+		case ERigVMOpCode::Exit:
+		{
+			return GetTypeHash(GetOpAt<FRigVMBaseOp>(InInstruction));
+		}
+		case ERigVMOpCode::BeginBlock:
+		case ERigVMOpCode::ArrayGetNum:
+		case ERigVMOpCode::ArraySetNum:
+		case ERigVMOpCode::ArrayAppend:
+		case ERigVMOpCode::ArrayClone:
+		case ERigVMOpCode::ArrayRemove:
+		case ERigVMOpCode::ArrayUnion:
+		{
+			return GetTypeHash(GetOpAt<FRigVMBinaryOp>(InInstruction));
+		}
+		case ERigVMOpCode::ArrayAdd:
+		case ERigVMOpCode::ArrayGetAtIndex:
+		case ERigVMOpCode::ArraySetAtIndex:
+		case ERigVMOpCode::ArrayInsert:
+		case ERigVMOpCode::ArrayDifference:
+		case ERigVMOpCode::ArrayIntersection:
+		{				
+			return GetTypeHash(GetOpAt<FRigVMTernaryOp>(InInstruction));
+		}
+		case ERigVMOpCode::ArrayFind:
+		{				
+			return GetTypeHash(GetOpAt<FRigVMQuaternaryOp>(InInstruction));
+		}
+		case ERigVMOpCode::ArrayIterator:
+		{				
+			return GetTypeHash(GetOpAt<FRigVMSenaryOp>(InInstruction));
+		}
+		case ERigVMOpCode::EndBlock:
+		{
+			return GetTypeHash(GetOpAt<FRigVMBaseOp>(InInstruction));
+		}
+		case ERigVMOpCode::Invalid:
+		{
+			ensure(false);
+			return 0;
+		}
+	}
+	return 0;
+}
+
 uint64 FRigVMByteCode::Num() const
 {
 	return (uint64)ByteCode.Num();
@@ -1349,6 +1528,154 @@ uint64 FRigVMByteCode::AddArrayReverseOp(FRigVMOperand InArrayArg)
 	return AddOp(FRigVMUnaryOp(ERigVMOpCode::ArrayReverse, InArrayArg));
 }
 
+FRigVMOperandArray FRigVMByteCode::GetOperandsForOp(const FRigVMInstruction& InInstruction) const
+{
+	switch(InInstruction.OpCode)
+	{
+		case ERigVMOpCode::Execute_0_Operands:
+		case ERigVMOpCode::Execute_1_Operands:
+		case ERigVMOpCode::Execute_2_Operands:
+		case ERigVMOpCode::Execute_3_Operands:
+		case ERigVMOpCode::Execute_4_Operands:
+		case ERigVMOpCode::Execute_5_Operands:
+		case ERigVMOpCode::Execute_6_Operands:
+		case ERigVMOpCode::Execute_7_Operands:
+		case ERigVMOpCode::Execute_8_Operands:
+		case ERigVMOpCode::Execute_9_Operands:
+		case ERigVMOpCode::Execute_10_Operands:
+		case ERigVMOpCode::Execute_11_Operands:
+		case ERigVMOpCode::Execute_12_Operands:
+		case ERigVMOpCode::Execute_13_Operands:
+		case ERigVMOpCode::Execute_14_Operands:
+		case ERigVMOpCode::Execute_15_Operands:
+		case ERigVMOpCode::Execute_16_Operands:
+		case ERigVMOpCode::Execute_17_Operands:
+		case ERigVMOpCode::Execute_18_Operands:
+		case ERigVMOpCode::Execute_19_Operands:
+		case ERigVMOpCode::Execute_20_Operands:
+		case ERigVMOpCode::Execute_21_Operands:
+		case ERigVMOpCode::Execute_22_Operands:
+		case ERigVMOpCode::Execute_23_Operands:
+		case ERigVMOpCode::Execute_24_Operands:
+		case ERigVMOpCode::Execute_25_Operands:
+		case ERigVMOpCode::Execute_26_Operands:
+		case ERigVMOpCode::Execute_27_Operands:
+		case ERigVMOpCode::Execute_28_Operands:
+		case ERigVMOpCode::Execute_29_Operands:
+		case ERigVMOpCode::Execute_30_Operands:
+		case ERigVMOpCode::Execute_31_Operands:
+		case ERigVMOpCode::Execute_32_Operands:
+		case ERigVMOpCode::Execute_33_Operands:
+		case ERigVMOpCode::Execute_34_Operands:
+		case ERigVMOpCode::Execute_35_Operands:
+		case ERigVMOpCode::Execute_36_Operands:
+		case ERigVMOpCode::Execute_37_Operands:
+		case ERigVMOpCode::Execute_38_Operands:
+		case ERigVMOpCode::Execute_39_Operands:
+		case ERigVMOpCode::Execute_40_Operands:
+		case ERigVMOpCode::Execute_41_Operands:
+		case ERigVMOpCode::Execute_42_Operands:
+		case ERigVMOpCode::Execute_43_Operands:
+		case ERigVMOpCode::Execute_44_Operands:
+		case ERigVMOpCode::Execute_45_Operands:
+		case ERigVMOpCode::Execute_46_Operands:
+		case ERigVMOpCode::Execute_47_Operands:
+		case ERigVMOpCode::Execute_48_Operands:
+		case ERigVMOpCode::Execute_49_Operands:
+		case ERigVMOpCode::Execute_50_Operands:
+		case ERigVMOpCode::Execute_51_Operands:
+		case ERigVMOpCode::Execute_52_Operands:
+		case ERigVMOpCode::Execute_53_Operands:
+		case ERigVMOpCode::Execute_54_Operands:
+		case ERigVMOpCode::Execute_55_Operands:
+		case ERigVMOpCode::Execute_56_Operands:
+		case ERigVMOpCode::Execute_57_Operands:
+		case ERigVMOpCode::Execute_58_Operands:
+		case ERigVMOpCode::Execute_59_Operands:
+		case ERigVMOpCode::Execute_60_Operands:
+		case ERigVMOpCode::Execute_61_Operands:
+		case ERigVMOpCode::Execute_62_Operands:
+		case ERigVMOpCode::Execute_63_Operands:
+		case ERigVMOpCode::Execute_64_Operands:
+		{
+			return GetOperandsForExecuteOp(InInstruction);
+		}
+		case ERigVMOpCode::Copy:
+		{
+			const FRigVMCopyOp& Op = GetOpAt<FRigVMCopyOp>(InInstruction.ByteCodeIndex);
+			return FRigVMOperandArray(&Op.Source, 2);
+		}
+		case ERigVMOpCode::Zero:
+		case ERigVMOpCode::BoolFalse:
+		case ERigVMOpCode::BoolTrue:
+		case ERigVMOpCode::Increment:
+		case ERigVMOpCode::Decrement:
+		case ERigVMOpCode::ArrayReset:
+		case ERigVMOpCode::ArrayReverse:
+		{
+			const FRigVMUnaryOp& Op = GetOpAt<FRigVMUnaryOp>(InInstruction.ByteCodeIndex);
+			return FRigVMOperandArray(&Op.Arg, 1);
+		}
+		case ERigVMOpCode::Equals:
+		case ERigVMOpCode::NotEquals:
+		{
+			const FRigVMComparisonOp& Op = GetOpAt<FRigVMComparisonOp>(InInstruction.ByteCodeIndex);
+			return FRigVMOperandArray(&Op.A, 3);
+		}
+		case ERigVMOpCode::JumpAbsoluteIf:
+		case ERigVMOpCode::JumpForwardIf:
+		case ERigVMOpCode::JumpBackwardIf:
+		{
+			const FRigVMJumpIfOp& Op = GetOpAt<FRigVMJumpIfOp>(InInstruction.ByteCodeIndex);
+			return FRigVMOperandArray(&Op.Arg, 1);
+		}
+		case ERigVMOpCode::BeginBlock:
+		case ERigVMOpCode::ArrayGetNum:
+		case ERigVMOpCode::ArraySetNum:
+		case ERigVMOpCode::ArrayAppend:
+		case ERigVMOpCode::ArrayClone:
+		case ERigVMOpCode::ArrayRemove:
+		case ERigVMOpCode::ArrayUnion:
+		{
+			const FRigVMBinaryOp& Op = GetOpAt<FRigVMBinaryOp>(InInstruction.ByteCodeIndex);
+			return FRigVMOperandArray(&Op.ArgA, 2);
+		}
+		case ERigVMOpCode::ArrayAdd:
+		case ERigVMOpCode::ArrayGetAtIndex:
+		case ERigVMOpCode::ArraySetAtIndex:
+		case ERigVMOpCode::ArrayInsert:
+		case ERigVMOpCode::ArrayDifference:
+		case ERigVMOpCode::ArrayIntersection:
+		{
+			const FRigVMTernaryOp& Op = GetOpAt<FRigVMTernaryOp>(InInstruction.ByteCodeIndex);
+			return FRigVMOperandArray(&Op.ArgA, 3);
+		}
+		case ERigVMOpCode::ArrayFind:
+		{
+			const FRigVMQuaternaryOp& Op = GetOpAt<FRigVMQuaternaryOp>(InInstruction.ByteCodeIndex);
+			return FRigVMOperandArray(&Op.ArgA, 4);
+		}
+		case ERigVMOpCode::ArrayIterator:
+		{
+			const FRigVMSenaryOp& Op = GetOpAt<FRigVMSenaryOp>(InInstruction.ByteCodeIndex);
+			return FRigVMOperandArray(&Op.ArgA, 6);
+		}
+		case ERigVMOpCode::JumpAbsolute:
+		case ERigVMOpCode::JumpForward:
+		case ERigVMOpCode::JumpBackward:
+		case ERigVMOpCode::ChangeType:
+		case ERigVMOpCode::Exit:
+		case ERigVMOpCode::EndBlock:
+		case ERigVMOpCode::Invalid:
+		default:
+		{
+			break;
+		}
+	}
+
+	return FRigVMOperandArray();
+}
+
 uint64 FRigVMByteCode::GetOpAlignment(ERigVMOpCode InOpCode) const
 {
 	switch (InOpCode)
diff --git a/Engine/Source/Runtime/RigVM/Private/RigVMCore/RigVMMemoryStorage.cpp b/Engine/Source/Runtime/RigVM/Private/RigVMCore/RigVMMemoryStorage.cpp
--- a/Engine/Source/Runtime/RigVM/Private/RigVMCore/RigVMMemoryStorage.cpp
+++ b/Engine/Source/Runtime/RigVM/Private/RigVMCore/RigVMMemoryStorage.cpp
@@ -2,6 +2,7 @@
 
 #include "RigVMCore/RigVMMemoryStorage.h"
 #include "RigVMModule.h"
+#include "RigVMTypeUtils.h"
 #include "UObject/Interface.h"
 
 ////////////////////////////////////////////////////////////////////////////////////////////////////////////////////////
@@ -20,6 +21,63 @@ FRigVMPropertyDescription::FRigVMPropertyDescription(const FProperty* InProperty
 	, Containers()
 	, DefaultValue(InDefaultValue)
 {
+	if(InName.IsNone() && Property != nullptr)
+	{
+		Name = Property->GetFName();
+	}
+
+	if(CPPType.IsEmpty() && Property != nullptr)
+	{
+		FString ExtendedType;
+		CPPType = Property->GetCPPType(&ExtendedType);
+		CPPType += ExtendedType;
+
+		if(CPPType == RigVMTypeUtils::UInt8Type)
+		{
+			if(CastField<FBoolProperty>(Property))
+			{
+				CPPType = RigVMTypeUtils::BoolType;
+			}
+		}
+		else if(CPPType == RigVMTypeUtils::UInt8ArrayType)
+		{
+			if(const FArrayProperty* ArrayProperty = CastField<FArrayProperty>(Property))
+			{
+				if(CastField<FBoolProperty>(ArrayProperty->Inner))
+				{
+					CPPType = RigVMTypeUtils::BoolArrayType;
+				}
+			}
+		}
+	}
+
+	if(CPPTypeObject == nullptr && Property != nullptr)
+	{
+		const FProperty* ValueProperty = Property;
+		if(const FArrayProperty* ArrayProperty = CastField<FArrayProperty>(ValueProperty))
+		{
+			ValueProperty = ArrayProperty->Inner;
+		}
+		if(const FStructProperty* StructProperty = CastField<FStructProperty>(ValueProperty))
+		{
+			CPPTypeObject = StructProperty->Struct;
+		}
+#if UE_RIGVM_UOBJECT_PROPERTIES_ENABLED
+		else if(const FObjectProperty* ObjectProperty = CastField<FObjectProperty>(ValueProperty))
+		{
+			CPPTypeObject = ObjectProperty->PropertyClass;
+		}
+#endif
+		else if(const FEnumProperty* EnumProperty = CastField<FEnumProperty>(ValueProperty))
+		{
+			CPPTypeObject = EnumProperty->GetEnum();
+		}
+		else if(const FByteProperty* ByteProperty = CastField<FByteProperty>(ValueProperty))
+		{
+			CPPTypeObject = ByteProperty->Enum;
+		}
+	}
+	
 	// make sure to use valid names only
 	SanitizeName();
 
@@ -219,6 +277,7 @@ void URigVMMemoryStorageGeneratorClass::PurgeClass(bool bRecompilingOnLoad)
 	Super::PurgeClass(bRecompilingOnLoad);
 
 	// clear our state as well
+	CachedMemoryHash = 0;
 	LinkedProperties.Reset();
 	PropertyPaths.Reset();
 	PropertyPathDescriptions.Reset();
@@ -418,6 +477,37 @@ bool URigVMMemoryStorageGeneratorClass::RemoveStorageClass(UObject* InOuter, ERi
 	return false;
 }
 
+uint32 URigVMMemoryStorageGeneratorClass::GetMemoryHash() const
+{
+	if(CachedMemoryHash != 0)
+	{
+		return CachedMemoryHash;
+	}
+	
+	CachedMemoryHash = 0;
+
+	for(const FProperty* Property : LinkedProperties)
+	{
+		CachedMemoryHash = HashCombine(CachedMemoryHash, GetTypeHash(Property->GetFName().ToString()));
+		CachedMemoryHash = HashCombine(CachedMemoryHash, GetTypeHash(Property->GetCPPType()));
+	}
+	
+	// for literals we also hash the content / defaults for each property
+	if(GetMemoryType() == ERigVMMemoryType::Literal)
+	{
+		if(URigVMMemoryStorage* CDO = Cast<URigVMMemoryStorage>(GetDefaultObject(true)))
+		{
+			for(const FProperty* Property : LinkedProperties)
+			{
+				FString DefaultValue;
+				Property->ExportTextItem_InContainer(DefaultValue, CDO, nullptr, nullptr, PPF_None, nullptr);
+				CachedMemoryHash = HashCombine(CachedMemoryHash, GetTypeHash(DefaultValue));
+			}
+		}
+	}
+	return CachedMemoryHash;
+}
+
 FRigVMMemoryStatistics URigVMMemoryStorageGeneratorClass::GetStatistics() const
 {
 	FRigVMMemoryStatistics Statistics;
@@ -622,6 +712,7 @@ FProperty* URigVMMemoryStorageGeneratorClass::AddProperty(URigVMMemoryStorageGen
 
 void URigVMMemoryStorageGeneratorClass::RefreshLinkedProperties()
 {
+	CachedMemoryHash = 0;
 	LinkedProperties.Reset();
 	const FProperty* Property = CastField<FProperty>(ChildProperties);
 	while(Property)
diff --git a/Engine/Source/Runtime/RigVM/Private/RigVMCore/RigVMNativized.cpp b/Engine/Source/Runtime/RigVM/Private/RigVMCore/RigVMNativized.cpp
new file mode 100644
--- /dev/null
+++ b/Engine/Source/Runtime/RigVM/Private/RigVMCore/RigVMNativized.cpp
@@ -0,0 +1,49 @@
+// Copyright Epic Games, Inc. All Rights Reserved.
+
+#include "RigVMCore/RigVMNativized.h"
+
+URigVMNativized::URigVMNativized()
+	: URigVM()
+	, TemporaryArrayIndex(INDEX_NONE)
+{
+}
+
+URigVMNativized::~URigVMNativized()
+{
+}
+
+void URigVMNativized::Serialize(FArchive& Ar)
+{
+	UObject::Serialize(Ar);
+}
+
+bool URigVMNativized::Initialize(TArrayView<URigVMMemoryStorage*> Memory, TArrayView<void*> AdditionalArguments,
+	bool bInitializeMemory)
+{
+	// nothing to do here 
+	return true;
+}
+
+bool URigVMNativized::Execute(TArrayView<URigVMMemoryStorage*> Memory, TArrayView<void*> AdditionalArguments,
+	const FName& InEntryName)
+{
+	// to be implemented by the generated code
+	return false;
+}
+
+const FRigVMInstructionArray& URigVMNativized::GetInstructions()
+{
+	static const FRigVMInstructionArray EmptyInstructions;
+	return EmptyInstructions;
+}
+
+const FRigVMExecuteContext& URigVMNativized::UpdateContext(TArrayView<void*> AdditionalArguments, int32 InNumberInstructions, const FName& InEntryName)
+{
+	UpdateExternalVariables();
+	
+	Context.Reset();
+	Context.OpaqueArguments = AdditionalArguments;
+	Context.SliceOffsets.AddZeroed(InNumberInstructions);
+	Context.PublicData.EventName = InEntryName;
+	return Context.PublicData;
+}
diff --git a/Engine/Source/Runtime/RigVM/Public/RigVMCore/RigVM.h b/Engine/Source/Runtime/RigVM/Public/RigVMCore/RigVM.h
--- a/Engine/Source/Runtime/RigVM/Public/RigVMCore/RigVM.h
+++ b/Engine/Source/Runtime/RigVM/Public/RigVMCore/RigVM.h
@@ -147,19 +147,28 @@ public:
 	virtual ~URigVM();
 
 	// UObject interface
-	void Serialize(FArchive& Ar);
-	void Save(FArchive& Ar);
-	void Load(FArchive& Ar);
-	void PostLoad() override;
+	virtual void Serialize(FArchive& Ar);
+	virtual void Save(FArchive& Ar);
+	virtual void Load(FArchive& Ar);
+	virtual void PostLoad() override;
+
+	// returns true if this is a nativized VM
+	virtual bool IsNativized() const { return false; }
+
+	// returns a unique hash to compare VMs
+	virtual uint32 GetVMHash() const;
+
+	// returns the VM's matching nativized class if it exists
+	UClass* GetNativizedClass(const TArray<FRigVMExternalVariable>& InExternalVariables = TArray<FRigVMExternalVariable>());
 	
 	// resets the container and maintains all memory
-	void Reset(bool IsIgnoringArchetypeRef = false);
+	virtual void Reset(bool IsIgnoringArchetypeRef = false);
 
 	// resets the container and removes all memory
-	void Empty();
+	virtual void Empty();
 
 	// resets the container and clones the input VM
-	void CopyFrom(URigVM* InVM, bool bDeferCopy = false, bool bReferenceLiteralMemory = false, bool bReferenceByteCode = false, bool bCopyExternalVariables = false, bool bCopyDynamicRegisters = false);
+	virtual void CopyFrom(URigVM* InVM, bool bDeferCopy = false, bool bReferenceLiteralMemory = false, bool bReferenceByteCode = false, bool bCopyExternalVariables = false, bool bCopyDynamicRegisters = false);
 
 	// sets the max array size allowed by this VM
 	FORCEINLINE void SetRuntimeSettings(FRigVMRuntimeSettings InRuntimeSettings)
@@ -168,30 +177,30 @@ public:
 	}
 
 	// Initializes all execute ops and their memory.
-	bool Initialize(TArrayView<URigVMMemoryStorage*> Memory, TArrayView<void*> AdditionalArguments, bool bInitializeMemory = true);
+	virtual bool Initialize(TArrayView<URigVMMemoryStorage*> Memory, TArrayView<void*> AdditionalArguments, bool bInitializeMemory = true);
 
 	// Executes the VM.
 	// You can optionally provide external memory to the execution
 	// and provide optional additional operands.
-	bool Execute(TArrayView<URigVMMemoryStorage*> Memory, TArrayView<void*> AdditionalArguments, const FName& InEntryName = NAME_None);
+	virtual bool Execute(TArrayView<URigVMMemoryStorage*> Memory, TArrayView<void*> AdditionalArguments, const FName& InEntryName = NAME_None);
 
 	// Executes the VM.
 	// You can optionally provide external memory to the execution
 	// and provide optional additional operands.
 	UFUNCTION(BlueprintCallable, Category = RigVM)
-	bool Execute(const FName& InEntryName = NAME_None);
+	virtual bool Execute(const FName& InEntryName = NAME_None);
 
 	// Add a function for execute instructions to this VM.
 	// Execute instructions can then refer to the function by index.
 	UFUNCTION()
-	int32 AddRigVMFunction(UScriptStruct* InRigVMStruct, const FName& InMethodName);
+	virtual int32 AddRigVMFunction(UScriptStruct* InRigVMStruct, const FName& InMethodName);
 
 	// Returns the name of a function given its index
 	UFUNCTION()
-	FString GetRigVMFunctionName(int32 InFunctionIndex) const;
+	virtual FString GetRigVMFunctionName(int32 InFunctionIndex) const;
 
 	// Returns a memory storage by type
-	URigVMMemoryStorage* GetMemoryByType(ERigVMMemoryType InMemoryType, bool bCreateIfNeeded = true);
+	virtual URigVMMemoryStorage* GetMemoryByType(ERigVMMemoryType InMemoryType, bool bCreateIfNeeded = true);
 	
 	// The default mutable work memory
 	FORCEINLINE URigVMMemoryStorage* GetWorkMemory(bool bCreateIfNeeded = true) { return GetMemoryByType(ERigVMMemoryType::Work, bCreateIfNeeded); }
@@ -213,7 +222,7 @@ public:
 	}
 
 	// Removes all memory from this VM
-	void ClearMemory();
+	virtual void ClearMemory();
 
 	UPROPERTY()
 	TObjectPtr<URigVMMemoryStorage> WorkMemoryStorageObject;
@@ -235,13 +244,13 @@ public:
 	FORCEINLINE const FRigVMByteCode& GetByteCode() const { return *ByteCodePtr; }
 
 	// Returns the instructions of the VM
-	const FRigVMInstructionArray& GetInstructions();
+	virtual const FRigVMInstructionArray& GetInstructions();
 
 	// Returns true if this VM's bytecode contains a given entry
-	bool ContainsEntry(const FName& InEntryName) const;
+	virtual bool ContainsEntry(const FName& InEntryName) const;
 
 	// Returns a list of all valid entry names for this VM's bytecode
-	TArray<FName> GetEntryNames() const;
+	virtual TArray<FName> GetEntryNames() const;
 
 #if WITH_EDITOR
 	
@@ -634,9 +643,13 @@ private:
 	UPROPERTY(transient)
 	FRigVMInstructionArray Instructions;
 
+protected:
 	UPROPERTY(transient)
 	FRigVMExtendedExecuteContext Context;
 
+	FORCEINLINE void SetInstructionIndex(uint16 InInstructionIndex) { Context.PublicData.InstructionIndex = InInstructionIndex; }
+
+private:
 	UPROPERTY(transient)
 	uint32 NumExecutions;
 
@@ -741,6 +754,6 @@ private:
 	FExecutionHaltedEvent OnExecutionHalted;
 #endif
 
-	
 	friend class URigVMCompiler;
+	friend struct FRigVMCodeGenerator;
 };
diff --git a/Engine/Source/Runtime/RigVM/Public/RigVMCore/RigVMByteCode.h b/Engine/Source/Runtime/RigVM/Public/RigVMCore/RigVMByteCode.h
--- a/Engine/Source/Runtime/RigVM/Public/RigVMCore/RigVMByteCode.h
+++ b/Engine/Source/Runtime/RigVM/Public/RigVMCore/RigVMByteCode.h
@@ -132,6 +132,11 @@ struct RIGVM_API FRigVMBaseOp
 	}
 
 	ERigVMOpCode OpCode;
+
+	friend FORCEINLINE uint32 GetTypeHash(const FRigVMBaseOp& Op)
+	{
+		return GetTypeHash(Op.OpCode);
+	}
 };
 
 
@@ -155,6 +160,11 @@ struct RIGVM_API FRigVMExecuteOp : public FRigVMBaseOp
 
 	uint16 FunctionIndex;
 
+	friend FORCEINLINE uint32 GetTypeHash(const FRigVMExecuteOp& Op)
+	{
+		return HashCombine(GetTypeHash((const FRigVMBaseOp&)Op), GetTypeHash(Op.FunctionIndex));
+	}
+
 	FORCEINLINE uint8 GetOperandCount() const { return uint8(OpCode) - uint8(ERigVMOpCode::Execute_0_Operands); }
 
 	void Serialize(FArchive& Ar);
@@ -198,6 +208,11 @@ struct RIGVM_API FRigVMUnaryOp : public FRigVMBaseOp
 
 	FRigVMOperand Arg;
 
+	friend FORCEINLINE uint32 GetTypeHash(const FRigVMUnaryOp& Op)
+	{
+		return HashCombine(GetTypeHash((const FRigVMBaseOp&)Op), GetTypeHash(Op.Arg));
+	}
+
 	void Serialize(FArchive& Ar);
 	FORCEINLINE friend FArchive& operator<<(FArchive& Ar, FRigVMUnaryOp& P)
 	{
@@ -238,6 +253,14 @@ struct RIGVM_API FRigVMBinaryOp : public FRigVMBaseOp
 	FRigVMOperand ArgA;
 	FRigVMOperand ArgB;
 
+	friend FORCEINLINE uint32 GetTypeHash(const FRigVMBinaryOp& Op)
+	{
+		uint32 Hash = GetTypeHash((const FRigVMBaseOp&)Op);
+		Hash = HashCombine(Hash, GetTypeHash(Op.ArgA));
+		Hash = HashCombine(Hash, GetTypeHash(Op.ArgB));
+		return Hash;
+	}
+
 	void Serialize(FArchive& Ar);
 	FORCEINLINE friend FArchive& operator<<(FArchive& Ar, FRigVMBinaryOp& P)
 	{
@@ -280,6 +303,15 @@ struct RIGVM_API FRigVMTernaryOp : public FRigVMBaseOp
 	FRigVMOperand ArgB;
 	FRigVMOperand ArgC;
 
+	friend FORCEINLINE uint32 GetTypeHash(const FRigVMTernaryOp& Op)
+	{
+		uint32 Hash = GetTypeHash((const FRigVMBaseOp&)Op);
+		Hash = HashCombine(Hash, GetTypeHash(Op.ArgA));
+		Hash = HashCombine(Hash, GetTypeHash(Op.ArgB));
+		Hash = HashCombine(Hash, GetTypeHash(Op.ArgC));
+		return Hash;
+	}
+
 	void Serialize(FArchive& Ar);
 	FORCEINLINE friend FArchive& operator<<(FArchive& Ar, FRigVMTernaryOp& P)
 	{
@@ -320,6 +352,16 @@ struct RIGVM_API FRigVMQuaternaryOp : public FRigVMBaseOp
 	FRigVMOperand ArgC;
 	FRigVMOperand ArgD;
 
+	friend FORCEINLINE uint32 GetTypeHash(const FRigVMQuaternaryOp& Op)
+	{
+		uint32 Hash = GetTypeHash((const FRigVMBaseOp&)Op);
+		Hash = HashCombine(Hash, GetTypeHash(Op.ArgA));
+		Hash = HashCombine(Hash, GetTypeHash(Op.ArgB));
+		Hash = HashCombine(Hash, GetTypeHash(Op.ArgC));
+		Hash = HashCombine(Hash, GetTypeHash(Op.ArgD));
+		return Hash;
+	}
+
 	void Serialize(FArchive& Ar);
 	FORCEINLINE friend FArchive& operator<<(FArchive& Ar, FRigVMQuaternaryOp& P)
 	{
@@ -360,6 +402,17 @@ struct RIGVM_API FRigVMQuinaryOp : public FRigVMBaseOp
 	FRigVMOperand ArgD;
 	FRigVMOperand ArgE;
 
+	friend FORCEINLINE uint32 GetTypeHash(const FRigVMQuinaryOp& Op)
+	{
+		uint32 Hash = GetTypeHash((const FRigVMBaseOp&)Op);
+		Hash = HashCombine(Hash, GetTypeHash(Op.ArgA));
+		Hash = HashCombine(Hash, GetTypeHash(Op.ArgB));
+		Hash = HashCombine(Hash, GetTypeHash(Op.ArgC));
+		Hash = HashCombine(Hash, GetTypeHash(Op.ArgD));
+		Hash = HashCombine(Hash, GetTypeHash(Op.ArgE));
+		return Hash;
+	}
+
 	void Serialize(FArchive& Ar);
 	FORCEINLINE friend FArchive& operator<<(FArchive& Ar, FRigVMQuinaryOp& P)
 	{
@@ -406,6 +459,18 @@ struct RIGVM_API FRigVMSenaryOp : public FRigVMBaseOp
 	FRigVMOperand ArgE;
 	FRigVMOperand ArgF;
 
+	friend FORCEINLINE uint32 GetTypeHash(const FRigVMSenaryOp& Op)
+	{
+		uint32 Hash = GetTypeHash((const FRigVMBaseOp&)Op);
+		Hash = HashCombine(Hash, GetTypeHash(Op.ArgA));
+		Hash = HashCombine(Hash, GetTypeHash(Op.ArgB));
+		Hash = HashCombine(Hash, GetTypeHash(Op.ArgC));
+		Hash = HashCombine(Hash, GetTypeHash(Op.ArgD));
+		Hash = HashCombine(Hash, GetTypeHash(Op.ArgE));
+		Hash = HashCombine(Hash, GetTypeHash(Op.ArgF));
+		return Hash;
+	}
+
 	void Serialize(FArchive& Ar);
 	FORCEINLINE friend FArchive& operator<<(FArchive& Ar, FRigVMSenaryOp& P)
 	{
@@ -466,6 +531,14 @@ public:
 	FRigVMOperand Source;
 	FRigVMOperand Target;
 
+	friend FORCEINLINE uint32 GetTypeHash(const FRigVMCopyOp& Op)
+	{
+		uint32 Hash = GetTypeHash((const FRigVMBaseOp&)Op);
+		Hash = HashCombine(Hash, GetTypeHash(Op.Source));
+		Hash = HashCombine(Hash, GetTypeHash(Op.Target));
+		return Hash;
+	}
+
 private:
 	uint16 NumBytes;
 	ERigVMRegisterType RegisterType;
@@ -515,6 +588,15 @@ struct RIGVM_API FRigVMComparisonOp : public FRigVMBaseOp
 	FRigVMOperand B;
 	FRigVMOperand Result;
 
+	friend FORCEINLINE uint32 GetTypeHash(const FRigVMComparisonOp& Op)
+	{
+		uint32 Hash = GetTypeHash((const FRigVMBaseOp&)Op);
+		Hash = HashCombine(Hash, GetTypeHash(Op.A));
+		Hash = HashCombine(Hash, GetTypeHash(Op.B));
+		Hash = HashCombine(Hash, GetTypeHash(Op.Result));
+		return Hash;
+	}
+
 	void Serialize(FArchive& Ar);
 	FORCEINLINE friend FArchive& operator<<(FArchive& Ar, FRigVMComparisonOp& P)
 	{
@@ -547,6 +629,13 @@ struct RIGVM_API FRigVMJumpOp : public FRigVMBaseOp
 
 	int32 InstructionIndex;
 
+	friend FORCEINLINE uint32 GetTypeHash(const FRigVMJumpOp& Op)
+	{
+		uint32 Hash = GetTypeHash((const FRigVMBaseOp&)Op);
+		Hash = HashCombine(Hash, GetTypeHash(Op.InstructionIndex));
+		return Hash;
+	}
+
 	void Serialize(FArchive& Ar);
 	FORCEINLINE friend FArchive& operator<<(FArchive& Ar, FRigVMJumpOp& P)
 	{
@@ -582,6 +671,14 @@ struct RIGVM_API FRigVMJumpIfOp : public FRigVMUnaryOp
 	int32 InstructionIndex;
 	bool Condition;
 
+	friend FORCEINLINE uint32 GetTypeHash(const FRigVMJumpIfOp& Op)
+	{
+		uint32 Hash = GetTypeHash((const FRigVMBaseOp&)Op);
+		Hash = HashCombine(Hash, GetTypeHash(Op.InstructionIndex));
+		Hash = HashCombine(Hash, GetTypeHash(Op.Condition));
+		return Hash;
+	}
+
 	void Serialize(FArchive& Ar);
 	FORCEINLINE friend FArchive& operator<<(FArchive& Ar, FRigVMJumpIfOp& P)
 	{
@@ -720,6 +817,18 @@ public:
 	// resets the container and removes all memory
 	void Empty();
 
+	// returns a unique hash identifying this bytecode
+	uint32 GetByteCodeHash() const;
+
+	// returns a unique hash for an operator at a given instruction index
+	uint32 GetOperatorHash(const FRigVMInstruction& InInstruction) const;
+
+	// overload for GetTypeHash
+	friend FORCEINLINE uint32 GetTypeHash(const FRigVMByteCode& InByteCode)
+	{
+		return InByteCode.GetByteCodeHash();
+	}
+	
 	// returns the number of instructions in this container
 	uint64 Num() const;
 
@@ -892,6 +1001,9 @@ public:
 		return GetOperandsAt(ByteCodeIndex, ExecuteOp.GetOperandCount());
 	}
 
+	// returns all of the operands for a given instruction
+	FRigVMOperandArray GetOperandsForOp(const FRigVMInstruction& InInstruction) const;
+
 	// returns the raw data of the byte code
 	FORCEINLINE const TArrayView<const uint8> GetByteCode() const
 	{
diff --git a/Engine/Source/Runtime/RigVM/Public/RigVMCore/RigVMCore.h b/Engine/Source/Runtime/RigVM/Public/RigVMCore/RigVMCore.h
new file mode 100644
--- /dev/null
+++ b/Engine/Source/Runtime/RigVM/Public/RigVMCore/RigVMCore.h
@@ -0,0 +1,6 @@
+// Copyright Epic Games, Inc. All Rights Reserved.
+
+#pragma once
+
+#include "CoreMinimal.h"
+#include "RigVMNativized.h"
diff --git a/Engine/Source/Runtime/RigVM/Public/RigVMCore/RigVMExecuteContext.h b/Engine/Source/Runtime/RigVM/Public/RigVMCore/RigVMExecuteContext.h
--- a/Engine/Source/Runtime/RigVM/Public/RigVMCore/RigVMExecuteContext.h
+++ b/Engine/Source/Runtime/RigVM/Public/RigVMCore/RigVMExecuteContext.h
@@ -267,6 +267,7 @@ private:
 
 	friend struct FRigVMExtendedExecuteContext;
 	friend class URigVM;
+	friend class URigVMNativized;
 	friend class UControlRig;
 	friend class UAdditiveControlRig;
 	friend struct FRigUnit_BeginExecution;
@@ -346,19 +347,24 @@ struct RIGVM_API FRigVMExtendedExecuteContext
 	}
 
 	FORCEINLINE bool IsValidArrayIndex(int32& InOutIndex, const FScriptArrayHelper& InArrayHelper) const
+	{
+		return IsValidArrayIndex(InOutIndex, InArrayHelper.Num());
+	}
+
+	FORCEINLINE bool IsValidArrayIndex(int32& InOutIndex, int32 InArraySize) const
 	{
 		const int32 InOriginalIndex = InOutIndex;
 
 		// we support wrapping the index around similar to python
 		if(InOutIndex < 0)
 		{
-			InOutIndex = InArrayHelper.Num() + InOutIndex;
+			InOutIndex = InArraySize + InOutIndex;
 		}
 
-		if(!InArrayHelper.IsValidIndex(InOutIndex))
+		if(InOutIndex < 0 || InOutIndex >= InArraySize)
 		{
 			static const TCHAR OutOfBoundsFormat[] = TEXT("Array Index (%d) out of bounds (count %d).");
-			PublicData.Logf(EMessageSeverity::Error, OutOfBoundsFormat, InOriginalIndex, InArrayHelper.Num());
+			PublicData.Logf(EMessageSeverity::Error, OutOfBoundsFormat, InOriginalIndex, InArraySize);
 			return false;
 		}
 		return true;
diff --git a/Engine/Source/Runtime/RigVM/Public/RigVMCore/RigVMExternalVariable.h b/Engine/Source/Runtime/RigVM/Public/RigVMCore/RigVMExternalVariable.h
--- a/Engine/Source/Runtime/RigVM/Public/RigVMCore/RigVMExternalVariable.h
+++ b/Engine/Source/Runtime/RigVM/Public/RigVMCore/RigVMExternalVariable.h
@@ -7,6 +7,7 @@
 #include "RigVMDefines.h"
 #include "RigVMMemory.h"
 #include "RigVMModule.h"
+#include "RigVMTypeUtils.h"
 #include "UObject/UnrealType.h"
 
 /**
@@ -555,6 +556,15 @@ struct RIGVM_API FRigVMExternalVariable
 			(bAllowNullPtr || Memory != nullptr);
 	}
 
+	FORCEINLINE FName GetExtendedCPPType() const
+	{
+		if(bIsArray)
+		{
+			return *RigVMTypeUtils::ArrayTypeFromBaseType(TypeName.ToString());
+		}
+		return TypeName;
+	}
+
 	FORCEINLINE static void MergeExternalVariable(TArray<FRigVMExternalVariable>& OutVariables, const FRigVMExternalVariable& InVariable)
 	{
 		if(!InVariable.IsValid(true))
diff --git a/Engine/Source/Runtime/RigVM/Public/RigVMCore/RigVMMemoryStorage.h b/Engine/Source/Runtime/RigVM/Public/RigVMCore/RigVMMemoryStorage.h
--- a/Engine/Source/Runtime/RigVM/Public/RigVMCore/RigVMMemoryStorage.h
+++ b/Engine/Source/Runtime/RigVM/Public/RigVMCore/RigVMMemoryStorage.h
@@ -216,6 +216,7 @@ public:
 	URigVMMemoryStorageGeneratorClass()
 		: Super()
 		, MemoryType(ERigVMMemoryType::Literal)
+		, CachedMemoryHash(0)
 	{}
 	
 	// UClass overrides
@@ -255,6 +256,9 @@ public:
 	// Returns the type of memory of this class (literal, work, etc)
 	ERigVMMemoryType GetMemoryType() const { return MemoryType; }
 
+	// Returns a hash of unique to the configuration of the memory
+	uint32 GetMemoryHash() const;
+
 	// The properties stored within this class
 	const TArray<const FProperty*>& GetProperties() const { return LinkedProperties; }
 
@@ -288,9 +292,12 @@ private:
 	// A list of decriptions for the property paths - used for serialization
 	TArray<FRigVMPropertyPathDescription> PropertyPathDescriptions;
 
+	mutable uint32 CachedMemoryHash;
+
 	friend class URigVMMemoryStorage;
 	friend class URigVMCompiler;
 	friend class URigVM;
+	friend struct FRigVMCodeGenerator;
 };
 
 /**
@@ -319,6 +326,17 @@ public:
 		return ERigVMMemoryType::Debug;
 	}
 
+	// Returns a hash of unique to the configuration of the memory
+	FORCEINLINE uint32 GetMemoryHash() const
+	{
+		if(URigVMMemoryStorageGeneratorClass* Class = Cast<URigVMMemoryStorageGeneratorClass>(GetClass()))
+		{
+			return Class->GetMemoryHash();
+		}
+		// empty debug containers don't have a generator class
+		return 0;
+	}
+
 	// Returns the number of properties stored in this instance
 	FORCEINLINE int32 Num() const
 	{
diff --git a/Engine/Source/Runtime/RigVM/Public/RigVMCore/RigVMNativized.h b/Engine/Source/Runtime/RigVM/Public/RigVMCore/RigVMNativized.h
new file mode 100644
--- /dev/null
+++ b/Engine/Source/Runtime/RigVM/Public/RigVMCore/RigVMNativized.h
@@ -0,0 +1,381 @@
+// Copyright Epic Games, Inc. All Rights Reserved.
+
+#pragma once
+
+#include "RigVM.h"
+#include "RigVMNativized.generated.h"
+
+UCLASS(BlueprintType, Abstract)
+class RIGVM_API URigVMNativized : public URigVM
+{
+	GENERATED_BODY()
+
+public:
+
+	URigVMNativized();
+	virtual ~URigVMNativized();
+
+	// URigVM interface
+	virtual void Serialize(FArchive& Ar) override;
+	virtual void PostLoad() override { return; }
+	virtual void Reset(bool IsIgnoringArchetypeRef = false) override { return; }
+	virtual bool IsNativized() const override { return true; }
+	virtual void Empty() override { return; }
+	virtual void CopyFrom(URigVM* InVM, bool bDeferCopy = false, bool bReferenceLiteralMemory = false, bool bReferenceByteCode = false, bool bCopyExternalVariables = false, bool bCopyDynamicRegisters = false) override { return; }
+	virtual bool Initialize(TArrayView<URigVMMemoryStorage*> Memory, TArrayView<void*> AdditionalArguments, bool bInitializeMemory = true) override;
+	virtual bool Execute(TArrayView<URigVMMemoryStorage*> Memory, TArrayView<void*> AdditionalArguments, const FName& InEntryName = NAME_None) override;
+	virtual int32 AddRigVMFunction(UScriptStruct* InRigVMStruct, const FName& InMethodName) override { return INDEX_NONE; }
+	virtual FString GetRigVMFunctionName(int32 InFunctionIndex) const override { return FString(); }
+	virtual URigVMMemoryStorage* GetMemoryByType(ERigVMMemoryType InMemoryType, bool bCreateIfNeeded = true) override { return nullptr; }
+	virtual void ClearMemory() override { return; }
+	virtual const FRigVMInstructionArray& GetInstructions() override;
+	virtual bool ContainsEntry(const FName& InEntryName) const override { return GetEntryNames().Contains(InEntryName); }
+	virtual TArray<FName> GetEntryNames() const override { return TArray<FName>(); }
+	
+protected:
+
+	const FRigVMExecuteContext& UpdateContext(TArrayView<void*> AdditionalArguments, int32 InNumberInstructions, const FName& InEntryName);
+	virtual void UpdateExternalVariables() {};
+	
+	class FErrorPipe : public FOutputDevice
+	{
+	public:
+
+		int32 NumErrors;
+
+		FErrorPipe()
+			: FOutputDevice()
+			, NumErrors(0)
+		{
+		}
+
+		virtual void Serialize(const TCHAR* V, ELogVerbosity::Type Verbosity, const class FName& Category) override
+		{
+			NumErrors++;
+		}
+	};
+
+	FORCEINLINE void BeginSlice(int32 InCount, int32 InRelativeIndex = 0)
+	{
+		Context.BeginSlice(InCount, InRelativeIndex);
+	}
+
+	FORCEINLINE void EndSlice()
+	{
+		Context.EndSlice();
+	}
+
+	FORCEINLINE bool IsValidArraySize(int32 InArraySize) const
+	{
+		return Context.IsValidArraySize(InArraySize);
+	}
+
+	template<typename T>
+	FORCEINLINE bool IsValidArrayIndex(int32& InOutIndex, const TArray<T>& InArray) const
+	{
+		return Context.IsValidArrayIndex(InOutIndex, InArray.Num());
+	}
+
+	template<typename T>
+	FORCEINLINE static const T& GetArrayElementSafe(const TArray<T>& InArray, const int32 InIndex)
+	{
+		if(InArray.IsValidIndex(InIndex))
+		{
+			return InArray[InIndex];
+		}
+		static const T EmptyElement;
+		return EmptyElement;
+	}
+
+	template<typename T>
+	FORCEINLINE static T& GetArrayElementSafe(TArray<T>& InArray, const int32 InIndex)
+	{
+		if(InArray.IsValidIndex(InIndex))
+		{
+			return InArray[InIndex];
+		}
+		static T EmptyElement;
+		return EmptyElement;
+	}
+
+	template<typename A, typename B>
+	FORCEINLINE static void CopyUnrelatedArrays(TArray<A>& Target, const TArray<B>& Source)
+	{
+		const int32 Num = Source.Num();
+		Target.SetNum(Num);
+		for(int32 Index = 0; Index < Num; Index++)
+		{
+			Target[Index] = (A)Source[Index];
+		}
+	}
+
+	FORCEINLINE void BroadcastExecutionReachedExit()
+	{
+		ExecutionReachedExit().Broadcast(Context.PublicData.GetEventName());
+	}
+
+	template<typename T>
+	FORCEINLINE T& GetExternalVariableRef(const FName& InExternalVariableName, const FName& InExpectedType) const
+	{
+		for(const FRigVMExternalVariable& ExternalVariable : GetExternalVariables())
+		{
+			if(ExternalVariable.Name == InExternalVariableName && ExternalVariable.GetExtendedCPPType() == InExpectedType)
+			{
+				if(ExternalVariable.IsValid(false))
+				{
+					return *(T*)ExternalVariable.Memory;
+				}
+			}
+		}
+		static T EmptyValue;
+		return EmptyValue;
+	}
+
+	template<typename T>
+	FORCEINLINE T& GetOperandSlice(TArray<T>& InOutArray)
+	{
+		const int32 SliceIndex = Context.GetSlice().GetIndex();
+		if(InOutArray.Num() <= SliceIndex)
+		{
+			InOutArray.AddDefaulted(1 + SliceIndex - InOutArray.Num());
+		}
+		return InOutArray[SliceIndex];
+	}
+
+	template<typename T>
+	FORCEINLINE static bool ArrayOp_Iterator(const TArray<T>& InArray, T& OutElement, int32 InIndex, int32& OutCount, float& OutRatio)
+	{
+		OutCount = InArray.Num();
+		const bool bContinue = InIndex >=0 && InIndex < OutCount;
+
+		if((OutCount <= 0) || !bContinue)
+		{
+			OutRatio = 0.f;
+		}
+		else
+		{
+			OutRatio = float(InIndex) / float(OutCount - 1);
+			OutElement = InArray[InIndex];
+		}
+		
+		return bContinue;
+	}
+
+	FORCEINLINE static bool ArrayOp_Iterator(const TArray<double>& InArray, float& OutElement, int32 InIndex, int32& OutCount, float& OutRatio)
+	{
+		double Element = (double)OutElement;
+		const bool bResult = ArrayOp_Iterator<double>(InArray, Element, InIndex, OutCount, OutRatio);
+		OutElement = (float)Element;
+		return bResult;
+	}
+
+	FORCEINLINE static bool ArrayOp_Iterator(const TArray<float>& InArray, double& OutElement, int32 InIndex, int32& OutCount, float& OutRatio)
+	{
+		float Element = (float)OutElement;
+		const bool bResult = ArrayOp_Iterator<float>(InArray, Element, InIndex, OutCount, OutRatio);
+		OutElement = (double)Element;
+		return bResult;
+	}
+
+	template<typename T>
+	FORCEINLINE static bool ArrayOp_Find(const TArray<T>& InArray, const T& InElement, int32& OutIndex)
+	{
+		OutIndex = InArray.Find(InElement);
+		return OutIndex != INDEX_NONE;
+	}
+
+	FORCEINLINE static bool ArrayOp_Find(const TArray<float>& InArray, const double& InElement, int32& OutIndex)
+	{
+		return ArrayOp_Find<float>(InArray, (float)InElement, OutIndex);
+	}
+
+	FORCEINLINE static bool ArrayOp_Find(const TArray<double>& InArray, const float& InElement, int32& OutIndex)
+	{
+		return ArrayOp_Find<double>(InArray, (double)InElement, OutIndex);
+	}
+
+	template<typename T>
+	FORCEINLINE static void ArrayOp_Union(TArray<T>& InOutA, const TArray<T>& InB)
+	{
+		for(const T& Element : InB)
+		{
+			InOutA.AddUnique(Element);
+		}
+	}
+
+	template<typename T>
+	FORCEINLINE static void ArrayOp_Difference(TArray<T>& InOutA, const TArray<T>& InB)
+	{
+		TArray<T> Temp;
+		Swap(InOutA, Temp);
+
+		for(const T& Element : Temp)
+		{
+			if(!InB.Contains(Element))
+			{
+				InOutA.Add(Element);
+			}
+		}
+		for(const T& Element : InB)
+		{
+			if(!Temp.Contains(Element))
+			{
+				InOutA.Add(Element);
+			}
+		}
+	}
+
+	template<typename T>
+	FORCEINLINE static void ArrayOp_Intersection(TArray<T>& InOutA, const TArray<T>& InB)
+	{
+		TArray<T> Temp;
+		Swap(InOutA, Temp);
+
+		for(const T& Element : Temp)
+		{
+			if(InB.Contains(Element))
+			{
+				InOutA.Add(Element);
+			}
+		}
+	}
+
+	template<typename T>
+	FORCEINLINE static void ArrayOp_Reverse(TArray<T>& InOutArray)
+	{
+		Algo::Reverse(InOutArray);
+	}
+
+public:
+
+	template <
+		typename T,
+		typename TEnableIf<TRigVMIsBaseStructure<T>::Value, T>::Type* = nullptr
+	>
+	FORCEINLINE static T GetStructConstant(const FString& InDefaultValue)
+	{
+		T Value;
+		if(!InDefaultValue.IsEmpty())
+		{
+			FErrorPipe ErrorPipe;
+			TBaseStructure<T>::Get()->ImportText(*InDefaultValue, &Value, nullptr, PPF_None, &ErrorPipe, FString());
+		}
+		return Value;
+	}
+
+	template <
+		typename T,
+		typename TEnableIf<TModels<CRigVMUStruct, T>::Value>::Type* = nullptr
+	>
+	FORCEINLINE static T GetStructConstant(const FString& InDefaultValue)
+	{
+		T Value;
+		if(!InDefaultValue.IsEmpty())
+		{
+			FErrorPipe ErrorPipe;
+			T::StaticStruct()->ImportText(*InDefaultValue, &Value, nullptr, PPF_None, &ErrorPipe, FString());
+		}
+		return Value;
+	}
+
+	template <
+		typename T,
+		typename TEnableIf<TRigVMIsBaseStructure<T>::Value, T>::Type* = nullptr
+	>
+	FORCEINLINE static TArray<T> GetStructArrayConstant(const FString& InDefaultValue)
+	{
+		TArray<T> Value;
+		if(!InDefaultValue.IsEmpty())
+		{
+			FErrorPipe ErrorPipe;
+			TBaseStructure<T>::Get()->ImportText(*InDefaultValue, &Value, nullptr, PPF_None, &ErrorPipe, FString());
+		}
+		return Value;
+	}
+
+	template <
+		typename T,
+		typename TEnableIf<TModels<CRigVMUStruct, T>::Value>::Type* = nullptr
+	>
+	FORCEINLINE static TArray<T> GetStructArrayConstant(const FString& InDefaultValue)
+	{
+		TArray<T> Value;
+		if(!InDefaultValue.IsEmpty())
+		{
+			FErrorPipe ErrorPipe;
+			T::StaticStruct()->ImportText(*InDefaultValue, &Value, nullptr, PPF_None, &ErrorPipe, FString());
+		}
+		return Value;
+	}
+
+protected:
+	
+	class FTransformSetter
+	{
+	public:
+		FORCEINLINE FTransformSetter(FTransform& InOutTransform)
+			: Transform(InOutTransform)
+		{
+		}
+
+		FORCEINLINE void SetTranslationX(const FTransform::FReal InValue) const { SetTranslationComponent(0, InValue); }
+		FORCEINLINE void SetTranslationY(const FTransform::FReal InValue) const { SetTranslationComponent(1, InValue); }
+		FORCEINLINE void SetTranslationZ(const FTransform::FReal InValue) const { SetTranslationComponent(2, InValue); }
+		FORCEINLINE void SetTranslationComponent(const uint8 InComponent, const FTransform::FReal InValue) const
+		{
+			FVector Translation = Transform.GetTranslation();
+			Translation.Component(InComponent) = InValue;
+			Transform.SetTranslation(Translation);
+		}
+		
+		FORCEINLINE void SetScaleX(const FTransform::FReal InValue) const { SetScaleComponent(0, InValue); }
+		FORCEINLINE void SetScaleY(const FTransform::FReal InValue) const { SetScaleComponent(1, InValue); }
+		FORCEINLINE void SetScaleZ(const FTransform::FReal InValue) const { SetScaleComponent(2, InValue); }
+		FORCEINLINE void SetScaleComponent(const uint8 InComponent, const FTransform::FReal InValue) const
+		{
+			FVector Scale = Transform.GetScale3D();
+			Scale.Component(InComponent) = InValue;
+			Transform.SetScale3D(Scale);
+		}
+		
+		FORCEINLINE void SetRotationX(const FTransform::FReal InValue) const { SetRotationComponent(0, InValue); }
+		FORCEINLINE void SetRotationY(const FTransform::FReal InValue) const { SetRotationComponent(1, InValue); }
+		FORCEINLINE void SetRotationZ(const FTransform::FReal InValue) const { SetRotationComponent(2, InValue); }
+		FORCEINLINE void SetRotationW(const FTransform::FReal InValue) const { SetRotationComponent(3, InValue); }
+		FORCEINLINE void SetRotationComponent(const uint8 InComponent, const FTransform::FReal InValue) const
+		{
+			FQuat Rotation = Transform.GetRotation();
+			switch(InComponent)
+			{
+				case 0:
+				{
+					Rotation.X = InValue;
+					break;
+				}
+				case 1:
+				{
+					Rotation.Y = InValue;
+					break;
+				}
+				case 2:
+				{
+					Rotation.Z = InValue;
+					break;
+				}
+				case 3:
+				default:
+				{
+					Rotation.W = InValue;
+					break;
+				}
+			}
+			Rotation.Normalize();
+			Transform.SetRotation(Rotation);
+		}
+
+	private:
+		FTransform& Transform;
+	};
+
+	int32 TemporaryArrayIndex;
+};
diff --git a/Engine/Source/Runtime/RigVM/Public/RigVMCore/RigVMPropertyPath.h b/Engine/Source/Runtime/RigVM/Public/RigVMCore/RigVMPropertyPath.h
--- a/Engine/Source/Runtime/RigVM/Public/RigVMCore/RigVMPropertyPath.h
+++ b/Engine/Source/Runtime/RigVM/Public/RigVMCore/RigVMPropertyPath.h
@@ -40,6 +40,9 @@ struct RIGVM_API FRigVMPropertyPathDescription
 		, SegmentPath(InSegmentPath)
 	{}
 
+	// returns true if this property path is valid
+	bool IsValid() const { return PropertyIndex != INDEX_NONE && !HeadCPPType.IsEmpty() && !SegmentPath.IsEmpty(); }
+
 	// Archive operator for serialization
 	FORCEINLINE friend FArchive& operator<<(FArchive& Ar, FRigVMPropertyPathDescription& Path)
 	{
diff --git a/Engine/Source/Runtime/RigVM/Public/RigVMModule.h b/Engine/Source/Runtime/RigVM/Public/RigVMModule.h
--- a/Engine/Source/Runtime/RigVM/Public/RigVMModule.h
+++ b/Engine/Source/Runtime/RigVM/Public/RigVMModule.h
@@ -9,7 +9,7 @@
 #include "CoreMinimal.h"
 #include "RigVMDefines.h"
 
-DECLARE_LOG_CATEGORY_EXTERN(LogRigVM, Log, All);
+RIGVM_API DECLARE_LOG_CATEGORY_EXTERN(LogRigVM, Log, All);
 
 namespace RigVMCore
 {
diff --git a/Engine/Source/Runtime/RigVM/Public/RigVMTypeUtils.h b/Engine/Source/Runtime/RigVM/Public/RigVMTypeUtils.h
--- a/Engine/Source/Runtime/RigVM/Public/RigVMTypeUtils.h
+++ b/Engine/Source/Runtime/RigVM/Public/RigVMTypeUtils.h
@@ -23,6 +23,7 @@ namespace RigVMTypeUtils
 	const FString FloatArrayType = TEXT("TArray<float>");
 	const FString DoubleArrayType = TEXT("TArray<double>");
 	const FString Int32ArrayType = TEXT("TArray<int32>");
+	const FString UInt8ArrayType = TEXT("TArray<uint8>");
 	const FString FNameArrayType = TEXT("TArray<FName>");
 	const FString FStringArrayType = TEXT("TArray<FString>");
 
@@ -37,6 +38,7 @@ namespace RigVMTypeUtils
 	const FName FloatArrayTypeName = *FloatArrayType;
 	const FName DoubleArrayTypeName = *DoubleArrayType;
 	const FName Int32ArrayTypeName = *Int32ArrayType;
+	const FName UInt8ArrayTypeName = *UInt8ArrayType;
 	const FName FNameArrayTypeName = *FNameArrayType;
 	const FName FStringArrayTypeName = *FStringArrayType;
 
@@ -53,7 +55,7 @@ namespace RigVMTypeUtils
 
 	FORCEINLINE FString BaseTypeFromArrayType(const FString& InCPPType)
 	{
-		return InCPPType.RightChop(7).LeftChop(1);
+		return InCPPType.RightChop(7).LeftChop(1).TrimStartAndEnd();
 	}
 
 	FORCEINLINE FString CPPTypeFromEnum(UEnum* InEnum)
