Control Rig: Do not allow functions with wildcards

diff --git a/Engine/Plugins/Animation/ControlRig/Source/ControlRigEditor/Private/ControlRigEditorModule.cpp b/Engine/Plugins/Animation/ControlRig/Source/ControlRigEditor/Private/ControlRigEditorModule.cpp
--- a/Engine/Plugins/Animation/ControlRig/Source/ControlRigEditor/Private/ControlRigEditorModule.cpp
+++ b/Engine/Plugins/Animation/ControlRig/Source/ControlRigEditor/Private/ControlRigEditorModule.cpp
@@ -2109,8 +2109,14 @@ void FControlRigEditorModule::GetContextMenuActions(const UControlRigGraphSchema
 							URigVMCollapseNode* CollapseNode = Controller->CollapseNodes(Nodes, TEXT("New Function"), true, true);
 							if(CollapseNode)
 							{
-								Controller->PromoteCollapseNodeToFunctionReferenceNode(CollapseNode->GetFName(), true, true);
-								Controller->CloseUndoBracket();
+								if (Controller->PromoteCollapseNodeToFunctionReferenceNode(CollapseNode->GetFName(), true, true).IsNone())
+								{
+									Controller->CancelUndoBracket();
+								}
+								else
+								{
+									Controller->CloseUndoBracket();
+								}
 							}
 							else
 							{
diff --git a/Engine/Source/Developer/RigVMDeveloper/Private/RigVMModel/RigVMController.cpp b/Engine/Source/Developer/RigVMDeveloper/Private/RigVMModel/RigVMController.cpp
--- a/Engine/Source/Developer/RigVMDeveloper/Private/RigVMModel/RigVMController.cpp
+++ b/Engine/Source/Developer/RigVMDeveloper/Private/RigVMModel/RigVMController.cpp
@@ -1557,7 +1557,7 @@ URigVMTemplateNode* URigVMController::ReplaceUnitNodeWithTemplateNode(const FNam
 	{
 		if(bSetupUndoRedo)
 		{
-			ActionStack->CancelAction(Action);
+			ActionStack->CancelAction(Action, this);
 		}
 		return nullptr;
 	}
@@ -1679,7 +1679,7 @@ bool URigVMController::UnresolveTemplateNodes(const TArray<URigVMTemplateNode*>&
 				ReportErrorf(TEXT("Cannot find expected pin '%s' on Template Node '%s'"), *PinName.ToString(), *Node->GetNodePath());
 				if(bSetupUndoRedo)
 				{
-					ActionStack->CancelAction(Action);
+					ActionStack->CancelAction(Action, this);
 				}
 				return false;
 			}
@@ -1728,7 +1728,7 @@ bool URigVMController::UnresolveTemplateNodes(const TArray<URigVMTemplateNode*>&
 		}
 		else
 		{
-			ActionStack->CancelAction(Action);
+			ActionStack->CancelAction(Action, this);
 			return false;
 		}
 	}
@@ -1835,7 +1835,7 @@ TArray<URigVMNode*> URigVMController::UpgradeNodes(const TArray<URigVMNode*>& In
 		{
 			if(bSetupUndoRedo)
 			{
-				ActionStack->CancelAction(Action);
+				ActionStack->CancelAction(Action, this);
 			}
 			ReportErrorf(TEXT("Couldn't remove link '%s' -> '%s'"), *LinkedPath.Key, *LinkedPath.Value);
 			return TArray<URigVMNode*>();
@@ -1945,7 +1945,7 @@ URigVMNode* URigVMController::UpgradeNode(URigVMNode* InNode, bool bSetupUndoRed
 		{
 			if(bSetupUndoRedo)
 			{
-				ActionStack->CancelAction(Action);
+				ActionStack->CancelAction(Action, this);
 			}
 			ReportErrorf(TEXT("Unable to remove node %s."), *NodeName);
 			return nullptr;
@@ -1956,7 +1956,7 @@ URigVMNode* URigVMController::UpgradeNode(URigVMNode* InNode, bool bSetupUndoRed
 		{
 			if(bSetupUndoRedo)
 			{
-				ActionStack->CancelAction(Action);
+				ActionStack->CancelAction(Action, this);
 			}
 			ReportErrorf(TEXT("Unable to upgrade node %s."), *NodeName);
 			return nullptr;
@@ -2128,7 +2128,7 @@ URigVMRerouteNode* URigVMController::AddRerouteNodeOnLink(URigVMLink* InLink, bo
 	{
 		if (bSetupUndoRedo)
 		{
-			ActionStack->CancelAction(Action);
+			ActionStack->CancelAction(Action, this);
 		}
 		return nullptr;
 	}
@@ -2397,7 +2397,7 @@ URigVMInjectionInfo* URigVMController::AddInjectedNode(const FString& InPinPath,
 		{
 			if (bSetupUndoRedo)
 			{
-				ActionStack->CancelAction(Action);
+				ActionStack->CancelAction(Action, this);
 			}
 			return nullptr;
 		}
@@ -2407,7 +2407,7 @@ URigVMInjectionInfo* URigVMController::AddInjectedNode(const FString& InPinPath,
 			RemoveNode(UnitNode, false);
 			if (bSetupUndoRedo)
 			{
-				ActionStack->CancelAction(Action);
+				ActionStack->CancelAction(Action, this);
 			}
 			return nullptr;
 		}
@@ -2424,7 +2424,7 @@ URigVMInjectionInfo* URigVMController::AddInjectedNode(const FString& InPinPath,
 			RemoveNode(UnitNode, false);
 			if (bSetupUndoRedo)
 			{
-				ActionStack->CancelAction(Action);
+				ActionStack->CancelAction(Action, this);
 			}
 			return nullptr;
 		}
@@ -2436,7 +2436,7 @@ URigVMInjectionInfo* URigVMController::AddInjectedNode(const FString& InPinPath,
 			RemoveNode(UnitNode, false);
 			if (bSetupUndoRedo)
 			{
-				ActionStack->CancelAction(Action);
+				ActionStack->CancelAction(Action, this);
 			}
 			return nullptr;
 		}
@@ -2587,7 +2587,7 @@ bool URigVMController::RemoveInjectedNode(const FString& InPinPath, bool bAsInpu
 	URigVMNode* NodeEjected = EjectNodeFromPin(InPinPath, bSetupUndoRedo);
 	if (!NodeEjected)
 	{
-		ActionStack->CancelAction(Action);
+		ActionStack->CancelAction(Action, this);
 		return false;
 	}
 
@@ -2611,7 +2611,7 @@ bool URigVMController::RemoveInjectedNode(const FString& InPinPath, bool bAsInpu
 	// 3.- Remove node
 	if (!RemoveNode(NodeEjected))
 	{
-		ActionStack->CancelAction(Action);
+		ActionStack->CancelAction(Action, this);
 		return false;
 	}
 
@@ -2961,7 +2961,7 @@ bool URigVMController::CloseUndoBracket()
 	{
 		return false;
 	}
-	return ActionStack->CloseUndoBracket();
+	return ActionStack->CloseUndoBracket(this);
 }
 
 bool URigVMController::CancelUndoBracket()
@@ -2970,7 +2970,7 @@ bool URigVMController::CancelUndoBracket()
 	{
 		return false;
 	}
-	return ActionStack->CancelUndoBracket();
+	return ActionStack->CancelUndoBracket(this);
 }
 
 FString URigVMController::ExportNodesToText(const TArray<FName>& InNodeNames)
@@ -4486,7 +4486,7 @@ TArray<URigVMNode*> URigVMController::ExpandLibraryNode(URigVMLibraryNode* InNod
 	{
 		if (bSetupUndoRedo)
 		{
-			ActionStack->CancelAction(ExpandAction);
+			ActionStack->CancelAction(ExpandAction, this);
 		}
 		return TArray<URigVMNode*>();
 	}
@@ -4538,7 +4538,7 @@ TArray<URigVMNode*> URigVMController::ExpandLibraryNode(URigVMLibraryNode* InNod
 							ReportErrorf(TEXT("Found variable %s of incompatible type with a local variable inside function %s"), *LocalVariable.Name.ToString(), *FunctionReferenceNode->GetReferencedNode()->GetName());
 							if (bSetupUndoRedo)
 							{
-								ActionStack->CancelAction(ExpandAction);
+								ActionStack->CancelAction(ExpandAction, this);
 							}
 							return TArray<URigVMNode*>();
 						}
@@ -5194,6 +5194,15 @@ URigVMFunctionReferenceNode* URigVMController::PromoteCollapseNodeToFunctionRefe
 		return nullptr;
 	}
 
+	for (URigVMPin* Pin : InCollapseNode->GetPins())
+	{
+		if (Pin->IsWildCard())
+		{
+			ReportErrorf(TEXT("Cannot create function %s because it contains a wildcard pin %s"), *InCollapseNode->GetName(), *Pin->GetName());
+			return nullptr;
+		}
+	}
+
 	FRigVMControllerCompileBracketScope CompileScope(this);
 	URigVMFunctionReferenceNode* FunctionRefNode = nullptr;
 
@@ -5627,7 +5636,7 @@ bool URigVMController::RemoveNode(URigVMNode* InNode, bool bSetupUndoRedo, bool
 
 		if (!EjectNodeFromPin(Pin->GetPinPath(), bSetupUndoRedo))
 		{
-			ActionStack->CancelAction(Action);
+			ActionStack->CancelAction(Action, this);
 			return false;
 		}
 		
@@ -5861,7 +5870,7 @@ bool URigVMController::RenameNode(URigVMNode* InNode, const FName& InNewName, bo
 	InNode->PreviousName = InNode->GetFName();
 	if (!RenameObject(InNode, *ValidNewName.ToString()))
 	{
-		ActionStack->CancelAction(Action);
+		ActionStack->CancelAction(Action, this);
 		return false;
 	}
 
@@ -6043,7 +6052,7 @@ bool URigVMController::SetNodeSelection(const TArray<FName>& InNodeNames, bool b
 		}
 		else
 		{
-			ActionStack->CancelAction(Action);
+			ActionStack->CancelAction(Action, this);
 		}
 	}
 
@@ -6705,7 +6714,7 @@ bool URigVMController::RenameVariable(const FName& InOldName, const FName& InNew
 		}
 		else
 		{
-			ActionStack->CancelAction(Action);
+			ActionStack->CancelAction(Action, this);
 		}
 	}
 
@@ -7378,7 +7387,7 @@ FString URigVMController::AddAggregatePin(URigVMNode* InNode, const FString& InP
 		{
 			if (bSetupUndoRedo)
 			{
-				ActionStack->CancelAction(Action);
+				ActionStack->CancelAction(Action, this);
 			}
 			return FString();
 		}
@@ -7410,7 +7419,7 @@ FString URigVMController::AddAggregatePin(URigVMNode* InNode, const FString& InP
 			{
 				if (bSetupUndoRedo)
 				{
-					ActionStack->CancelAction(Action);
+					ActionStack->CancelAction(Action, this);
 				}
 				return FString();	
 			}
@@ -7424,7 +7433,7 @@ FString URigVMController::AddAggregatePin(URigVMNode* InNode, const FString& InP
 		{
 			if (bSetupUndoRedo)
 			{
-				ActionStack->CancelAction(Action);
+				ActionStack->CancelAction(Action, this);
 			}
 			return FString();
 		}
@@ -7434,7 +7443,7 @@ FString URigVMController::AddAggregatePin(URigVMNode* InNode, const FString& InP
 		{
 			if (bSetupUndoRedo)
 			{
-				ActionStack->CancelAction(Action);
+				ActionStack->CancelAction(Action, this);
 			}
 			return FString();
 		}
@@ -7449,7 +7458,7 @@ FString URigVMController::AddAggregatePin(URigVMNode* InNode, const FString& InP
 		{
 			if (bSetupUndoRedo)
 			{
-				ActionStack->CancelAction(Action);
+				ActionStack->CancelAction(Action, this);
 			}
 			return FString();
 		}
@@ -7494,7 +7503,7 @@ FString URigVMController::AddAggregatePin(URigVMNode* InNode, const FString& InP
 		{
 			if (bSetupUndoRedo)
 			{
-				ActionStack->CancelAction(Action);
+				ActionStack->CancelAction(Action, this);
 			}
 			return FString();
 		}
@@ -7504,7 +7513,7 @@ FString URigVMController::AddAggregatePin(URigVMNode* InNode, const FString& InP
 	{
 		if (bSetupUndoRedo)
 		{
-			ActionStack->CancelAction(Action);
+			ActionStack->CancelAction(Action, this);
 		}
 		return FString();
 	}
@@ -7530,7 +7539,7 @@ FString URigVMController::AddAggregatePin(URigVMNode* InNode, const FString& InP
 		{
 			if (bSetupUndoRedo)
 			{
-				ActionStack->CancelAction(Action);
+				ActionStack->CancelAction(Action, this);
 			}
 			return FString();
 		}
@@ -7627,7 +7636,7 @@ FString URigVMController::AddAggregatePin(URigVMNode* InNode, const FString& InP
 	{
 		if (bSetupUndoRedo)
 		{
-			ActionStack->CancelAction(Action);
+			ActionStack->CancelAction(Action, this);
 		}
 		return FString();
 	}
@@ -7797,7 +7806,7 @@ bool URigVMController::RemoveAggregatePin(URigVMPin* InPin, bool bSetupUndoRedo,
 		}
 		else
 		{
-			ActionStack->CancelAction(Action);			
+			ActionStack->CancelAction(Action, this);			
 		}
 	}
 
@@ -8172,7 +8181,7 @@ bool URigVMController::SetArrayPinSize(const FString& InArrayPinPath, int32 InSi
 		{
 			if (bSetupUndoRedo)
 			{
-				ActionStack->CancelAction(Action);
+				ActionStack->CancelAction(Action, this);
 			}
 			return false;
 		}
@@ -8185,7 +8194,7 @@ bool URigVMController::SetArrayPinSize(const FString& InArrayPinPath, int32 InSi
 		{
 			if (bSetupUndoRedo)
 			{
-				ActionStack->CancelAction(Action);
+				ActionStack->CancelAction(Action, this);
 			}
 			return false;
 		}
@@ -8200,7 +8209,7 @@ bool URigVMController::SetArrayPinSize(const FString& InArrayPinPath, int32 InSi
 		}
 		else
 		{
-			ActionStack->CancelAction(Action);
+			ActionStack->CancelAction(Action, this);
 		}
 	}
 
@@ -8343,7 +8352,7 @@ bool URigVMController::BindPinToVariable(URigVMPin* InPin, const FString& InNewB
 		{
 			if (bSetupUndoRedo)
 			{
-				ActionStack->CancelAction(Action);
+				ActionStack->CancelAction(Action, this);
 			}
 			return false;
 		}
@@ -8363,7 +8372,7 @@ bool URigVMController::BindPinToVariable(URigVMPin* InPin, const FString& InNewB
 			{
 				if (bSetupUndoRedo)
 				{
-					ActionStack->CancelAction(Action);
+					ActionStack->CancelAction(Action, this);
 				}
 				return false;
 			}
@@ -8375,7 +8384,7 @@ bool URigVMController::BindPinToVariable(URigVMPin* InPin, const FString& InNewB
 	{
 		if (bSetupUndoRedo)
 		{
-			ActionStack->CancelAction(Action);
+			ActionStack->CancelAction(Action, this);
 		}
 		return false;
 	}
@@ -9184,7 +9193,7 @@ bool URigVMController::BreakAllLinks(URigVMPin* Pin, bool bAsInput, bool bSetupU
 		}
 		else
 		{
-			ActionStack->CancelAction(Action);
+			ActionStack->CancelAction(Action, this);
 		}
 	}
 
@@ -9496,7 +9505,7 @@ bool URigVMController::RemoveExposedPin(const FName& InPinName, bool bSetupUndoR
 		}
 		else
 		{
-			ActionStack->CancelAction(Action);
+			ActionStack->CancelAction(Action, this);
 		}
 	}
 
@@ -9628,7 +9637,7 @@ bool URigVMController::RenameExposedPin(const FName& InOldPinName, const FName&
 
 	if (!Local::RenamePin(this, Pin, PinName))
 	{
-		ActionStack->CancelAction(Action);
+		ActionStack->CancelAction(Action, this);
 		return false;
 	}
 
@@ -9777,7 +9786,7 @@ bool URigVMController::ChangeExposedPinType(const FName& InPinName, const FStrin
 		{
 			if (bSetupUndoRedo)
 			{
-				ActionStack->CancelAction(Action);
+				ActionStack->CancelAction(Action, this);
 			}
 			return false;
 		}
@@ -14513,7 +14522,7 @@ bool URigVMController::ResolveWildCardPin(const FString& InPinPath, const FStrin
 
 			if (bSetupUndoRedo)
 			{
-				ActionStack->CancelAction(Action);
+				ActionStack->CancelAction(Action, this);
 			}
 		}
 	}
diff --git a/Engine/Source/Developer/RigVMDeveloper/Private/RigVMModel/RigVMControllerActions.cpp b/Engine/Source/Developer/RigVMDeveloper/Private/RigVMModel/RigVMControllerActions.cpp
--- a/Engine/Source/Developer/RigVMDeveloper/Private/RigVMModel/RigVMControllerActions.cpp
+++ b/Engine/Source/Developer/RigVMDeveloper/Private/RigVMModel/RigVMControllerActions.cpp
@@ -59,12 +59,12 @@ bool URigVMActionStack::OpenUndoBracket(const FString& InTitle)
 	return true;
 }
 
-bool URigVMActionStack::CloseUndoBracket()
+bool URigVMActionStack::CloseUndoBracket(URigVMController* InController)
 {
 	ensure(BracketActions.Num() > 0);
 	if(BracketActions.Last()->IsEmpty())
 	{
-		return CancelUndoBracket();
+		return CancelUndoBracket(InController);
 	}
 	FRigVMBaseAction* Action = BracketActions.Pop();
 	EndAction(*Action);
@@ -72,11 +72,11 @@ bool URigVMActionStack::CloseUndoBracket()
 	return true;
 }
 
-bool URigVMActionStack::CancelUndoBracket()
+bool URigVMActionStack::CancelUndoBracket(URigVMController* InController)
 {
 	ensure(BracketActions.Num() > 0);
 	FRigVMBaseAction* Action = BracketActions.Pop();
-	CancelAction(*Action);
+	CancelAction(*Action, InController);
 	delete(Action);
 	return true;
 }
diff --git a/Engine/Source/Developer/RigVMDeveloper/Public/RigVMModel/RigVMControllerActions.h b/Engine/Source/Developer/RigVMDeveloper/Public/RigVMModel/RigVMControllerActions.h
--- a/Engine/Source/Developer/RigVMDeveloper/Public/RigVMModel/RigVMControllerActions.h
+++ b/Engine/Source/Developer/RigVMDeveloper/Public/RigVMModel/RigVMControllerActions.h
@@ -221,12 +221,14 @@ public:
 	// Cancels an action, closes a bracket / scope and discards all 
 	// actions to this point.
 	template<class ActionType>
-	void CancelAction(ActionType& InAction)
+	void CancelAction(ActionType& InAction, URigVMController* InController)
 	{
 		ensure(CurrentActions.Num() > 0);
 		ensure((FRigVMBaseAction*)&InAction == CurrentActions.Last());
 		CurrentActions.Pop();
 
+		InAction.Undo(InController);
+
 #if RIGVM_ACTIONSTACK_VERBOSE_LOG		
 		TGuardValue<int32> TabDepthGuard(LogActionDepth, CurrentActions.Num());
 		LogAction<ActionType>(InAction, FRigVMBaseAction::CancelActionPrefix);
@@ -301,12 +303,12 @@ public:
 	// Closes an undo bracket / scope.
 	// This is primary useful for Python.
 	UFUNCTION()
-	bool CloseUndoBracket();
+	bool CloseUndoBracket(URigVMController* InController);
 
 	// Cancels an undo bracket / scope.
 	// This is primary useful for Python.
 	UFUNCTION()
-	bool CancelUndoBracket();
+	bool CancelUndoBracket(URigVMController* InController);
 
 	// Pops the last action from the undo stack and perform undo on it.
 	// Note: This should really only be used for unit tests,
