Back out changelist 19573926

diff --git a/Engine/Plugins/Animation/ControlRig/Source/ControlRigEditor/Private/Editor/SControlRigStackView.cpp b/Engine/Plugins/Animation/ControlRig/Source/ControlRigEditor/Private/Editor/SControlRigStackView.cpp
--- a/Engine/Plugins/Animation/ControlRig/Source/ControlRigEditor/Private/Editor/SControlRigStackView.cpp
+++ b/Engine/Plugins/Animation/ControlRig/Source/ControlRigEditor/Private/Editor/SControlRigStackView.cpp
@@ -87,75 +87,60 @@ void SRigStackItem::Construct(const FArguments& InArgs, const TSharedRef<STableV
 		[
 			SNew(SHorizontalBox)
 			+ SHorizontalBox::Slot()
-			.AutoWidth()
+			.MaxWidth(25.f)
+			.VAlign(VAlign_Center)
+			.HAlign(HAlign_Right)
+			[
+				SAssignNew(NumberWidget, STextBlock)
+				.Text(this, &SRigStackItem::GetIndexText)
+				.Font(IDetailLayoutBuilder::GetDetailFont())
+			]
+			+ SHorizontalBox::Slot()
+			.MaxWidth(22.f)
+			.VAlign(VAlign_Center)
+			.HAlign(HAlign_Center)
+			[
+				SNew(SImage)
+				.Image(Icon)
+			]
+			
+			+ SHorizontalBox::Slot()
+			.FillWidth(1.f)
+			.MaxWidth(200.f)
 			.VAlign(VAlign_Center)
 			.HAlign(HAlign_Left)
 			[
-				SNew(SHorizontalBox)
-				+ SHorizontalBox::Slot()
-				.MaxWidth(25.f)
-				.VAlign(VAlign_Center)
-				.HAlign(HAlign_Right)
-				[
-					SAssignNew(NumberWidget, STextBlock)
-					.Text(this, &SRigStackItem::GetIndexText)
-					.Font(IDetailLayoutBuilder::GetDetailFont())
-				]
-				+ SHorizontalBox::Slot()
-				.MaxWidth(22.f)
-				.VAlign(VAlign_Center)
-				.HAlign(HAlign_Center)
-				[
-					SNew(SImage)
-					.Image(Icon)
-				]
-				
-				+ SHorizontalBox::Slot()
-				.FillWidth(1.f)
-				.AutoWidth()
-				.VAlign(VAlign_Center)
-				.HAlign(HAlign_Fill)
-				[
-					SAssignNew(TextWidget, STextBlock)
-					.Text(this, &SRigStackItem::GetLabelText)
-					.Font(this, &SRigStackItem::GetLabelFont)
-					.Justification(ETextJustify::Left)
-				]
+				SAssignNew(TextWidget, STextBlock)
+				.Text(this, &SRigStackItem::GetLabelText)
+				.Font(this, &SRigStackItem::GetLabelFont)
 			]
 
 			+ SHorizontalBox::Slot()
 			.FillWidth(1.f)
+			.MaxWidth(40.f)
 			[
 				SNew(SSpacer)
 			]
 
 			+ SHorizontalBox::Slot()
-			.AutoWidth()
-			.VAlign(VAlign_Center)
-			.HAlign(HAlign_Right)
-			[
-				SNew(SHorizontalBox)
-				+ SHorizontalBox::Slot()
-				.AutoWidth()
-				.VAlign(VAlign_Center)
-	            .HAlign(HAlign_Left)
-	            [
-	                SNew(STextBlock)
-	                .Text(this, &SRigStackItem::GetVisitedCountText)
-	                .Font(IDetailLayoutBuilder::GetDetailFont())
-	            ]
-	            
-	            + SHorizontalBox::Slot()
-	            .Padding(20, 0, 0, 0)
-				.AutoWidth()
-	            .VAlign(VAlign_Center)
-	            .HAlign(HAlign_Left)
-	            [
-	                SNew(STextBlock)
-	                .Text(this, &SRigStackItem::GetDurationText)
-	                .Font(IDetailLayoutBuilder::GetDetailFont())
-	            ]
-	        ]
+			.MaxWidth(36.f)
+            .VAlign(VAlign_Center)
+            .HAlign(HAlign_Right)
+            [
+                SNew(STextBlock)
+                .Text(this, &SRigStackItem::GetVisitedCountText)
+                .Font(IDetailLayoutBuilder::GetDetailFont())
+            ]
+
+			+ SHorizontalBox::Slot()
+			.MaxWidth(80.f)
+            .VAlign(VAlign_Center)
+            .HAlign(HAlign_Right)
+            [
+                SNew(STextBlock)
+                .Text(this, &SRigStackItem::GetDurationText)
+                .Font(IDetailLayoutBuilder::GetDetailFont())
+            ]
         ], OwnerTable);
 }
 
diff --git a/Engine/Source/Developer/RigVMDeveloper/Private/RigVMCompiler/RigVMCompiler.cpp b/Engine/Source/Developer/RigVMDeveloper/Private/RigVMCompiler/RigVMCompiler.cpp
--- a/Engine/Source/Developer/RigVMDeveloper/Private/RigVMCompiler/RigVMCompiler.cpp
+++ b/Engine/Source/Developer/RigVMDeveloper/Private/RigVMCompiler/RigVMCompiler.cpp
@@ -3,7 +3,6 @@
 #include "RigVMCompiler/RigVMCompiler.h"
 #include "RigVMModel/RigVMController.h"
 #include "RigVMCore/RigVMExecuteContext.h"
-#include "RigVMCore/RigVMNativized.h"
 #include "RigVMDeveloperModule.h"
 #include "UObject/PropertyPortFlags.h"
 #include "UObject/Interface.h"
@@ -60,6 +59,8 @@ FRigVMCompileSettings::FRigVMCompileSettings()
 {
 }
 
+#if !UE_RIGVM_UCLASS_BASED_STORAGE_DISABLED
+
 FRigVMOperand FRigVMCompilerWorkData::AddProperty(
 	ERigVMMemoryType InMemoryType,
 	const FName& InName,
@@ -125,6 +126,8 @@ int32 FRigVMCompilerWorkData::FindOrAddPropertyPath(const FRigVMOperand& InOpera
 	return Descriptions.Add(FRigVMPropertyPathDescription(InOperand.GetRegisterIndex(), InHeadCPPType, InSegmentPath));
 }
 
+#endif
+
 URigVMCompiler::URigVMCompiler()
 {
 }
@@ -142,12 +145,6 @@ bool URigVMCompiler::Compile(URigVMGraph* InGraph, URigVMController* InControlle
 		return false;
 	}
 
-	if(OutVM->GetClass()->IsChildOf(URigVMNativized::StaticClass()))
-	{
-		ReportError(TEXT("Provided vm is nativized."));
-		return false;
-	}
-
 	DECLARE_SCOPE_HIERARCHICAL_COUNTER_FUNC()
 
 	TArray<FRigVMUserDataArray> UserData = InRigVMUserData;
@@ -322,30 +319,31 @@ bool URigVMCompiler::Compile(URigVMGraph* InGraph, URigVMController* InControlle
 	
 	for (const FRigVMExternalVariable& ExternalVariable : InExternalVariables)
 	{
+#if !UE_RIGVM_UCLASS_BASED_STORAGE_DISABLED
 		check(ExternalVariable.Property);
+#endif
 		FRigVMOperand Operand = OutVM->AddExternalVariable(ExternalVariable);
 		FString Hash = FString::Printf(TEXT("Variable::%s"), *ExternalVariable.Name.ToString());
 		OutOperands->Add(Hash, Operand);
 	}
 
-	FRigVMCompilerWorkData WorkData;
-
-	WorkData.AST = InAST;
-	if (!WorkData.AST.IsValid())
+	TSharedPtr<FRigVMParserAST> AST = InAST;
+	if (!AST.IsValid())
 	{
-		WorkData.AST = MakeShareable(new FRigVMParserAST(InGraph, InController, Settings.ASTSettings, InExternalVariables, UserData));
-		InGraph->RuntimeAST = WorkData.AST;
+		AST = MakeShareable(new FRigVMParserAST(InGraph, InController, Settings.ASTSettings, InExternalVariables, UserData));
+		InGraph->RuntimeAST = AST;
 #if UE_BUILD_DEBUG
 		//UE_LOG(LogRigVMDeveloper, Display, TEXT("%s"), *AST->DumpDot());
 #endif
 	}
-	ensure(WorkData.AST.IsValid());
+	ensure(AST.IsValid());
 
+	FRigVMCompilerWorkData WorkData;
 	WorkData.VM = OutVM;
 	WorkData.PinPathToOperand = OutOperands;
 	WorkData.RigVMUserData = UserData[0];
 	WorkData.bSetupMemory = true;
-	WorkData.ProxySources = &WorkData.AST->SharedOperandPins;
+	WorkData.ProxySources = &AST->SharedOperandPins;
 
 	// tbd: do we need this only when we have no pins?
 	//if(!WorkData.WatchedPins.IsEmpty())
@@ -402,12 +400,16 @@ bool URigVMCompiler::Compile(URigVMGraph* InGraph, URigVMController* InControlle
 			auto AddDefaultValueOperand = [&](URigVMPin* Pin)
 			{
 				FRigVMASTProxy PinProxy = FRigVMASTProxy::MakeFromUObject(Pin);
-				FRigVMVarExprAST* TempVarExpr = WorkData.AST->MakeExpr<FRigVMVarExprAST>(FRigVMExprAST::EType::Literal, PinProxy);
+				FRigVMVarExprAST* TempVarExpr = AST->MakeExpr<FRigVMVarExprAST>(FRigVMExprAST::EType::Literal, PinProxy);
 				FRigVMOperand Operand = FindOrAddRegister(TempVarExpr, WorkData, false);
 
+#if UE_RIGVM_UCLASS_BASED_STORAGE_DISABLED
+				WorkData.VM->GetLiteralMemory().SetRegisterValueFromString(Operand, LocalVariable.CPPType, LocalVariable.CPPTypeObject, {LocalVariable.DefaultValue});
+#else
 				check(Operand.GetMemoryType() == ERigVMMemoryType::Literal);
 				TArray<FRigVMPropertyDescription>& LiteralProperties = WorkData.PropertyDescriptions.FindChecked(Operand.GetMemoryType());
 				LiteralProperties[Operand.GetRegisterIndex()].DefaultValue = LocalVariable.DefaultValue;
+#endif
 			};
 			
 			// To create the default value in the literal memory, we need to find a pin in a variable node (or bounded to a local variable) that
@@ -430,6 +432,7 @@ bool URigVMCompiler::Compile(URigVMGraph* InGraph, URigVMController* InControlle
 		}
 	}	
 
+#if !UE_RIGVM_UCLASS_BASED_STORAGE_DISABLED
 	if(Settings.EnablePinWatches)
 	{
 		for(int32 GraphIndex=0; GraphIndex<VisitedGraphs.Num(); GraphIndex++)
@@ -447,13 +450,61 @@ bool URigVMCompiler::Compile(URigVMGraph* InGraph, URigVMController* InControlle
 			}
 		}
 	}
+#endif
+
+	// define all parameters independent from sorted nodes
+	for (const FRigVMExprAST* Expr : AST->Expressions)
+	{
+		if (Expr->IsA(FRigVMExprAST::EType::Literal))
+		{
+			continue;
+		}
+		if (Expr->IsA(FRigVMExprAST::EType::Var))
+		{
+			const FRigVMVarExprAST* VarExpr = Expr->To<FRigVMVarExprAST>();
+			{
+				if(Cast<URigVMParameterNode>(VarExpr->GetPin()->GetNode()))
+				{
+					if (VarExpr->GetPin()->GetName() == TEXT("Value"))
+					{
+						FindOrAddRegister(VarExpr, WorkData, false /* watchvalue */);
+					}
+				}
+			}
+		}
+	}
 
 	WorkData.ExprComplete.Reset();
-	for (FRigVMExprAST* RootExpr : *WorkData.AST)
+	for (FRigVMExprAST* RootExpr : *AST)
 	{
 		TraverseExpression(RootExpr, WorkData);
 	}
 
+	// If a parameter node has no corresponding AST node, because it's not on the execution path,
+	// then add a dummy parameter with no register index, so that we can properly add the pins on the 
+	// ControlRig node later and not lose any existing external connections.
+	for (URigVMNode* Node : InGraph->GetNodes())
+	{
+		if (URigVMParameterNode* ParameterNode = Cast<URigVMParameterNode>(Node))
+		{
+			FRigVMASTProxy ParameterNodeProxy = FRigVMASTProxy::MakeFromUObject(ParameterNode);
+			const FRigVMExprAST* ParameterExpr = AST->GetExprForSubject(ParameterNodeProxy);
+			if (!ParameterExpr || ParameterExpr->IsA(FRigVMExprAST::NoOp))
+			{
+				FName Name = ParameterNode->GetParameterName();
+
+				if (!WorkData.VM->ParametersNameMap.Contains(Name))
+				{
+					ERigVMParameterType ParameterType = ParameterNode->IsInput() ? ERigVMParameterType::Input : ERigVMParameterType::Output;
+					FRigVMParameter Parameter(ParameterType, Name, INDEX_NONE, ParameterNode->GetCPPType(), nullptr);
+					WorkData.VM->ParametersNameMap.Add(Parameter.Name, WorkData.VM->Parameters.Add(Parameter));
+				}
+			}
+		}
+	}
+
+#if !UE_RIGVM_UCLASS_BASED_STORAGE_DISABLED
+
 	if(WorkData.WatchedPins.Num() > 0)
 	{
 		for(int32 GraphIndex=0; GraphIndex<VisitedGraphs.Num(); GraphIndex++)
@@ -472,7 +523,7 @@ bool URigVMCompiler::Compile(URigVMGraph* InGraph, URigVMController* InControlle
 					}
 					FRigVMASTProxy PinProxy = FRigVMASTProxy::MakeFromUObject(ModelPin);
 					FRigVMVarExprAST TempVarExpr(FRigVMExprAST::EType::Var, PinProxy);
-					TempVarExpr.ParserPtr = WorkData.AST.Get();
+					TempVarExpr.ParserPtr = AST.Get();
 
 					FindOrAddRegister(&TempVarExpr, WorkData, true);
 				}
@@ -503,9 +554,11 @@ bool URigVMCompiler::Compile(URigVMGraph* InGraph, URigVMController* InControlle
 
 	WorkData.VM->ClearMemory();
 
+#endif
+
 	WorkData.bSetupMemory = false;
 	WorkData.ExprComplete.Reset();
-	for (FRigVMExprAST* RootExpr : *WorkData.AST)
+	for (FRigVMExprAST* RootExpr : *AST)
 	{
 		TraverseExpression(RootExpr, WorkData);
 	}
@@ -520,12 +573,31 @@ bool URigVMCompiler::Compile(URigVMGraph* InGraph, URigVMController* InControlle
 	// setup debug registers after all other registers have been created
 	if(Settings.EnablePinWatches)
 	{
+#if UE_RIGVM_UCLASS_BASED_STORAGE_DISABLED
+		for(int32 GraphIndex=0; GraphIndex<VisitedGraphs.Num(); GraphIndex++)
+		{
+			URigVMGraph* VisitedGraph = VisitedGraphs[GraphIndex];
+			for(URigVMNode* ModelNode : VisitedGraph->GetNodes())
+			{
+				for(URigVMPin* ModelPin : ModelNode->GetPins())
+				{
+					if(ModelPin->RequiresWatch(true))
+					{
+						CreateDebugRegister(ModelPin, WorkData.VM, WorkData.PinPathToOperand, AST);
+					}
+				}
+			}
+		}
+#else
 		for(URigVMPin* WatchedPin : WorkData.WatchedPins)
 		{
-			MarkDebugWatch(true, WatchedPin, WorkData.VM, WorkData.PinPathToOperand, WorkData.AST);
+			MarkDebugWatch(true, WatchedPin, WorkData.VM, WorkData.PinPathToOperand, AST);
 		}
+#endif
 	}
 
+#if !UE_RIGVM_UCLASS_BASED_STORAGE_DISABLED
+
 	// now that we have determined the needed memory, let's
 	// update the property paths once more
 	for(ERigVMMemoryType MemoryType : MemoryTypes)
@@ -553,6 +625,8 @@ bool URigVMCompiler::Compile(URigVMGraph* InGraph, URigVMController* InControlle
 		WorkData.VM->ExternalPropertyPathDescriptions = *Descriptions;
 	}
 
+#endif
+
 	return true;
 }
 
@@ -747,11 +821,6 @@ int32 URigVMCompiler::TraverseCallExtern(const FRigVMCallExternExprAST* InExpr,
 		return INDEX_NONE;
 	}
 
-	if(UnitNode->GetScriptStruct() == nullptr)
-	{
-		return INDEX_NONE;
-	}
-
 	int32 InstructionIndex = INDEX_NONE;
 
 	if (WorkData.bSetupMemory)
@@ -1005,79 +1074,102 @@ void URigVMCompiler::TraverseAssign(const FRigVMAssignExprAST* InExpr, FRigVMCom
 		{
 			return;
 		}
-		
+
 		// if this is a copy - we should check if operands need offsets
 		if (InExpr->GetType() == FRigVMExprAST::EType::Copy)
 		{
 			struct Local
 			{
-				static void SetupRegisterOffset(URigVM* VM, const FRigVMASTLinkDescription& InLink, URigVMPin* Pin,
-					FRigVMOperand& Operand, const FRigVMVarExprAST* VarExpr, bool bSource, FRigVMCompilerWorkData& WorkData)
+				static void SetupRegisterOffset(URigVM* VM, URigVMPin* Pin, FRigVMOperand& Operand, const FRigVMVarExprAST* VarExpr, bool bSource, FRigVMCompilerWorkData& WorkData)
 				{
-					const bool bHasTargetSegmentPath = !bSource && !InLink.SegmentPath.IsEmpty();
-					
 					URigVMPin* RootPin = Pin->GetRootPin();
-					if (Pin == RootPin && !bHasTargetSegmentPath)
+					if (Pin == RootPin)
 					{
 						return;
 					}
 
-					FString SegmentPath = Pin->GetSegmentPath(false);
-					if(bHasTargetSegmentPath)
+					if (URigVMSelectNode* SelectNode = Cast<URigVMSelectNode>(RootPin->GetNode()))
 					{
-						if(SegmentPath.IsEmpty())
-						{
-							SegmentPath = InLink.SegmentPath;
-						}
-						else
+						if (Pin->GetParentPin() == RootPin && RootPin->GetName() == URigVMSelectNode::ValueName)
 						{
-							SegmentPath = URigVMPin::JoinPinPath(SegmentPath, InLink.SegmentPath);
+							return;
 						}
 					}
 
-					// for select nodes we create a register for each case (since the cases are fixed in size)
-					// thus we do not need to setup a registeroffset for the array element.
-					if (URigVMSelectNode* SelectNode = Cast<URigVMSelectNode>(RootPin->GetNode()))
-					{
-						if(RootPin->GetName() == URigVMSelectNode::ValueName)
-						{
-							if (Pin->GetParentPin() == RootPin)
-							{
-								return;
-							}
+					FString SegmentPath = Pin->GetSegmentPath();
 
-							// if the pin is a sub pin of a case of the select (for example: Values.0.Translation)
-							// we'll need to re-adjust the root pin to the case pin (for example: Values.0)
-							TArray<FString> SegmentPathPaths;
-							if(ensure(URigVMPin::SplitPinPath(SegmentPath, SegmentPathPaths)))
-							{
-								RootPin = RootPin->FindSubPin(SegmentPathPaths[0]);
+#if UE_RIGVM_UCLASS_BASED_STORAGE_DISABLED
 
-								SegmentPathPaths.RemoveAt(0);
-								ensure(SegmentPathPaths.Num() > 0);
-								SegmentPath = URigVMPin::JoinPinPath(SegmentPathPaths);
-							}
-							else
-							{
-								return;
-							}
+					int32 ArrayIndex = INDEX_NONE;
+					if (RootPin->IsArray())
+					{
+						FString SegmentArrayIndex = SegmentPath;
+						int32 DotIndex = INDEX_NONE;
+						if (SegmentPath.FindChar('.', DotIndex))
+						{
+							SegmentArrayIndex = SegmentPath.Left(DotIndex);
+							SegmentPath = SegmentPath.Mid(DotIndex + 1);
 						}
+						else
+						{		
+							SegmentPath = FString();
+						}
+						ArrayIndex = FCString::Atoi(*SegmentArrayIndex);
+					}
+				
+					if (Operand.GetMemoryType() == ERigVMMemoryType::External)
+					{
+						const FRigVMExternalVariable& ExternalVariable = VM->GetExternalVariables()[Operand.GetRegisterIndex()];
+						UScriptStruct* ScriptStruct = CastChecked<UScriptStruct>(ExternalVariable.TypeObject);
+						int32 RegisterOffset = VM->GetWorkMemory().GetOrAddRegisterOffset(INDEX_NONE, ScriptStruct, SegmentPath, ArrayIndex == INDEX_NONE ? 0 : ArrayIndex, ExternalVariable.Size);
+						Operand = FRigVMOperand(ERigVMMemoryType::External, Operand.GetRegisterIndex(), RegisterOffset);
+					}
+					else
+					{
+						FRigVMMemoryContainer& Memory = Operand.GetMemoryType() == ERigVMMemoryType::Literal ? VM->GetLiteralMemory() :
+							(Operand.GetMemoryType() == ERigVMMemoryType::Work ? VM->GetWorkMemory() : VM->GetDebugMemory());
+						Operand = Memory.GetOperand(Operand.GetRegisterIndex(), SegmentPath, ArrayIndex);
 					}
 
+#else
+
 					const int32 PropertyPathIndex = WorkData.FindOrAddPropertyPath(Operand, RootPin->GetCPPType(), SegmentPath);
 					Operand = FRigVMOperand(Operand.GetMemoryType(), Operand.GetRegisterIndex(), PropertyPathIndex);
+					
+#endif
 				}
 			};
 
-			const FRigVMASTLinkDescription& Link = InExpr->GetLink();
-			Local::SetupRegisterOffset(WorkData.VM, Link, InExpr->GetSourcePin(), Source, SourceExpr, true, WorkData);
-			Local::SetupRegisterOffset(WorkData.VM, Link, InExpr->GetTargetPin(), Target, TargetExpr, false, WorkData);
+			Local::SetupRegisterOffset(WorkData.VM, InExpr->GetSourcePin(), Source, SourceExpr, true, WorkData);
+			Local::SetupRegisterOffset(WorkData.VM, InExpr->GetTargetPin(), Target, TargetExpr, false, WorkData);
 		}
 
 		FRigVMCopyOp CopyOp = WorkData.VM->GetCopyOpForOperands(Source, Target);
 		if(CopyOp.IsValid())
 		{
-			AddCopyOperator(CopyOp, InExpr, SourceExpr, TargetExpr, WorkData);
+			WorkData.VM->GetByteCode().AddCopyOp(CopyOp);
+
+			int32 InstructionIndex = WorkData.VM->GetByteCode().GetNumInstructions() - 1;
+			if (Settings.SetupNodeInstructionIndex)
+			{
+				if (URigVMPin* SourcePin = InExpr->GetSourcePin())
+				{
+					if (URigVMVariableNode* VariableNode = Cast<URigVMVariableNode>(SourcePin->GetNode()))
+					{
+						const FRigVMCallstack Callstack = SourceExpr->GetProxy().GetSibling(VariableNode).GetCallstack();
+						WorkData.VM->GetByteCode().SetSubject(InstructionIndex, Callstack.GetCallPath(), Callstack.GetStack());
+					}
+				}
+
+				if (URigVMPin* TargetPin = InExpr->GetTargetPin())
+				{
+					if (URigVMVariableNode* VariableNode = Cast<URigVMVariableNode>(TargetPin->GetNode()))
+					{
+						const FRigVMCallstack Callstack = TargetExpr->GetProxy().GetSibling(VariableNode).GetCallstack();
+						WorkData.VM->GetByteCode().SetSubject(InstructionIndex, Callstack.GetCallPath(), Callstack.GetStack());
+					}
+				}
+			}
 		}
 	}
 }
@@ -1272,6 +1364,19 @@ void URigVMCompiler::TraverseSelect(const FRigVMSelectExprAST* InExpr, FRigVMCom
 			{
 				FName LiteralName = *FString::FromInt(CaseIndex);
 
+#if UE_RIGVM_UCLASS_BASED_STORAGE_DISABLED
+				
+				int32 Register = WorkData.VM->GetLiteralMemory().Add<int32>(LiteralName, CaseIndex);
+ 
+#if WITH_EDITORONLY_DATA
+				WorkData.VM->GetLiteralMemory().Registers[Register].BaseCPPType = TEXT("int32");
+				WorkData.VM->GetLiteralMemory().Registers[Register].BaseCPPTypeObject = nullptr;
+#endif
+
+				FRigVMOperand Operand = WorkData.VM->GetLiteralMemory().GetOperand(Register);
+				
+#else
+
 				const FString DefaultValue = FString::FromInt(CaseIndex);
 				FRigVMOperand Operand = WorkData.AddProperty(
 					ERigVMMemoryType::Literal,
@@ -1280,18 +1385,34 @@ void URigVMCompiler::TraverseSelect(const FRigVMSelectExprAST* InExpr, FRigVMCom
 					nullptr,
 					DefaultValue);
 				
+#endif
 				WorkData.IntegerLiterals.Add(CaseIndex, Operand);
 			}
 		}
 
 		if (!WorkData.ComparisonOperand.IsValid())
 		{
+#if UE_RIGVM_UCLASS_BASED_STORAGE_DISABLED
+			
+			int32 Register = WorkData.VM->GetWorkMemory().Add<bool>(FName(TEXT("IntEquals")), false);
+
+#if WITH_EDITORONLY_DATA
+			WorkData.VM->GetWorkMemory().Registers[Register].BaseCPPType = TEXT("bool");
+			WorkData.VM->GetWorkMemory().Registers[Register].BaseCPPTypeObject = nullptr;
+#endif
+
+			WorkData.ComparisonOperand = WorkData.VM->GetWorkMemory().GetOperand(Register);
+
+#else
+
 			WorkData.ComparisonOperand = WorkData.AddProperty(
 				ERigVMMemoryType::Work,
 				FName(TEXT("IntEquals")),
 				TEXT("bool"),
 				nullptr,
 				TEXT("false"));
+
+#endif
 		}
 		return;
 	}
@@ -1712,107 +1833,6 @@ void URigVMCompiler::TraverseArray(const FRigVMArrayExprAST* InExpr, FRigVMCompi
 	}
 }
 
-void URigVMCompiler::AddCopyOperator(const FRigVMCopyOp& InOp, const FRigVMAssignExprAST* InAssignExpr,
-	const FRigVMVarExprAST* InSourceExpr, const FRigVMVarExprAST* InTargetExpr,  FRigVMCompilerWorkData& WorkData,
-	bool bDelayCopyOperations)
-{
-	if(bDelayCopyOperations)
-	{
-		// if this is a full literal copy, let's delay it.
-		// to maintain the execution order we want nodes which compose a value
-		// to delay their reset to the default value, which happens prior to
-		// computing dependencies.
-		// so for example an external variable of FVector may need to be reset
-		// to a literal value prior to the rest of the composition, for example
-		// if there's a float link only on the Y component. the execution order
-		// desired is this:
-		//
-		// * Run all dependent branches
-		// * Copy the literal value into the variable
-		// * Copy the parts into the variable (like the Y component).
-		// 
-		// By delaying the copy operator until right before the very first composition
-		// copy operator we ensure the desired execution order
-		if(InOp.Target.GetRegisterOffset() == INDEX_NONE && 
-			InOp.Source.GetMemoryType() == ERigVMMemoryType::Literal &&
-			InOp.Source.GetRegisterOffset() == INDEX_NONE)
-		{
-			if(URigVMPin* Pin = InTargetExpr->GetPin())
-			{
-				if(URigVMPin* RootPin = Pin->GetRootPin())
-				{
-					const FRigVMASTProxy RootPinProxy = InTargetExpr->GetProxy().GetSibling(RootPin);
-
-					// if the root pin has only links on its subpins
-					if(WorkData.AST->GetSourceLinkIndices(RootPinProxy, false).Num() == 0)
-					{
-						if(WorkData.AST->GetSourceLinkIndices(RootPinProxy, true).Num() > 0)
-						{					
-							FRigVMCompilerWorkData::FCopyOpInfo DeferredCopyOp;
-							DeferredCopyOp.Op = InOp;
-							DeferredCopyOp.AssignExpr = InAssignExpr;
-							DeferredCopyOp.SourceExpr = InSourceExpr;
-							DeferredCopyOp.TargetExpr = InTargetExpr;
-				
-							const FRigVMOperand Key(InOp.Target.GetMemoryType(), InOp.Target.GetRegisterIndex());
-							WorkData.DeferredCopyOps.FindOrAdd(Key) = DeferredCopyOp;
-							return;
-						}
-					}
-				}
-			}
-		}
-		
-		bDelayCopyOperations = false;
-	}
-
-	// loop up a potentially delayed copy operation which needs to happen
-	// just prior to this one and inject it as well.
-	if(!bDelayCopyOperations)
-	{
-		const FRigVMOperand DeferredKey(InOp.Target.GetMemoryType(), InOp.Target.GetRegisterIndex());
-		const FRigVMCompilerWorkData::FCopyOpInfo* DeferredCopyOpPtr = WorkData.DeferredCopyOps.Find(DeferredKey);
-		if(DeferredCopyOpPtr != nullptr)
-		{
-			FRigVMCompilerWorkData::FCopyOpInfo CopyOpInfo = *DeferredCopyOpPtr;
-			WorkData.DeferredCopyOps.Remove(DeferredKey);
-			AddCopyOperator(CopyOpInfo, WorkData, false);
-		}
-	}
-	
-	WorkData.VM->GetByteCode().AddCopyOp(InOp);
-
-	int32 InstructionIndex = WorkData.VM->GetByteCode().GetNumInstructions() - 1;
-	if (Settings.SetupNodeInstructionIndex)
-	{
-		if (URigVMPin* SourcePin = InAssignExpr->GetSourcePin())
-		{
-			if (URigVMVariableNode* VariableNode = Cast<URigVMVariableNode>(SourcePin->GetNode()))
-			{
-				const FRigVMCallstack Callstack = InSourceExpr->GetProxy().GetSibling(VariableNode).GetCallstack();
-				WorkData.VM->GetByteCode().SetSubject(InstructionIndex, Callstack.GetCallPath(), Callstack.GetStack());
-			}
-		}
-
-		if (URigVMPin* TargetPin = InAssignExpr->GetTargetPin())
-		{
-			if (URigVMVariableNode* VariableNode = Cast<URigVMVariableNode>(TargetPin->GetNode()))
-			{
-				const FRigVMCallstack Callstack = InTargetExpr->GetProxy().GetSibling(VariableNode).GetCallstack();
-				WorkData.VM->GetByteCode().SetSubject(InstructionIndex, Callstack.GetCallPath(), Callstack.GetStack());
-			}
-		}
-	}
-}
-
-void URigVMCompiler::AddCopyOperator(
-	const FRigVMCompilerWorkData::FCopyOpInfo& CopyOpInfo,
-	FRigVMCompilerWorkData& WorkData,
-	bool bDelayCopyOperations)
-{
-	AddCopyOperator(CopyOpInfo.Op, CopyOpInfo.AssignExpr, CopyOpInfo.SourceExpr, CopyOpInfo.TargetExpr, WorkData, bDelayCopyOperations);
-}
-
 void URigVMCompiler::InitializeLocalVariables(const FRigVMExprAST* InExpr, FRigVMCompilerWorkData& WorkData)
 {
 	// Initialize local variables if we are entering a new graph
@@ -1920,7 +1940,7 @@ void URigVMCompiler::InitializeLocalVariables(const FRigVMExprAST* InExpr, FRigV
 	}
 }
 
-FString URigVMCompiler::GetPinHash(const URigVMPin* InPin, const FRigVMVarExprAST* InVarExpr, bool bIsDebugValue, const FRigVMASTProxy& InPinProxy)
+FString URigVMCompiler::GetPinHash(const URigVMPin* InPin, const FRigVMVarExprAST* InVarExpr, bool bIsDebugValue)
 {
 	FString Prefix = bIsDebugValue ? TEXT("DebugWatch:") : TEXT("");
 	FString Suffix;
@@ -1934,6 +1954,7 @@ FString URigVMCompiler::GetPinHash(const URigVMPin* InPin, const FRigVMVarExprAS
 
 	bool bIsExecutePin = false;
 	bool bIsLiteral = false;
+	bool bIsParameter = false;
 	bool bIsVariable = false;
 
 	if (InVarExpr != nullptr && !bIsDebugValue)
@@ -1957,11 +1978,11 @@ FString URigVMCompiler::GetPinHash(const URigVMPin* InPin, const FRigVMVarExprAS
 		bIsExecutePin = InPin->IsExecuteContext();
 		bIsLiteral = InVarExpr->GetType() == FRigVMExprAST::EType::Literal;
 
+		bIsParameter = Cast<URigVMParameterNode>(Node) != nullptr;
 		bIsVariable = Cast<URigVMVariableNode>(Node) != nullptr || InVarExpr->IsA(FRigVMExprAST::ExternalVar);
 
 		// determine if this is an initialization for an IO pin
-		if (!bIsLiteral &&
-			!bIsVariable &&
+		if (!bIsLiteral && !bIsParameter && !bIsVariable &&
 			!bIsExecutePin && (InPin->GetDirection() == ERigVMPinDirection::IO ||
 			(InPin->GetDirection() == ERigVMPinDirection::Input && InPin->GetSourceLinks().Num() == 0)))
 		{
@@ -1974,7 +1995,14 @@ FString URigVMCompiler::GetPinHash(const URigVMPin* InPin, const FRigVMVarExprAS
 	}
 
 	bool bUseFullNodePath = true;
-	if (URigVMVariableNode* VariableNode = Cast<URigVMVariableNode>(Node))
+	if (URigVMParameterNode* ParameterNode = Cast<URigVMParameterNode>(Node))
+	{
+		if (InPin->GetName() == TEXT("Value") && !bIsLiteral && !bIsDebugValue)
+		{
+			return FString::Printf(TEXT("%sParameter::%s%s"), *Prefix, *ParameterNode->GetParameterName().ToString(), *Suffix);
+		}
+	}
+	else if (URigVMVariableNode* VariableNode = Cast<URigVMVariableNode>(Node))
 	{
 		if (InPin->GetName() == TEXT("Value") && !bIsDebugValue)
 		{
@@ -1985,15 +2013,10 @@ FString URigVMCompiler::GetPinHash(const URigVMPin* InPin, const FRigVMVarExprAS
 				if (bIsLiteral)
 				{
 					// Literal values will be reused for all instance of local variables
-					if (InVarExpr && InVarExpr->NumParents() == 0 && InVarExpr->NumChildren() == 0)
+					if (InVarExpr->NumParents() == 0 && InVarExpr->NumChildren() == 0)
 					{
 						return FString::Printf(TEXT("%sLocalVariableDefault::%s|%s%s"), *Prefix, *Node->GetGraph()->GetGraphName(), *VariableName.ToString(), *Suffix);
 					}
-					else if (InVarExpr)
-					{
-						const FString GraphPath = InVarExpr->GetProxy().GetCallstack().GetCallPath(false);
-						return FString::Printf(TEXT("%sLocalVariable::%s|%s%s"), *Prefix, *GraphPath, *VariableName.ToString(), *Suffix);					
-					}
 					else
 					{
 						return FString::Printf(TEXT("%sLocalVariable::%s|%s%s"), *Prefix, *Node->GetGraph()->GetGraphName(), *VariableName.ToString(), *Suffix);					
@@ -2020,20 +2043,6 @@ FString URigVMCompiler::GetPinHash(const URigVMPin* InPin, const FRigVMVarExprAS
 					}
 				}
 			}
-			else if(VariableNode->IsInputArgument())
-			{
-				FString FullPath;
-				if (InPinProxy.IsValid())
-				{
-					FullPath = InPinProxy.GetCallstack().GetCallPath(true);
-				}
-				else if(InVarExpr)
-				{						
-					const FRigVMASTProxy NodeProxy = InVarExpr->GetProxy().GetSibling(Node);
-					FullPath = InPinProxy.GetCallstack().GetCallPath(true);
-				}
-				return FString::Printf(TEXT("%s%s%s"), *Prefix, *FullPath, *Suffix);
-			}
 
 			if (!bIsLiteral)
 			{		
@@ -2062,7 +2071,11 @@ FString URigVMCompiler::GetPinHash(const URigVMPin* InPin, const FRigVMVarExprAS
 	}
 	else
 	{
+#if UE_RIGVM_UCLASS_BASED_STORAGE_DISABLED
+		if (InVarExpr)
+#else
 		if (InVarExpr && !bIsDebugValue)
+#endif
 		{
 			const FRigVMASTProxy NodeProxy = InVarExpr->GetProxy().GetSibling(Node);
 			if (const FRigVMExprAST* NodeExpr = InVarExpr->GetParser()->GetExprForSubject(NodeProxy))
@@ -2075,36 +2088,10 @@ FString URigVMCompiler::GetPinHash(const URigVMPin* InPin, const FRigVMVarExprAS
 					bUseFullNodePath = false;
 				}
 			}
-			else if(Node->IsA<URigVMFunctionEntryNode>() || Node->IsA<URigVMFunctionReturnNode>())
-			{
-				const FString FullPath = InPinProxy.GetCallstack().GetCallPath(true);
-				return FString::Printf(TEXT("%s%s%s"), *Prefix, *FullPath, *Suffix);
-			}
-		}
-	}
-
-	if (InPinProxy.IsValid())
-	{
-		const FString FullPath = InPinProxy.GetCallstack().GetCallPath(true);
-		return FString::Printf(TEXT("%s%s%s"), *Prefix, *FullPath, *Suffix);
-	}
-
-	if (InVarExpr)
-	{
-		if (bUseFullNodePath)
-		{
-			FString FullPath = InVarExpr->GetProxy().GetCallstack().GetCallPath(true);
-			return FString::Printf(TEXT("%s%s%s"), *Prefix, *FullPath, *Suffix);
-		}
-		else
-		{
-			return FString::Printf(TEXT("%s%s%s"), *Prefix, *InPin->GetPinPath(), *Suffix);
 		}
 	}
 
-	FString PinPath = InPin->GetPinPath(bUseFullNodePath);
-	ensureMsgf(!PinPath.StartsWith(TEXT("FunctionLibrary::")), TEXT("A library path should never be part of a pin hash %s."), *PinPath);
-	return FString::Printf(TEXT("%s%s%s"), *Prefix, *PinPath, *Suffix);
+	return FString::Printf(TEXT("%s%s%s"), *Prefix, *InPin->GetPinPath(bUseFullNodePath), *Suffix);
 }
 
 const FRigVMVarExprAST* URigVMCompiler::GetSourceVarExpr(const FRigVMExprAST* InExpr)
@@ -2157,6 +2144,87 @@ const FRigVMVarExprAST* URigVMCompiler::GetSourceVarExpr(const FRigVMExprAST* In
 	return nullptr;
 }
 
+#if UE_RIGVM_UCLASS_BASED_STORAGE_DISABLED
+
+void URigVMCompiler::CreateDebugRegister(URigVMPin* InPin, URigVM* OutVM,
+	TMap<FString, FRigVMOperand>* OutOperands, TSharedPtr<FRigVMParserAST> InRuntimeAST)
+{
+	check(InPin);
+	check(OutVM);
+	check(OutOperands);
+	check(InRuntimeAST.IsValid());
+	
+	FRigVMCompilerWorkData WorkData;
+	WorkData.VM = OutVM;
+	WorkData.PinPathToOperand = OutOperands;
+
+	URigVMPin* Pin = InPin->GetRootPin();
+	URigVMPin* SourcePin = Pin;
+	if(Settings.ASTSettings.bFoldAssignments)
+	{
+		while(SourcePin->GetSourceLinks().Num() > 0)
+		{
+			SourcePin = SourcePin->GetSourceLinks()[0]->GetSourcePin();
+		}
+	}
+	
+	TArray<const FRigVMExprAST*> Expressions = InRuntimeAST->GetExpressionsForSubject(SourcePin);
+	for(const FRigVMExprAST* Expression : Expressions)
+	{
+		check(Expression->IsA(FRigVMExprAST::EType::Var));
+		const FRigVMVarExprAST* VarExpression = Expression->To<FRigVMVarExprAST>();
+		const FString PinHash = GetPinHash(SourcePin, VarExpression, false);
+		if(const FRigVMOperand* Operand = OutOperands->Find(PinHash))
+		{
+			const FRigVMASTProxy PinProxy = FRigVMASTProxy::MakeFromUObject(Pin);
+			FRigVMVarExprAST TempVarExpr(FRigVMExprAST::EType::Var, PinProxy);
+			TempVarExpr.ParserPtr = InRuntimeAST.Get();
+	
+			const FRigVMOperand DebugOperand = FindOrAddRegister(&TempVarExpr, WorkData, true);
+			if(DebugOperand.IsValid())
+			{
+				FRigVMOperand KeyOperand(Operand->GetMemoryType(), Operand->GetRegisterIndex()); // no register offset
+				OutVM->OperandToDebugRegisters.FindOrAdd(KeyOperand).AddUnique(DebugOperand);
+			}
+		}
+	}
+}
+
+void URigVMCompiler::RemoveDebugRegister(URigVMPin* InPin, URigVM* OutVM,
+	TMap<FString, FRigVMOperand>* OutOperands, TSharedPtr<FRigVMParserAST> InRuntimeAST)
+{
+	check(OutVM);
+	check(OutOperands);
+
+	URigVMPin* Pin = InPin->GetRootPin();
+	const FRigVMASTProxy PinProxy = FRigVMASTProxy::MakeFromUObject(Pin);
+	FRigVMVarExprAST TempVarExpr(FRigVMExprAST::EType::Var, PinProxy);
+	TempVarExpr.ParserPtr = InRuntimeAST.Get();
+
+	const FString PinHash = GetPinHash(Pin, &TempVarExpr, true);
+	if(const FRigVMOperand* DebugOperand = OutOperands->Find(PinHash))
+	{
+		TArray<FRigVMOperand> KeysToRemove;
+		for(TPair<FRigVMOperand, TArray<FRigVMOperand>>& Pair : OutVM->OperandToDebugRegisters)
+		{
+			if(Pair.Value.Remove(*DebugOperand) > 0)
+			{
+				if(Pair.Value.IsEmpty())
+				{
+					KeysToRemove.AddUnique(Pair.Key);
+				}
+			}
+		}
+		for(const FRigVMOperand& KeyToRemove : KeysToRemove)
+		{
+			OutVM->OperandToDebugRegisters.Remove(KeyToRemove);
+		}
+		OutOperands->Remove(PinHash);
+	}
+}
+
+#else
+
 void URigVMCompiler::MarkDebugWatch(bool bRequired, URigVMPin* InPin, URigVM* OutVM,
 	TMap<FString, FRigVMOperand>* OutOperands, TSharedPtr<FRigVMParserAST> InRuntimeAST)
 {
@@ -2232,6 +2300,8 @@ void URigVMCompiler::MarkDebugWatch(bool bRequired, URigVMPin* InPin, URigVM* Ou
 	}
 }
 
+#endif
+
 UScriptStruct* URigVMCompiler::GetScriptStructForCPPType(const FString& InCPPType)
 {
 	if (InCPPType == TEXT("FRotator"))
@@ -2359,10 +2429,12 @@ FRigVMOperand URigVMCompiler::FindOrAddRegister(const FRigVMVarExprAST* InVarExp
 
 	URigVMPin* Pin = InVarExpr->GetPin();
 
+#if !UE_RIGVM_UCLASS_BASED_STORAGE_DISABLED
 	if(Pin->IsExecuteContext() && bIsDebugValue)
 	{
 		return FRigVMOperand();
 	}
+#endif
 	
 	FString CPPType = Pin->GetCPPType();
 	FString BaseCPPType = Pin->IsArray() ? Pin->GetArrayElementCppType() : CPPType;
@@ -2372,6 +2444,7 @@ FRigVMOperand URigVMCompiler::FindOrAddRegister(const FRigVMVarExprAST* InVarExp
 
 	bool bIsExecutePin = Pin->IsExecuteContext();
 	bool bIsLiteral = InVarExpr->GetType() == FRigVMExprAST::EType::Literal && !bIsDebugValue;
+	bool bIsParameter = InVarExpr->IsGraphParameter();
 	bool bIsVariable = Pin->IsRootPin() && (Pin->GetName() == URigVMVariableNode::ValueName) &&
 		InVarExpr->GetPin()->GetNode()->IsA<URigVMVariableNode>();
 
@@ -2392,9 +2465,21 @@ FRigVMOperand URigVMCompiler::FindOrAddRegister(const FRigVMVarExprAST* InVarExp
 		}
 	}
 
+#if UE_RIGVM_UCLASS_BASED_STORAGE_DISABLED
+	
+	FRigVMMemoryContainer& Memory = 
+	    bIsLiteral ? WorkData.VM->GetLiteralMemory() :
+			(bIsDebugValue ? WorkData.VM->GetDebugMemory() : WorkData.VM->GetWorkMemory());
+
+	const ERigVMMemoryType MemoryType = Memory.GetMemoryType();
+
+#else
+
 	const ERigVMMemoryType MemoryType =
 		bIsLiteral ? ERigVMMemoryType::Literal:
 		(bIsDebugValue ? ERigVMMemoryType::Debug : ERigVMMemoryType::Work);
+	
+#endif
 
 	FRigVMOperand const* ExistingOperand = WorkData.PinPathToOperand->Find(Hash);
 	if (ExistingOperand)
@@ -2421,8 +2506,36 @@ FRigVMOperand URigVMCompiler::FindOrAddRegister(const FRigVMVarExprAST* InVarExp
 		{
 			if (Pin->GetDirection() == ERigVMPinDirection::Hidden)
 			{
+#if UE_RIGVM_UCLASS_BASED_STORAGE_DISABLED
+				ensure(WorkData.DefaultStructs.Num() > 0);
+				TSharedPtr<FStructOnScope> DefaultStruct = WorkData.DefaultStructs.Last();
+				FRigVMStruct* VMStruct = (FRigVMStruct *)DefaultStruct->GetStructMemory();
+
+				int32 DesiredArraySize = VMStruct->GetArraySize(Pin->GetFName(), WorkData.RigVMUserData);
+
+#endif
+				
 				JoinedDefaultValue = Pin->GetDefaultValue(PinOverride);
 				DefaultValues = URigVMPin::SplitDefaultValue(JoinedDefaultValue);
+
+#if UE_RIGVM_UCLASS_BASED_STORAGE_DISABLED
+
+				if (DefaultValues.Num() != DesiredArraySize)
+				{
+					FString FirstDefaultValue;
+					if (Pin->GetArraySize() > 0)
+					{
+						FirstDefaultValue = Pin->GetSubPins()[0]->GetDefaultValue(PinOverride);
+					}
+
+					DefaultValues.Reset();
+					for (int32 Index = 0; Index < DesiredArraySize; Index++)
+					{
+						DefaultValues.Add(FirstDefaultValue);
+					}
+				}
+				
+#endif
 			}
 			else
 			{
@@ -2477,6 +2590,47 @@ FRigVMOperand URigVMCompiler::FindOrAddRegister(const FRigVMVarExprAST* InVarExp
 			int32 Register = INDEX_NONE;
 
 			// debug watch register might already exists - look for them by name
+#if UE_RIGVM_UCLASS_BASED_STORAGE_DISABLED
+			if(bIsDebugValue && Memory.SupportsNames())
+			{
+				Register = Memory.GetIndex(RegisterName);
+				if(Register != INDEX_NONE)
+				{
+					bool bFoundValidRegister = false;
+					if(ScriptStruct)
+					{
+						if(Memory[Register].BaseCPPTypeObject == ScriptStruct)
+						{
+							bFoundValidRegister = true;
+						}
+					}
+					else if(UEnum* Enum = Pin->GetEnum())
+					{
+						if(Memory[Register].BaseCPPTypeObject == Enum)
+						{
+							bFoundValidRegister = true;
+						}
+					}
+					else
+					{
+						if(Memory[Register].BaseCPPType == *BaseCPPType)
+						{
+							bFoundValidRegister = true;
+						}
+					}
+
+					if(bFoundValidRegister)
+					{
+						Operand = Memory.GetOperand(Register);
+						if(ExistingOperand == nullptr)
+						{
+							WorkData.PinPathToOperand->Add(Hash, Operand);
+						}
+						return Operand;
+					}
+				}
+			}
+#else
 			if(bIsDebugValue)
 			{
 				Operand = WorkData.FindProperty(MemoryType, RegisterName);
@@ -2493,12 +2647,192 @@ FRigVMOperand URigVMCompiler::FindOrAddRegister(const FRigVMVarExprAST* InVarExp
 					}
 				}
 			}
+#endif
+
+#if UE_RIGVM_UCLASS_BASED_STORAGE_DISABLED
+			
+			if (ScriptStruct)
+			{
+				TArray<uint8, TAlignedHeapAllocator<16>> Data;
+				Data.AddUninitialized(ScriptStruct->GetStructureSize() * DefaultValues.Num());
+				uint8* Ptr = Data.GetData();
+
+				for (FString DefaultValue : DefaultValues)
+				{
+					ScriptStruct->InitializeStruct(Ptr, 1);
+					if (!DefaultValue.IsEmpty() && DefaultValue != TEXT("()"))
+					{
+						FRigVMCompilerImportErrorContext ErrorPipe(this);
+						ScriptStruct->ImportText(*DefaultValue, Ptr, nullptr, PPF_None, &ErrorPipe, ScriptStruct->GetName());
+					}
+					Ptr = Ptr + ScriptStruct->GetStructureSize();
+				}
+
+				Register = Memory.AddRegisterArray(!Pin->IsDynamicArray() && !bIsDebugValue, RegisterName, ScriptStruct->GetStructureSize(), DefaultValues.Num(), Pin->IsArray(), Data.GetData(), NumSlices, ERigVMRegisterType::Struct, ScriptStruct);
+				check(Register != INDEX_NONE);
+				
+#if WITH_EDITORONLY_DATA
+				Memory.Registers[Register].BaseCPPType = *ScriptStruct->GetStructCPPName();
+				Memory.Registers[Register].BaseCPPTypeObject = ScriptStruct;
+#endif
+				ScriptStruct->DestroyStruct(Data.GetData(), DefaultValues.Num());
+
+				Operand = Memory.GetOperand(Register);
+			}
+			else if (UEnum* Enum = Pin->GetEnum())
+			{
+				FRigVMByteArray Values;
+				for (FString DefaultValue : DefaultValues)
+				{
+					if (DefaultValue.IsEmpty())
+					{
+						Values.Add((uint8)Enum->GetValueByIndex(0));
+					}
+					else
+					{
+						Values.Add((uint8)Enum->GetValueByNameString(DefaultValue));
+					}
+				}
+
+				Register = Memory.AddRegisterArray<uint8>(!Pin->IsDynamicArray() && !bIsDebugValue, RegisterName, Values.Num(), Pin->IsArray(), Values.GetData(), NumSlices, ERigVMRegisterType::Plain, nullptr);
+				check(Register != INDEX_NONE);
+				
+#if WITH_EDITORONLY_DATA
+				Memory.Registers[Register].BaseCPPType = *Enum->CppType;
+				Memory.Registers[Register].BaseCPPTypeObject = Enum;
+#endif
+				Operand = Memory.GetOperand(Register);
+			}
+			else if (BaseCPPType == TEXT("bool"))
+			{
+				TArray<bool> Values;
+				for (FString DefaultValue : DefaultValues)
+				{
+					Values.Add((DefaultValue == TEXT("True")) || (DefaultValue == TEXT("true")) || (DefaultValue == TEXT("1")));
+				}
+
+				Register = Memory.AddRegisterArray<bool>(!Pin->IsDynamicArray() && !bIsDebugValue, RegisterName, Values.Num(), Pin->IsArray(), (uint8*)Values.GetData(), NumSlices, ERigVMRegisterType::Plain, nullptr);
+				check(Register != INDEX_NONE);
+				
+#if WITH_EDITORONLY_DATA
+				Memory.Registers[Register].BaseCPPType = *BaseCPPType;
+				Memory.Registers[Register].BaseCPPTypeObject = nullptr;
+#endif
+				Operand = Memory.GetOperand(Register);
+			}
+			else if (BaseCPPType == TEXT("int32"))
+			{
+				TArray<int32> Values;
+				for (FString DefaultValue : DefaultValues)
+				{
+					if (DefaultValue.IsEmpty())
+					{
+						Values.Add(0);
+					}
+					else
+					{
+						Values.Add(FCString::Atoi(*DefaultValue));
+					}
+				}
+				Register = Memory.AddRegisterArray<int32>(!Pin->IsDynamicArray() && !bIsDebugValue, RegisterName, Values.Num(), Pin->IsArray(), (uint8*)Values.GetData(), NumSlices, ERigVMRegisterType::Plain, nullptr);
+				check(Register != INDEX_NONE);
+
+#if WITH_EDITORONLY_DATA
+				Memory.Registers[Register].BaseCPPType = *BaseCPPType;
+				Memory.Registers[Register].BaseCPPTypeObject = nullptr;
+#endif
+				Operand = Memory.GetOperand(Register);
+			}
+			else if (BaseCPPType == TEXT("float"))
+			{
+				TArray<float> Values;
+				for (FString DefaultValue : DefaultValues)
+				{
+					if (DefaultValue.IsEmpty())
+					{
+						Values.Add(0.f);
+					}
+					else
+					{
+						Values.Add(FCString::Atof(*DefaultValue));
+					}
+				}
+				Register = Memory.AddRegisterArray<float>(!Pin->IsDynamicArray() && !bIsDebugValue, RegisterName, Values.Num(), Pin->IsArray(), (uint8*)Values.GetData(), NumSlices, ERigVMRegisterType::Plain, nullptr);
+				check(Register != INDEX_NONE);
+				
+#if WITH_EDITORONLY_DATA
+				Memory.Registers[Register].BaseCPPType = *BaseCPPType;
+				Memory.Registers[Register].BaseCPPTypeObject = nullptr;
+#endif
+				Operand = Memory.GetOperand(Register);
+			}
+			else if (BaseCPPType == TEXT("double"))
+			{
+				TArray<double> Values;
+				for (FString DefaultValue : DefaultValues)
+				{
+					if (DefaultValue.IsEmpty())
+					{
+						Values.Add(0.0);
+					}
+					else
+					{
+						Values.Add(FCString::Atod(*DefaultValue));
+					}
+				}
+				Register = Memory.AddRegisterArray<double>(!Pin->IsDynamicArray() && !bIsDebugValue, RegisterName, Values.Num(), Pin->IsArray(), (uint8*)Values.GetData(), NumSlices, ERigVMRegisterType::Plain, nullptr);
+				check(Register != INDEX_NONE);
+				
+#if WITH_EDITORONLY_DATA
+				Memory.Registers[Register].BaseCPPType = *BaseCPPType;
+				Memory.Registers[Register].BaseCPPTypeObject = nullptr;
+#endif
+				Operand = Memory.GetOperand(Register);
+			}
+			else if (BaseCPPType == TEXT("FName"))
+			{
+				TArray<FName> Values;
+				for (FString DefaultValue : DefaultValues)
+				{
+					Values.Add(*DefaultValue);
+				}
+
+				Register = Memory.AddRegisterArray<FName>(!Pin->IsDynamicArray() && !bIsDebugValue, RegisterName, Values.Num(), Pin->IsArray(), (uint8*)Values.GetData(), NumSlices, ERigVMRegisterType::Name, nullptr);
+				check(Register != INDEX_NONE);
+				
+#if WITH_EDITORONLY_DATA
+				Memory.Registers[Register].BaseCPPType = *BaseCPPType;
+				Memory.Registers[Register].BaseCPPTypeObject = nullptr;
+#endif
+				Operand = Memory.GetOperand(Register);
+			}
+			else if (BaseCPPType == TEXT("FString"))
+			{
+				Register = Memory.AddRegisterArray<FString>(!Pin->IsDynamicArray() && !bIsDebugValue, RegisterName, DefaultValues.Num(), Pin->IsArray(), (uint8*)DefaultValues.GetData(), NumSlices, ERigVMRegisterType::String, nullptr);
+				check(Register != INDEX_NONE);
+				
+#if WITH_EDITORONLY_DATA
+				Memory.Registers[Register].BaseCPPType = *BaseCPPType;
+				Memory.Registers[Register].BaseCPPTypeObject = nullptr;
+#endif
+				Operand = Memory.GetOperand(Register);
+			}
+			else
+			{
+				ensure(false);
+			}
+
+			Memory[Operand.GetRegisterIndex()].bIsArray = Pin->IsArray();
+		}
+
+#else
+
 		}
 
 		if(bIsDebugValue)
 		{
 			// debug values are always stored as arrays
-			CPPType = RigVMTypeUtils::ArrayTypeFromBaseType(CPPType);
+			CPPType = FString::Printf(URigVMController::TArrayTemplate, *CPPType);
 			JoinedDefaultValue = URigVMPin::GetDefaultValueForArray({ JoinedDefaultValue });
 		}
 		else if(Pin->GetDirection() == ERigVMPinDirection::Hidden && Pin->GetNode()->IsA<URigVMUnitNode>())
@@ -2509,7 +2843,7 @@ FRigVMOperand URigVMCompiler::FindOrAddRegister(const FRigVMVarExprAST* InVarExp
 
 			if (!Property->HasMetaData(FRigVMStruct::SingletonMetaName))
 			{
-				CPPType = RigVMTypeUtils::ArrayTypeFromBaseType(CPPType);
+				CPPType = FString::Printf(URigVMController::TArrayTemplate, *CPPType);
 				JoinedDefaultValue = URigVMPin::GetDefaultValueForArray({ JoinedDefaultValue });
 			}
 		}
@@ -2520,6 +2854,18 @@ FRigVMOperand URigVMCompiler::FindOrAddRegister(const FRigVMVarExprAST* InVarExp
 		}
 	
 		Operand = WorkData.AddProperty(MemoryType, RegisterName, CPPType, Pin->GetCPPTypeObject(), JoinedDefaultValue);
+		
+#endif
+
+		if (bIsParameter && !bIsLiteral && !bIsDebugValue)
+		{
+			URigVMParameterNode* ParameterNode = Cast<URigVMParameterNode>(Pin->GetNode());
+			check(ParameterNode);
+			FName Name = ParameterNode->GetParameterName();
+			ERigVMParameterType ParameterType = ParameterNode->IsInput() ? ERigVMParameterType::Input : ERigVMParameterType::Output;
+			FRigVMParameter Parameter(ParameterType, Name, Operand.GetRegisterIndex(), Pin->GetCPPType(), ScriptStruct);
+			WorkData.VM->ParametersNameMap.FindOrAdd(Parameter.Name) = WorkData.VM->Parameters.Add(Parameter);
+		}
 	}
 	ensure(Operand.IsValid());
 
@@ -2536,7 +2882,7 @@ FRigVMOperand URigVMCompiler::FindOrAddRegister(const FRigVMVarExprAST* InVarExp
 			{
 				if (URigVMPin* VirtualPin = Cast<URigVMPin>(Proxy.GetSubject()))
 				{
-					FString VirtualPinHash = GetPinHash(VirtualPin, InVarExpr, bIsDebugValue, Proxy);
+					FString VirtualPinHash = GetPinHash(VirtualPin, InVarExpr, bIsDebugValue);
 					WorkData.PinPathToOperand->Add(VirtualPinHash, Operand);
 				}	
 			}
@@ -2573,8 +2919,6 @@ const FRigVMCompilerWorkData::FRigVMASTProxyArray& URigVMCompiler::FindProxiesWi
 
 	PinProxiesToProcess.Add(InProxy);
 
-	const FString CPPType = InProxy.GetSubjectChecked<URigVMPin>()->GetCPPType();
-
 	for(int32 ProxyIndex = 0; ProxyIndex < PinProxiesToProcess.Num(); ProxyIndex++)
 	{
 		const FRigVMASTProxy& CurrentProxy = PinProxiesToProcess[ProxyIndex];
@@ -2583,19 +2927,13 @@ const FRigVMCompilerWorkData::FRigVMASTProxyArray& URigVMCompiler::FindProxiesWi
 		{
 			if (URigVMPin* Pin = Cast<URigVMPin>(CurrentProxy.GetSubject()))
 			{
-				if (Pin->GetNode()->IsA<URigVMVariableNode>())
+				if (Pin->GetNode()->IsA<URigVMVariableNode>() || Pin->GetNode()->IsA<URigVMParameterNode>())
 				{
 					if (Pin->GetDirection() == ERigVMPinDirection::Input)
 					{
 						continue;
 					}
 				}
-				// due to LWC we may have two pins that don't
-				// actually share the same CPP type (float vs double)
-				if(Pin->GetCPPType() != CPPType)
-				{
-					continue;
-				}
 			}
 			PinProxies.Add(CurrentProxy);
 		}
@@ -2648,9 +2986,9 @@ bool URigVMCompiler::ValidateNode(URigVMNode* InNode)
 {
 	check(InNode);
 
-	if(InNode->HasWildCardPin())
+	if(InNode->HasUnknownTypePin())
 	{
-		static const FString UnknownTypeMessage = TEXT("Node @@ has unresolved pins of wildcard type.");
+		static const FString UnknownTypeMessage = TEXT("Node @@ has unresolved pins of unknown type.");
 		Settings.Report(EMessageSeverity::Error, InNode, UnknownTypeMessage);
 		return false;
 	}
