Control Rig: Fix non-unity builds.

#rnx

diff --git a/Engine/Source/Developer/RigVMDeveloper/Private/RigVMCompiler/RigVMCodeGenerator.cpp b/Engine/Source/Developer/RigVMDeveloper/Private/RigVMCompiler/RigVMCodeGenerator.cpp
--- a/Engine/Source/Developer/RigVMDeveloper/Private/RigVMCompiler/RigVMCodeGenerator.cpp
+++ b/Engine/Source/Developer/RigVMDeveloper/Private/RigVMCompiler/RigVMCodeGenerator.cpp
@@ -2,8 +2,8 @@
 
 #include "RigVMCompiler/RigVMCodeGenerator.h"
 #include "RigVMCore/RigVMExecuteContext.h"
+#include "RigVMCore/RigVMStruct.h"
 #include "RigVMDeveloperModule.h"
-#include "RigVMDeveloperTypeUtils.h"
 #include "Algo/Count.h"
 #include "Animation/Rig.h"
 
diff --git a/Engine/Source/Developer/RigVMDeveloper/Public/RigVMCompiler/RigVMCodeGenerator.h b/Engine/Source/Developer/RigVMDeveloper/Public/RigVMCompiler/RigVMCodeGenerator.h
--- a/Engine/Source/Developer/RigVMDeveloper/Public/RigVMCompiler/RigVMCodeGenerator.h
+++ b/Engine/Source/Developer/RigVMDeveloper/Public/RigVMCompiler/RigVMCodeGenerator.h
@@ -2,12 +2,13 @@
 
 #pragma once
 
-#include "CoreMinimal.h"
-
 #include "RigVMCore/RigVM.h"
-#include "RigVMCore/RigVMStruct.h"
 #include "RigVMModel/RigVMGraph.h"
-#include "RigVMCompiler/RigVMCompiler.h"
+
+#include "UObject/StrongObjectPtr.h"
+
+class URigVMGraph;
+
 
 struct RIGVMDEVELOPER_API FRigVMCodeGenerator
 {
