Control Rig: Fix crash during undo / redo

diff --git a/Engine/Source/Developer/RigVMDeveloper/Private/RigVMModel/RigVMControllerActions.cpp b/Engine/Source/Developer/RigVMDeveloper/Private/RigVMModel/RigVMControllerActions.cpp
--- a/Engine/Source/Developer/RigVMDeveloper/Private/RigVMModel/RigVMControllerActions.cpp
+++ b/Engine/Source/Developer/RigVMDeveloper/Private/RigVMModel/RigVMControllerActions.cpp
@@ -1459,11 +1459,32 @@ bool FRigVMChangePinTypeAction::Undo(URigVMController* InController)
 	{
 		return false;
 	}
+	
+	if(const URigVMPin* TemplateNodePin = InController->GetGraph()->FindPin(PinPath))
+	{
+		if(TemplateNodePin->GetNode()->IsA<URigVMTemplateNode>() && TemplateNodePin->IsWildCard())
+		{
+			return InController->ResolveWildCardPin(PinPath, OldCPPType, OldCPPTypeObjectPath, false, false);
+		}
+	}
+	
 	return InController->ChangePinType(PinPath, OldCPPType, OldCPPTypeObjectPath, false, bSetupOrphanPins, bBreakLinks, bRemoveSubPins);
 }
 
 bool FRigVMChangePinTypeAction::Redo(URigVMController* InController)
 {
+	if(const URigVMPin* TemplateNodePin = InController->GetGraph()->FindPin(PinPath))
+	{
+		if(TemplateNodePin->GetNode()->IsA<URigVMTemplateNode>() && TemplateNodePin->IsWildCard())
+		{
+			if(!InController->ResolveWildCardPin(PinPath, NewCPPType, NewCPPTypeObjectPath, false, false))
+			{
+				return false;
+			}
+			return FRigVMBaseAction::Redo(InController);
+		}
+	}
+
 	if(!InController->ChangePinType(PinPath, NewCPPType, NewCPPTypeObjectPath, false, bSetupOrphanPins, bBreakLinks, bRemoveSubPins))
 	{
 		return false;
