Control Rig: Fix crashes during upgrade / unresolve

diff --git a/Engine/Plugins/Animation/ControlRig/Source/ControlRigDeveloper/Private/Graph/ControlRigGraph.cpp b/Engine/Plugins/Animation/ControlRig/Source/ControlRigDeveloper/Private/Graph/ControlRigGraph.cpp
--- a/Engine/Plugins/Animation/ControlRig/Source/ControlRigDeveloper/Private/Graph/ControlRigGraph.cpp
+++ b/Engine/Plugins/Animation/ControlRig/Source/ControlRigDeveloper/Private/Graph/ControlRigGraph.cpp
@@ -931,6 +931,7 @@ void UControlRigGraph::RemoveAndDeleteNode(UEdGraphNode* InNode)
 		}
 		while (ExistingObject);
 		InNode->Rename(*DeletedName, NewOuter, REN_ForceNoResetLoaders | REN_DoNotDirty | REN_DontCreateRedirectors | REN_NonTransactional);
+		InNode->MarkAsGarbage();
 	}
 
 	// Remove the node and notify of the removal
diff --git a/Engine/Source/Developer/RigVMDeveloper/Private/RigVMModel/RigVMController.cpp b/Engine/Source/Developer/RigVMDeveloper/Private/RigVMModel/RigVMController.cpp
--- a/Engine/Source/Developer/RigVMDeveloper/Private/RigVMModel/RigVMController.cpp
+++ b/Engine/Source/Developer/RigVMDeveloper/Private/RigVMModel/RigVMController.cpp
@@ -1664,6 +1664,8 @@ bool URigVMController::UnresolveTemplateNodes(const TArray<URigVMTemplateNode*>&
 			continue;
 		}
 
+		EjectAllInjectedNodes(Node, bSetupUndoRedo);
+
 		int32 PermutationIndex = INDEX_NONE;
 		FRigVMTemplate::FTypeMap Types;
 		Node->GetTemplate()->FullyResolve(Types, PermutationIndex);
@@ -1912,7 +1914,9 @@ URigVMNode* URigVMController::UpgradeNode(URigVMNode* InNode, bool bSetupUndoRed
 		return InNode; 
 	}
 
-	TMap<FString, FPinState> PinStates = GetPinStates(InNode);
+	TMap<FString, FString> RedirectedPinPaths;
+	TMap<FString, FPinState> PinStates = GetPinStates(InNode, true);
+	EjectAllInjectedNodes(InNode, bSetupUndoRedo);
 
 	const FString NodeName = InNode->GetName();
 	const FVector2D NodePosition = InNode->GetPosition();
@@ -1967,10 +1971,27 @@ URigVMNode* URigVMController::UpgradeNode(URigVMNode* InNode, bool bSetupUndoRed
 			const FString DefaultValue = UpgradeInfo.GetDefaultValueForPin(Pin->GetFName());
 			if(!DefaultValue.IsEmpty())
 			{
-				SetPinDefaultValue(Pin, DefaultValue, true, false, false);
+				SetPinDefaultValue(Pin, DefaultValue, true, bSetupUndoRedo, false);
 			}
 		}
 
+		// redirect pin state paths
+		for(TPair<FString, FPinState>& PinState : PinStates)
+		{
+			for(int32 TrueFalse = 0; TrueFalse < 2; TrueFalse++)
+			{
+				const FString RemappedInputPath = UpgradeInfo.RemapPin(PinState.Key, TrueFalse == 0, false);
+				if(RemappedInputPath != PinState.Key)
+				{
+					if(!RedirectedPinPaths.Contains(PinState.Key))
+					{
+						RedirectedPinPaths.Add(PinState.Key, RemappedInputPath);
+					}
+				}
+			}
+		}
+
+
 		UpgradedNode = NewUnitNode;
 	}
 	else
@@ -1986,7 +2007,7 @@ URigVMNode* URigVMController::UpgradeNode(URigVMNode* InNode, bool bSetupUndoRed
 	{
 		PinState.Value.DefaultValue.Reset();
 	}
-	ApplyPinStates(UpgradedNode, PinStates);
+	ApplyPinStates(UpgradedNode, PinStates, RedirectedPinPaths, bSetupUndoRedo);
 
 	if(bSetupUndoRedo)
 	{
@@ -2730,7 +2751,6 @@ URigVMInjectionInfo* URigVMController::InjectNodeIntoPin(URigVMPin* InPin, bool
 		ActionStack->BeginAction(Action);
 	}
 	
-	
 	URigVMInjectionInfo* InjectionInfo = NewObject<URigVMInjectionInfo>(InPin);
 	{
 		Notify(ERigVMGraphNotifType::NodeRemoved, NodeToInject);
@@ -2925,6 +2945,49 @@ URigVMNode* URigVMController::EjectNodeFromPin(URigVMPin* InPin, bool bSetupUndo
 	return NodeToEject;
 }
 
+bool URigVMController::EjectAllInjectedNodes(URigVMNode* InNode, bool bSetupUndoRedo, bool bPrintPythonCommands)
+{
+	if(!IsValidNodeForGraph(InNode))
+	{
+		return false;
+	}
+
+	bool bHasAnyInjectedNode = false;
+	for(URigVMPin* Pin : InNode->GetPins())
+	{
+		bHasAnyInjectedNode = bHasAnyInjectedNode || Pin->HasInjectedNodes();
+	}
+
+	if(!bHasAnyInjectedNode)
+	{
+		return false;
+	}
+
+	FRigVMBaseAction EjectAllInjectedNodesAction;
+	if (bSetupUndoRedo)
+	{
+		ActionStack->BeginAction(EjectAllInjectedNodesAction);
+	}
+
+	for(URigVMPin* Pin : InNode->GetPins())
+	{
+		if(Pin->HasInjectedNodes())
+		{
+			if(!EjectNodeFromPin(Pin, bSetupUndoRedo, bPrintPythonCommands))
+			{
+				return false;
+			}
+		}
+	}
+
+	if(bSetupUndoRedo)
+	{
+		ActionStack->EndAction(EjectAllInjectedNodesAction);
+	}
+
+	return true;
+}
+
 
 bool URigVMController::Undo()
 {
@@ -13882,7 +13945,7 @@ TMap<FString, FString> URigVMController::GetRedirectedPinPaths(URigVMNode* InNod
 	return RedirectedPinPaths;
 }
 
-URigVMController::FPinState URigVMController::GetPinState(URigVMPin* InPin) const
+URigVMController::FPinState URigVMController::GetPinState(URigVMPin* InPin, bool bStoreWeakInjectionInfos) const
 {
 	FPinState State;
 	State.Direction = InPin->GetDirection();
@@ -13891,10 +13954,20 @@ URigVMController::FPinState URigVMController::GetPinState(URigVMPin* InPin) cons
 	State.DefaultValue = InPin->GetDefaultValue();
 	State.bIsExpanded = InPin->IsExpanded();
 	State.InjectionInfos = InPin->GetInjectedNodes();
+
+	if(bStoreWeakInjectionInfos)
+	{
+		for(URigVMInjectionInfo* InjectionInfo : State.InjectionInfos)
+		{
+			State.WeakInjectionInfos.Add(InjectionInfo->GetWeakInfo());
+		}
+		State.InjectionInfos.Reset();
+	}
+	
 	return State;
 }
 
-TMap<FString, URigVMController::FPinState> URigVMController::GetPinStates(URigVMNode* InNode) const
+TMap<FString, URigVMController::FPinState> URigVMController::GetPinStates(URigVMNode* InNode, bool bStoreWeakInjectionInfos) const
 {
 	TMap<FString, FPinState> PinStates;
 
@@ -13914,14 +13987,14 @@ TMap<FString, URigVMController::FPinState> URigVMController::GetPinStates(URigVM
 		{
 			EnsurePinValidity(Pin, true);
 		}
-		FPinState State = GetPinState(Pin);
+		FPinState State = GetPinState(Pin, bStoreWeakInjectionInfos);
 		PinStates.Add(PinPath, State);
 	}
 
 	return PinStates;
 }
 
-void URigVMController::ApplyPinState(URigVMPin* InPin, const FPinState& InPinState)
+void URigVMController::ApplyPinState(URigVMPin* InPin, const FPinState& InPinState, bool bSetupUndoRedo)
 {
 	for (URigVMInjectionInfo* InjectionInfo : InPinState.InjectionInfos)
 	{
@@ -13931,15 +14004,42 @@ void URigVMController::ApplyPinState(URigVMPin* InPin, const FPinState& InPinSta
 		InPin->InjectionInfos.Add(InjectionInfo);
 	}
 
+	// alternatively if the injection infos are not provided as strong pointers
+	// we can fall back onto the weak ptr information and try again
+	if(InPinState.InjectionInfos.IsEmpty())
+	{
+		for (const URigVMInjectionInfo::FWeakInfo& InjectionInfo : InPinState.WeakInjectionInfos)
+		{
+			if(const URigVMNode* FormerlyInjectedNode = InjectionInfo.Node.Get())
+			{
+				if(InjectionInfo.bInjectedAsInput)
+				{
+					const FString OutputPinPath = URigVMPin::JoinPinPath(FormerlyInjectedNode->GetNodePath(), InjectionInfo.OutputPinName.ToString()); 
+					AddLink(OutputPinPath, InPin->GetPinPath(), bSetupUndoRedo, false);
+				}
+				else
+				{
+					const FString InputPinPath = URigVMPin::JoinPinPath(FormerlyInjectedNode->GetNodePath(), InjectionInfo.InputPinName.ToString()); 
+					AddLink(InPin->GetPinPath(), InputPinPath, bSetupUndoRedo, false);
+				}
+
+				if(InPin->IsRootPin())
+				{
+					InjectNodeIntoPin(InPin, InjectionInfo.bInjectedAsInput, InjectionInfo.InputPinName, InjectionInfo.OutputPinName, bSetupUndoRedo);
+				}
+			}
+		}
+	}
+
 	if (!InPinState.DefaultValue.IsEmpty())
 	{
-		SetPinDefaultValue(InPin, InPinState.DefaultValue, true, false, false);
+		SetPinDefaultValue(InPin, InPinState.DefaultValue, true, bSetupUndoRedo, false);
 	}
 
-	SetPinExpansion(InPin, InPinState.bIsExpanded, false);
+	SetPinExpansion(InPin, InPinState.bIsExpanded, bSetupUndoRedo);
 }
 
-void URigVMController::ApplyPinStates(URigVMNode* InNode, const TMap<FString, URigVMController::FPinState>& InPinStates, const TMap<FString, FString>& InRedirectedPinPaths)
+void URigVMController::ApplyPinStates(URigVMNode* InNode, const TMap<FString, URigVMController::FPinState>& InPinStates, const TMap<FString, FString>& InRedirectedPinPaths, bool bSetupUndoRedo)
 {
 	FRigVMControllerCompileBracketScope CompileBracketScope(this);
 	for (const TPair<FString, FPinState>& PinStatePair : InPinStates)
@@ -13954,7 +14054,7 @@ void URigVMController::ApplyPinStates(URigVMNode* InNode, const TMap<FString, UR
 
 		if (URigVMPin* Pin = InNode->FindPin(PinPath))
 		{
-			ApplyPinState(Pin, PinState);
+			ApplyPinState(Pin, PinState, bSetupUndoRedo);
 		}
 		else
 		{
@@ -15410,13 +15510,13 @@ void URigVMController::SanitizeName(FString& InOutName, bool bAllowPeriod, bool
 	}
 }
 
-TArray<TPair<FString, FString>> URigVMController::GetLinkedPinPaths(URigVMNode* InNode)
+TArray<TPair<FString, FString>> URigVMController::GetLinkedPinPaths(URigVMNode* InNode, bool bIncludeInjectionNodes)
 {
 	const TArray<URigVMNode*> Nodes = {InNode};
-	return GetLinkedPinPaths(Nodes);
+	return GetLinkedPinPaths(Nodes, bIncludeInjectionNodes);
 }
 
-TArray<TPair<FString, FString>> URigVMController::GetLinkedPinPaths(const TArray<URigVMNode*>& InNodes)
+TArray<TPair<FString, FString>> URigVMController::GetLinkedPinPaths(const TArray<URigVMNode*>& InNodes, bool bIncludeInjectionNodes)
 {
 	TArray<TPair<FString, FString>> LinkedPaths;
 	for(URigVMNode* Node : InNodes)
@@ -15424,6 +15524,14 @@ TArray<TPair<FString, FString>> URigVMController::GetLinkedPinPaths(const TArray
 		TArray<URigVMLink*> Links = Node->GetLinks();
 		for(URigVMLink* Link : Links)
 		{
+			if(!bIncludeInjectionNodes)
+			{
+				if(Link->GetSourcePin()->GetNode()->IsInjected() ||
+					Link->GetTargetPin()->GetNode()->IsInjected())
+				{
+					continue;
+				}
+			}
 			TPair<FString, FString> LinkedPath(Link->GetSourcePin()->GetPinPath(), Link->GetTargetPin()->GetPinPath());
 			LinkedPaths.AddUnique(LinkedPath);
 		}
diff --git a/Engine/Source/Developer/RigVMDeveloper/Private/RigVMModel/RigVMPin.cpp b/Engine/Source/Developer/RigVMDeveloper/Private/RigVMModel/RigVMPin.cpp
--- a/Engine/Source/Developer/RigVMDeveloper/Private/RigVMModel/RigVMPin.cpp
+++ b/Engine/Source/Developer/RigVMDeveloper/Private/RigVMModel/RigVMPin.cpp
@@ -35,6 +35,16 @@ URigVMPin* URigVMInjectionInfo::GetPin() const
 	return CastChecked<URigVMPin>(GetOuter());
 }
 
+URigVMInjectionInfo::FWeakInfo URigVMInjectionInfo::GetWeakInfo() const
+{
+	FWeakInfo Info;
+	Info.bInjectedAsInput = bInjectedAsInput;
+	Info.Node = Node;
+	Info.InputPinName = InputPin != nullptr ? InputPin->GetFName() : NAME_None;
+	Info.OutputPinName = OutputPin != nullptr ? OutputPin->GetFName() : NAME_None;
+	return Info;
+}
+
 const URigVMPin::FPinOverrideMap URigVMPin::EmptyPinOverrideMap;
 const URigVMPin::FPinOverride URigVMPin::EmptyPinOverride = URigVMPin::FPinOverride(FRigVMASTProxy(), EmptyPinOverrideMap);
 const FString URigVMPin::OrphanPinPrefix = TEXT("Orphan::");
diff --git a/Engine/Source/Developer/RigVMDeveloper/Public/RigVMModel/RigVMController.h b/Engine/Source/Developer/RigVMDeveloper/Public/RigVMModel/RigVMController.h
--- a/Engine/Source/Developer/RigVMDeveloper/Public/RigVMModel/RigVMController.h
+++ b/Engine/Source/Developer/RigVMDeveloper/Public/RigVMModel/RigVMController.h
@@ -905,6 +905,7 @@ private:
 	URigVMInjectionInfo* InjectNodeIntoPin(const FString& InPinPath, bool bAsInput, const FName& InInputPinName, const FName& InOutputPinName, bool bSetupUndoRedo = true);
 	URigVMInjectionInfo* InjectNodeIntoPin(URigVMPin* InPin, bool bAsInput, const FName& InInputPinName, const FName& InOutputPinName, bool bSetupUndoRedo = true);
 	URigVMNode* EjectNodeFromPin(URigVMPin* InPin, bool bSetupUndoRedo = true, bool bPrintPythonCommands = false);
+	bool EjectAllInjectedNodes(URigVMNode* InNode, bool bSetupUndoRedo = true, bool bPrintPythonCommands = false);
 
 	// try to reconnect source and target pins after a node deletion
 	void RelinkSourceAndTargetPins(URigVMNode* RigNode, bool bSetupUndoRedo = true);
@@ -941,13 +942,14 @@ private:
 		FString DefaultValue;
 		bool bIsExpanded;
 		TArray<URigVMInjectionInfo*> InjectionInfos;
+		TArray<URigVMInjectionInfo::FWeakInfo> WeakInjectionInfos;
 	};
 
 	TMap<FString, FString> GetRedirectedPinPaths(URigVMNode* InNode) const;
-	FPinState GetPinState(URigVMPin* InPin) const;
-	TMap<FString, FPinState> GetPinStates(URigVMNode* InNode) const;
-	void ApplyPinState(URigVMPin* InPin, const FPinState& InPinState);
-	void ApplyPinStates(URigVMNode* InNode, const TMap<FString, FPinState>& InPinStates, const TMap<FString, FString>& InRedirectedPinPaths = TMap<FString, FString>());
+	FPinState GetPinState(URigVMPin* InPin, bool bStoreWeakInjectionInfos = false) const;
+	TMap<FString, FPinState> GetPinStates(URigVMNode* InNode, bool bStoreWeakInjectionInfos = false) const;
+	void ApplyPinState(URigVMPin* InPin, const FPinState& InPinState, bool bSetupUndoRedo = false);
+	void ApplyPinStates(URigVMNode* InNode, const TMap<FString, FPinState>& InPinStates, const TMap<FString, FString>& InRedirectedPinPaths = TMap<FString, FString>(), bool bSetupUndoRedo = false);
 
 
 	static FLinearColor GetColorFromMetadata(const FString& InMetadata);
@@ -995,8 +997,8 @@ public:
 	static FString GetSanitizedPinName(const FString& InName);
 	static FString GetSanitizedPinPath(const FString& InName);
 	static void SanitizeName(FString& InOutName, bool bAllowPeriod, bool bAllowSpace);
-	static TArray<TPair<FString, FString>> GetLinkedPinPaths(URigVMNode* InNode);
-	static TArray<TPair<FString, FString>> GetLinkedPinPaths(const TArray<URigVMNode*>& InNodes);
+	static TArray<TPair<FString, FString>> GetLinkedPinPaths(URigVMNode* InNode, bool bIncludeInjectionNodes = false);
+	static TArray<TPair<FString, FString>> GetLinkedPinPaths(const TArray<URigVMNode*>& InNodes, bool bIncludeInjectionNodes = false);
 
 private: 
 	UPROPERTY(transient)
diff --git a/Engine/Source/Developer/RigVMDeveloper/Public/RigVMModel/RigVMPin.h b/Engine/Source/Developer/RigVMDeveloper/Public/RigVMModel/RigVMPin.h
--- a/Engine/Source/Developer/RigVMDeveloper/Public/RigVMModel/RigVMPin.h
+++ b/Engine/Source/Developer/RigVMDeveloper/Public/RigVMModel/RigVMPin.h
@@ -57,6 +57,16 @@ public:
 	// Returns the pin of this injected node.
 	UFUNCTION(BlueprintCallable, Category = RigVMInjectionInfo)
 	URigVMPin* GetPin() const;
+
+	struct FWeakInfo
+	{
+		TWeakObjectPtr<URigVMNode> Node;
+		bool bInjectedAsInput;
+		FName InputPinName;
+		FName OutputPinName;
+	};
+
+	FWeakInfo GetWeakInfo() const;
 };
 
 /**
diff --git a/Engine/Source/Editor/GraphEditor/Private/SGraphPanel.cpp b/Engine/Source/Editor/GraphEditor/Private/SGraphPanel.cpp
--- a/Engine/Source/Editor/GraphEditor/Private/SGraphPanel.cpp
+++ b/Engine/Source/Editor/GraphEditor/Private/SGraphPanel.cpp
@@ -1801,8 +1801,11 @@ void SGraphPanel::OnGraphChanged(const FEdGraphEditAction& EditAction)
 				if (NodePtr.IsValid())
 				{
 					UEdGraphNode* Node = NodePtr.Get();
-					Parent->RemoveNode(Node);
-					Parent->AddNode(Node, bForceUserAdded ? WasUserAdded : NotUserAdded);
+					if(IsValid(Node))
+					{
+						Parent->RemoveNode(Node);
+						Parent->AddNode(Node, bForceUserAdded ? WasUserAdded : NotUserAdded);
+					}
 				}
 				return EActiveTimerReturnType::Stop;
 			};
