Control Rig: Fix RigVM.Graph.Variables unit test failing

diff --git a/Engine/Source/Developer/RigVMDeveloper/Private/RigVMModel/RigVMController.cpp b/Engine/Source/Developer/RigVMDeveloper/Private/RigVMModel/RigVMController.cpp
--- a/Engine/Source/Developer/RigVMDeveloper/Private/RigVMModel/RigVMController.cpp
+++ b/Engine/Source/Developer/RigVMDeveloper/Private/RigVMModel/RigVMController.cpp
@@ -1089,11 +1089,7 @@ void URigVMController::RefreshVariableNode(const FName& InNodeName, const FName&
 
 			if (InVariableName.IsValid() && VariablePin->DefaultValue != InVariableName.ToString())
 			{
-				if (bSetupUndoRedo)
-				{
-					VariablePin->Modify();
-				}
-				VariablePin->DefaultValue = InVariableName.ToString();
+				SetPinDefaultValue(VariablePin, InVariableName.ToString(), false, bSetupUndoRedo, false);
 				Notify(ERigVMGraphNotifType::PinDefaultValueChanged, VariablePin);
 				Notify(ERigVMGraphNotifType::VariableRenamed, VariableNode);
 			}
diff --git a/Engine/Source/Developer/RigVMDeveloper/Public/RigVMModel/RigVMController.h b/Engine/Source/Developer/RigVMDeveloper/Public/RigVMModel/RigVMController.h
--- a/Engine/Source/Developer/RigVMDeveloper/Public/RigVMModel/RigVMController.h
+++ b/Engine/Source/Developer/RigVMDeveloper/Public/RigVMModel/RigVMController.h
@@ -199,17 +199,13 @@ public:
 	void RefreshVariableNode(const FName& InNodeName, const FName& InVariableName, const FString& InCPPType, UObject* InCPPTypeObject, bool bSetupUndoRedo, bool bSetupOrphanPins = true);
 
 	// Removes all nodes related to a given variable
-	UFUNCTION(BlueprintCallable, Category = RigVMController)
 	void OnExternalVariableRemoved(const FName& InVarName, bool bSetupUndoRedo);
 
 	// Renames the variable name in all relevant nodes
-	UFUNCTION(BlueprintCallable, Category = RigVMController)
 	bool OnExternalVariableRenamed(const FName& InOldVarName, const FName& InNewVarName, bool bSetupUndoRedo);
 
 	// Changes the data type of all nodes matching a given variable name
-	UFUNCTION(BlueprintCallable, Category = RigVMController)
 	void OnExternalVariableTypeChanged(const FName& InVarName, const FString& InCPPType, UObject* InCPPTypeObject, bool bSetupUndoRedo);
-	UFUNCTION(BlueprintCallable, Category = RigVMController)
 	void OnExternalVariableTypeChangedFromObjectPath(const FName& InVarName, const FString& InCPPType, const FString& InCPPTypeObjectPath, bool bSetupUndoRedo);
 
 	// Refreshes the variable node with the new data
