RigVM: Fix constraint boundaries for code generator

diff --git a/Engine/Source/Developer/RigVMDeveloper/Private/RigVMCompiler/RigVMCodeGenerator.cpp b/Engine/Source/Developer/RigVMDeveloper/Private/RigVMCompiler/RigVMCodeGenerator.cpp
--- a/Engine/Source/Developer/RigVMDeveloper/Private/RigVMCompiler/RigVMCodeGenerator.cpp
+++ b/Engine/Source/Developer/RigVMDeveloper/Private/RigVMCompiler/RigVMCodeGenerator.cpp
@@ -1324,7 +1324,7 @@ void FRigVMCodeGenerator::ParseOperationGroups()
 			auto IsWithinConstraint = [&bMoveTowardsHead, &MiddleHead, &MiddleTail](const TTuple<int32,int32>& Constraint) -> bool
 			{
 				int32& Middle = bMoveTowardsHead ? MiddleHead : MiddleTail;
-				return FMath::IsWithin(Middle, Constraint.Get<0>(), Constraint.Get<1>());
+				return FMath::IsWithinInclusive(Middle, Constraint.Get<0>(), Constraint.Get<1>());
 			};
 
 			// while the middle instruction is within a constraint - move it  
@@ -1941,7 +1941,7 @@ bool FRigVMCodeGenerator::IsOperationPartOfGroup(int32 InOperationIndex, int32 I
 		}
 	}
 
-	return FMath::IsWithin(InOperationIndex, Group.First, Group.Last);;
+	return FMath::IsWithinInclusive(InOperationIndex, Group.First, Group.Last);;
 }
 
 bool FRigVMCodeGenerator::IsPropertyPartOfGroup(int32 InPropertyIndex, int32 InGroupIndex) const
