Control Rig: Fix local variables not working properly

diff --git a/Engine/Source/Developer/RigVMDeveloper/Private/RigVMCompiler/RigVMCompiler.cpp b/Engine/Source/Developer/RigVMDeveloper/Private/RigVMCompiler/RigVMCompiler.cpp
--- a/Engine/Source/Developer/RigVMDeveloper/Private/RigVMCompiler/RigVMCompiler.cpp
+++ b/Engine/Source/Developer/RigVMDeveloper/Private/RigVMCompiler/RigVMCompiler.cpp
@@ -1920,7 +1920,7 @@ void URigVMCompiler::InitializeLocalVariables(const FRigVMExprAST* InExpr, FRigV
 	}
 }
 
-FString URigVMCompiler::GetPinHash(const URigVMPin* InPin, const FRigVMVarExprAST* InVarExpr, bool bIsDebugValue, const FRigVMASTProxy& InPinProxy)
+FString URigVMCompiler::GetPinHashImpl(const URigVMPin* InPin, const FRigVMVarExprAST* InVarExpr, bool bIsDebugValue, const FRigVMASTProxy& InPinProxy)
 {
 	FString Prefix = bIsDebugValue ? TEXT("DebugWatch:") : TEXT("");
 	FString Suffix;
@@ -2011,7 +2011,8 @@ FString URigVMCompiler::GetPinHash(const URigVMPin* InPin, const FRigVMVarExprAS
 							if(URigVMLibraryNode* LibraryNode = ParentProxy.GetSubject<URigVMLibraryNode>())
 							{
 								// Local variables for non-root graphs are in the format "LocalVariable::PathToGraph|VariableName"
-								return FString::Printf(TEXT("%sLocalVariable::%s|%s%s"), *Prefix, *LibraryNode->GetNodePath(true), *VariableName.ToString(), *Suffix);
+								const FString GraphPath = ParentProxy.GetCallstack().GetCallPath(true);
+								return FString::Printf(TEXT("%sLocalVariable::%s|%s%s"), *Prefix, *GraphPath, *VariableName.ToString(), *Suffix);
 							}
 						}
 
@@ -2103,10 +2104,16 @@ FString URigVMCompiler::GetPinHash(const URigVMPin* InPin, const FRigVMVarExprAS
 	}
 
 	FString PinPath = InPin->GetPinPath(bUseFullNodePath);
-	ensureMsgf(!PinPath.StartsWith(TEXT("FunctionLibrary::")), TEXT("A library path should never be part of a pin hash %s."), *PinPath);
 	return FString::Printf(TEXT("%s%s%s"), *Prefix, *PinPath, *Suffix);
 }
 
+FString URigVMCompiler::GetPinHash(const URigVMPin* InPin, const FRigVMVarExprAST* InVarExpr, bool bIsDebugValue, const FRigVMASTProxy& InPinProxy)
+{
+	const FString Hash = GetPinHashImpl(InPin, InVarExpr, bIsDebugValue, InPinProxy);
+	ensureMsgf(!Hash.Contains(TEXT("FunctionLibrary::")), TEXT("A library path should never be part of a pin hash %s."), *Hash);
+	return Hash;
+}
+
 const FRigVMVarExprAST* URigVMCompiler::GetSourceVarExpr(const FRigVMExprAST* InExpr)
 {
 	if(InExpr)
diff --git a/Engine/Source/Developer/RigVMDeveloper/Public/RigVMCompiler/RigVMCompiler.h b/Engine/Source/Developer/RigVMDeveloper/Public/RigVMCompiler/RigVMCompiler.h
--- a/Engine/Source/Developer/RigVMDeveloper/Public/RigVMCompiler/RigVMCompiler.h
+++ b/Engine/Source/Developer/RigVMDeveloper/Public/RigVMCompiler/RigVMCompiler.h
@@ -159,6 +159,8 @@ private:
 	TArray<URigVMPin*> GetLinkedPins(URigVMPin* InPin, bool bInputs = true, bool bOutputs = true, bool bRecursive = true);
 	uint16 GetElementSizeFromCPPType(const FString& InCPPType, UScriptStruct* InScriptStruct);
 
+	static FString GetPinHashImpl(const URigVMPin* InPin, const FRigVMVarExprAST* InVarExpr, bool bIsDebugValue = false, const FRigVMASTProxy& InPinProxy = FRigVMASTProxy());
+
 	void TraverseExpression(const FRigVMExprAST* InExpr, FRigVMCompilerWorkData& WorkData);
 	void TraverseChildren(const FRigVMExprAST* InExpr, FRigVMCompilerWorkData& WorkData);
 	void TraverseBlock(const FRigVMBlockExprAST* InExpr, FRigVMCompilerWorkData& WorkData);
