RigVM: Provide subject to user workflow callbacks

diff --git a/Engine/Source/Developer/RigVMDeveloper/Private/RigVMModel/Nodes/RigVMUnitNode.cpp b/Engine/Source/Developer/RigVMDeveloper/Private/RigVMModel/Nodes/RigVMUnitNode.cpp
--- a/Engine/Source/Developer/RigVMDeveloper/Private/RigVMModel/Nodes/RigVMUnitNode.cpp
+++ b/Engine/Source/Developer/RigVMDeveloper/Private/RigVMModel/Nodes/RigVMUnitNode.cpp
@@ -110,9 +110,14 @@ FString URigVMUnitNode::GetDeprecatedMetadata() const
 	return FString();
 }
 
-TArray<FRigVMUserWorkflow> URigVMUnitNode::GetSupportedWorkflows(ERigVMUserWorkflowType InType) const
+TArray<FRigVMUserWorkflow> URigVMUnitNode::GetSupportedWorkflows(ERigVMUserWorkflowType InType, const UObject* InSubject) const
 {
-	TArray<FRigVMUserWorkflow> Workflows = Super::GetSupportedWorkflows(InType);
+	TArray<FRigVMUserWorkflow> Workflows = Super::GetSupportedWorkflows(InType, InSubject);
+
+	if(InSubject == nullptr)
+	{
+		InSubject = this;
+	}
 
 	if(UScriptStruct* Struct = GetScriptStruct())
 	{
@@ -120,8 +125,8 @@ TArray<FRigVMUserWorkflow> URigVMUnitNode::GetSupportedWorkflows(ERigVMUserWorkf
 
 		const TSharedPtr<FStructOnScope> StructOnScope = ConstructStructInstance();
 		const FRigVMStruct* StructMemory = (const FRigVMStruct*)StructOnScope->GetStructMemory();
-		TArray<FRigVMUserWorkflow> StructWorkflows = StructMemory->GetWorkflows(InType);
-		StructWorkflows.Append(URigVMUserWorkflowRegistry::Get()->GetWorkflows(InType, Struct, this));
+		TArray<FRigVMUserWorkflow> StructWorkflows = StructMemory->GetWorkflows(InType, InSubject);
+		StructWorkflows.Append(URigVMUserWorkflowRegistry::Get()->GetWorkflows(InType, Struct, InSubject));
 		Swap(Workflows, StructWorkflows);
 		Workflows.Append(StructWorkflows);
 	}
diff --git a/Engine/Source/Developer/RigVMDeveloper/Private/RigVMModel/RigVMController.cpp b/Engine/Source/Developer/RigVMDeveloper/Private/RigVMModel/RigVMController.cpp
--- a/Engine/Source/Developer/RigVMDeveloper/Private/RigVMModel/RigVMController.cpp
+++ b/Engine/Source/Developer/RigVMDeveloper/Private/RigVMModel/RigVMController.cpp
@@ -7165,7 +7165,7 @@ bool URigVMController::SetPinDefaultValue(URigVMPin* InPin, const FString& InDef
 
 	if(bSetupUndoRedo)
 	{
-		TArray<FRigVMUserWorkflow> WorkflowsOnPinChange = InPin->GetNode()->GetSupportedWorkflows(ERigVMUserWorkflowType::OnPinDefaultChanged);
+		TArray<FRigVMUserWorkflow> WorkflowsOnPinChange = InPin->GetNode()->GetSupportedWorkflows(ERigVMUserWorkflowType::OnPinDefaultChanged, InPin);
 		for(const FRigVMUserWorkflow& Workflow : WorkflowsOnPinChange)
 		{
 			const URigVMUserWorkflowOptions* Options = MakeOptionsForWorkflow(InPin, Workflow);
diff --git a/Engine/Source/Developer/RigVMDeveloper/Private/RigVMModel/RigVMNode.cpp b/Engine/Source/Developer/RigVMDeveloper/Private/RigVMModel/RigVMNode.cpp
--- a/Engine/Source/Developer/RigVMDeveloper/Private/RigVMModel/RigVMNode.cpp
+++ b/Engine/Source/Developer/RigVMDeveloper/Private/RigVMModel/RigVMNode.cpp
@@ -530,9 +530,13 @@ double URigVMNode::GetInstructionMicroSeconds(URigVM* InVM, const FRigVMASTProxy
 	return -1.0;
 }
 
-TArray<FRigVMUserWorkflow> URigVMNode::GetSupportedWorkflows(ERigVMUserWorkflowType InType) const
+TArray<FRigVMUserWorkflow> URigVMNode::GetSupportedWorkflows(ERigVMUserWorkflowType InType, const UObject* InSubject) const
 {
-	return URigVMUserWorkflowRegistry::Get()->GetWorkflows(InType, nullptr, this);
+	if(InSubject == nullptr)
+	{
+		InSubject = this;
+	}
+	return URigVMUserWorkflowRegistry::Get()->GetWorkflows(InType, nullptr, InSubject);
 }
 
 #if WITH_EDITOR
diff --git a/Engine/Source/Developer/RigVMDeveloper/Private/RigVMUserWorkflowRegistry.cpp b/Engine/Source/Developer/RigVMDeveloper/Private/RigVMUserWorkflowRegistry.cpp
--- a/Engine/Source/Developer/RigVMDeveloper/Private/RigVMUserWorkflowRegistry.cpp
+++ b/Engine/Source/Developer/RigVMDeveloper/Private/RigVMUserWorkflowRegistry.cpp
@@ -22,7 +22,7 @@ void URigVMUserWorkflowRegistry::UnregisterProvider(int32 InHandle)
 	});
 }
 
-TArray<FRigVMUserWorkflow> URigVMUserWorkflowRegistry::GetWorkflows(ERigVMUserWorkflowType InType, const UScriptStruct* InStruct, const URigVMNode* InNode) const
+TArray<FRigVMUserWorkflow> URigVMUserWorkflowRegistry::GetWorkflows(ERigVMUserWorkflowType InType, const UScriptStruct* InStruct, const UObject* InSubject) const
 {
 	TArray<FRigVMUserWorkflow> Workflows;
 
@@ -35,9 +35,9 @@ TArray<FRigVMUserWorkflow> URigVMUserWorkflowRegistry::GetWorkflows(ERigVMUserWo
 
 	for(const TTuple<int32, const UScriptStruct*, FRigVMUserWorkflowProvider>& Provider : Providers)
 	{
-		if(Provider.Get<1>() == InStruct && Provider.Get<1>() != nullptr)
+		if(Provider.Get<1>() == InStruct || Provider.Get<1>() == nullptr || InStruct == nullptr)
 		{
-			Workflows.Append(Provider.Get<2>().Execute(InNode));
+			Workflows.Append(Provider.Get<2>().Execute(InSubject));
 		}
 	}
 
diff --git a/Engine/Source/Developer/RigVMDeveloper/Public/RigVMModel/Nodes/RigVMUnitNode.h b/Engine/Source/Developer/RigVMDeveloper/Public/RigVMModel/Nodes/RigVMUnitNode.h
--- a/Engine/Source/Developer/RigVMDeveloper/Public/RigVMModel/Nodes/RigVMUnitNode.h
+++ b/Engine/Source/Developer/RigVMDeveloper/Public/RigVMModel/Nodes/RigVMUnitNode.h
@@ -27,7 +27,7 @@ public:
 	virtual FName GetEventName() const override;
 	virtual bool IsLoopNode() const override;
 	virtual bool CanBeUpgraded() const override { return GetUpgradeInfo().IsValid(); }
-	virtual TArray<FRigVMUserWorkflow> GetSupportedWorkflows(ERigVMUserWorkflowType InType = ERigVMUserWorkflowType::All) const override;
+	virtual TArray<FRigVMUserWorkflow> GetSupportedWorkflows(ERigVMUserWorkflowType InType, const UObject* InSubject) const override;
 	virtual bool IsAggregate() const override;
 	virtual URigVMPin* GetFirstAggregatePin() const override;
 	virtual URigVMPin* GetSecondAggregatePin() const override;
diff --git a/Engine/Source/Developer/RigVMDeveloper/Public/RigVMModel/RigVMNode.h b/Engine/Source/Developer/RigVMDeveloper/Public/RigVMModel/RigVMNode.h
--- a/Engine/Source/Developer/RigVMDeveloper/Public/RigVMModel/RigVMNode.h
+++ b/Engine/Source/Developer/RigVMDeveloper/Public/RigVMModel/RigVMNode.h
@@ -211,7 +211,7 @@ public:
 
 	// returns all supported workflows of the node
 	UFUNCTION(BlueprintCallable, Category = RigVMNode)
-	virtual TArray<FRigVMUserWorkflow> GetSupportedWorkflows(ERigVMUserWorkflowType InType = ERigVMUserWorkflowType::All) const;
+	virtual TArray<FRigVMUserWorkflow> GetSupportedWorkflows(ERigVMUserWorkflowType InType, const UObject* InSubject) const;
 
 	UFUNCTION(BlueprintCallable, Category = RigVMNode)
 	bool HasBreakpoint() const { return bHasBreakpoint; }
diff --git a/Engine/Source/Developer/RigVMDeveloper/Public/RigVMUserWorkflowRegistry.h b/Engine/Source/Developer/RigVMDeveloper/Public/RigVMUserWorkflowRegistry.h
--- a/Engine/Source/Developer/RigVMDeveloper/Public/RigVMUserWorkflowRegistry.h
+++ b/Engine/Source/Developer/RigVMDeveloper/Public/RigVMUserWorkflowRegistry.h
@@ -7,7 +7,7 @@
 #include "RigVMModel/RigVMNode.h"
 #include "RigVMUserWorkflowRegistry.generated.h"
 
-DECLARE_DYNAMIC_DELEGATE_RetVal_OneParam(TArray<FRigVMUserWorkflow>, FRigVMUserWorkflowProvider, const URigVMNode*, InNode);
+DECLARE_DYNAMIC_DELEGATE_RetVal_OneParam(TArray<FRigVMUserWorkflow>, FRigVMUserWorkflowProvider, const UObject*, InSubject);
 
 UCLASS(BlueprintType)
 class RIGVMDEVELOPER_API URigVMUserWorkflowRegistry : public UObject
@@ -26,7 +26,7 @@ public:
 	void UnregisterProvider(int32 InHandle);
 
 	UFUNCTION(BlueprintPure, Category = FRigVMUserWorkflowRegistry)
-	TArray<FRigVMUserWorkflow> GetWorkflows(ERigVMUserWorkflowType InType, const UScriptStruct* InStruct, const URigVMNode* InNode) const;
+	TArray<FRigVMUserWorkflow> GetWorkflows(ERigVMUserWorkflowType InType, const UScriptStruct* InStruct, const UObject* InSubject) const;
 
 private:
 
diff --git a/Engine/Source/Runtime/RigVM/Private/RigVMCore/RigVMStruct.cpp b/Engine/Source/Runtime/RigVM/Private/RigVMCore/RigVMStruct.cpp
--- a/Engine/Source/Runtime/RigVM/Private/RigVMCore/RigVMStruct.cpp
+++ b/Engine/Source/Runtime/RigVM/Private/RigVMCore/RigVMStruct.cpp
@@ -221,9 +221,9 @@ float FRigVMStruct::GetRatioFromIndex(int32 InIndex, int32 InCount)
 	return ((float)FMath::Clamp<int32>(InIndex, 0, InCount - 1)) / ((float)(InCount - 1));
 }
 
-TArray<FRigVMUserWorkflow> FRigVMStruct::GetWorkflows(ERigVMUserWorkflowType InType) const
+TArray<FRigVMUserWorkflow> FRigVMStruct::GetWorkflows(ERigVMUserWorkflowType InType, const UObject* InSubject) const
 {
-	return GetSupportedWorkflows().FilterByPredicate([InType](const FRigVMUserWorkflow& InWorkflow) -> bool
+	return GetSupportedWorkflows(InSubject).FilterByPredicate([InType](const FRigVMUserWorkflow& InWorkflow) -> bool
 	{
 		return uint32(InWorkflow.GetType()) & uint32(InType) &&
 			InWorkflow.IsValid();
diff --git a/Engine/Source/Runtime/RigVM/Private/RigVMCore/RigVMUserWorkflow.cpp b/Engine/Source/Runtime/RigVM/Private/RigVMCore/RigVMUserWorkflow.cpp
--- a/Engine/Source/Runtime/RigVM/Private/RigVMCore/RigVMUserWorkflow.cpp
+++ b/Engine/Source/Runtime/RigVM/Private/RigVMCore/RigVMUserWorkflow.cpp
@@ -27,10 +27,11 @@ bool URigVMUserWorkflowOptions::RequiresDialog(const FProperty* InProperty) cons
 	{
 		return false;
 	}
-
-	static const FName SubjectName = TEXT("Subject");
-	static const FName WorkflowName = TEXT("Workflow");
-	return InProperty->GetFName() != SubjectName && InProperty->GetFName() != WorkflowName;
+	if(InProperty->HasAnyPropertyFlags(CPF_DisableEditOnInstance))
+	{
+		return false;
+	}
+	return true;
 }
 
 void URigVMUserWorkflowOptions::Report(EMessageSeverity::Type InSeverity,  const FString& InMessage) const
diff --git a/Engine/Source/Runtime/RigVM/Public/RigVMCore/RigVMStruct.h b/Engine/Source/Runtime/RigVM/Public/RigVMCore/RigVMStruct.h
--- a/Engine/Source/Runtime/RigVM/Public/RigVMCore/RigVMStruct.h
+++ b/Engine/Source/Runtime/RigVM/Public/RigVMCore/RigVMStruct.h
@@ -223,7 +223,7 @@ public:
 	FORCEINLINE virtual void OnUnitNodeCreated(FRigVMUnitNodeCreatedContext& InContext) const {}
 
 	// user workflow
-	TArray<FRigVMUserWorkflow> GetWorkflows(ERigVMUserWorkflowType InType = ERigVMUserWorkflowType::All) const; 
+	TArray<FRigVMUserWorkflow> GetWorkflows(ERigVMUserWorkflowType InType, const UObject* InSubject) const; 
 
 #if WITH_EDITOR
 	static bool ValidateStruct(UScriptStruct* InStruct, FString* OutErrorMessage);
@@ -300,7 +300,7 @@ protected:
 	static float GetRatioFromIndex(int32 InIndex, int32 InCount);
 	TMap<FName, FString> GetDefaultValues(UScriptStruct* InScriptStruct) const;
 	bool ApplyUpgradeInfo(const FRigVMStructUpgradeInfo& InUpgradeInfo);
-	FORCEINLINE virtual TArray<FRigVMUserWorkflow> GetSupportedWorkflows() const { return TArray<FRigVMUserWorkflow>(); } 
+	FORCEINLINE virtual TArray<FRigVMUserWorkflow> GetSupportedWorkflows(const UObject* InSubject) const { return TArray<FRigVMUserWorkflow>(); } 
 
 	friend struct FRigVMStructUpgradeInfo;
 	friend class FRigVMGraphStructUpgradeInfoTest;
