RigVM: Simplify setup of user workflows

diff --git a/Engine/Plugins/Animation/ControlRig/Source/ControlRig/Private/Units/ControlRigNodeWorkflow.cpp b/Engine/Plugins/Animation/ControlRig/Source/ControlRig/Private/Units/ControlRigNodeWorkflow.cpp
--- a/Engine/Plugins/Animation/ControlRig/Source/ControlRig/Private/Units/ControlRigNodeWorkflow.cpp
+++ b/Engine/Plugins/Animation/ControlRig/Source/ControlRig/Private/Units/ControlRigNodeWorkflow.cpp
@@ -5,6 +5,7 @@
 
 #if WITH_EDITOR
 #include "RigVMModel/RigVMNode.h"
+#include "RigVMModel/RigVMController.h"
 #endif
 
 bool UControlRigWorkflowOptions::EnsureAtLeastOneRigElementSelected() const
@@ -33,7 +34,7 @@ TArray<FRigVMUserWorkflow> UControlRigTransformWorkflowOptions::ProvideWorkflows
 					TEXT("Set from hierarchy"),
 					TEXT("Sets the pin to match the global transform of the selected element in the hierarchy"),
 					ERigVMUserWorkflowType::PinContext,
-					FRigVMWorkflowGetActionsDelegate::CreateStatic(&UControlRigTransformWorkflowOptions::ProvideTransformWorkflow),
+					FRigVMPerformUserWorkflowDelegate::CreateStatic(&UControlRigTransformWorkflowOptions::PerformTransformWorkflow),
 					StaticClass()
 				);
 			}
@@ -44,11 +45,12 @@ TArray<FRigVMUserWorkflow> UControlRigTransformWorkflowOptions::ProvideWorkflows
 }
 
 #if WITH_EDITOR
-TArray<FRigVMUserWorkflowAction> UControlRigTransformWorkflowOptions::ProvideTransformWorkflow(
-	const URigVMUserWorkflowOptions* InOptions)
-{
-	TArray<FRigVMUserWorkflowAction> Actions;
 
+bool UControlRigTransformWorkflowOptions::PerformTransformWorkflow(const URigVMUserWorkflowOptions* InOptions,
+	UObject* InController)
+{
+	URigVMController* Controller = CastChecked<URigVMController>(InController);
+	
 	if(const UControlRigTransformWorkflowOptions* Options = Cast<UControlRigTransformWorkflowOptions>(InOptions))
 	{
 		if(Options->EnsureAtLeastOneRigElementSelected())
@@ -57,16 +59,15 @@ TArray<FRigVMUserWorkflowAction> UControlRigTransformWorkflowOptions::ProvideTra
 			if(FRigTransformElement* TransformElement = (FRigTransformElement*)Options->Hierarchy->Find<FRigTransformElement>(Key))
 			{
 				const FTransform Transform = Options->Hierarchy->GetTransform(TransformElement, Options->TransformType);
-				
-				Actions.Emplace(
-					ERigVMUserWorkflowActionType::SetPinDefaultValue,
-					InOptions->GetSubject<URigVMPin>(),
+
+				return Controller->SetPinDefaultValue(
+					InOptions->GetSubject<URigVMPin>()->GetPinPath(),
 					FRigVMStruct::ExportToFullyQualifiedText<FTransform>(Transform)
 				);
 			}
 		}
 	}
 	
-	return Actions;
+	return false;
 }
 #endif
\ No newline at end of file
diff --git a/Engine/Plugins/Animation/ControlRig/Source/ControlRig/Public/Units/ControlRigNodeWorkflow.h b/Engine/Plugins/Animation/ControlRig/Source/ControlRig/Public/Units/ControlRigNodeWorkflow.h
--- a/Engine/Plugins/Animation/ControlRig/Source/ControlRig/Public/Units/ControlRigNodeWorkflow.h
+++ b/Engine/Plugins/Animation/ControlRig/Source/ControlRig/Public/Units/ControlRigNodeWorkflow.h
@@ -40,6 +40,6 @@ public:
 protected:
 
 #if WITH_EDITOR
-	static TArray<FRigVMUserWorkflowAction> ProvideTransformWorkflow(const URigVMUserWorkflowOptions* InOptions);
+	static bool PerformTransformWorkflow(const URigVMUserWorkflowOptions* InOptions, UObject* InController);
 #endif
 };
diff --git a/Engine/Source/Developer/RigVMDeveloper/Private/RigVMModel/Nodes/RigVMUnitNode.cpp b/Engine/Source/Developer/RigVMDeveloper/Private/RigVMModel/Nodes/RigVMUnitNode.cpp
--- a/Engine/Source/Developer/RigVMDeveloper/Private/RigVMModel/Nodes/RigVMUnitNode.cpp
+++ b/Engine/Source/Developer/RigVMDeveloper/Private/RigVMModel/Nodes/RigVMUnitNode.cpp
@@ -4,6 +4,7 @@
 
 #include "Animation/Rig.h"
 #include "RigVMCore/RigVMStruct.h"
+#include "RigVMModel/RigVMControllerActions.h"
 #include "RigVMUserWorkflowRegistry.h"
 
 FString URigVMUnitNode::GetNodeTitle() const
diff --git a/Engine/Source/Developer/RigVMDeveloper/Private/RigVMModel/RigVMController.cpp b/Engine/Source/Developer/RigVMDeveloper/Private/RigVMModel/RigVMController.cpp
--- a/Engine/Source/Developer/RigVMDeveloper/Private/RigVMModel/RigVMController.cpp
+++ b/Engine/Source/Developer/RigVMDeveloper/Private/RigVMModel/RigVMController.cpp
@@ -835,6 +835,111 @@ URigVMUnitNode* URigVMController::AddUnitNodeFromStructPath(const FString& InScr
 	return AddUnitNode(ScriptStruct, InMethodName, InPosition, InNodeName, bSetupUndoRedo, bPrintPythonCommand);
 }
 
+URigVMUnitNode* URigVMController::AddUnitNodeWithDefaults(UScriptStruct* InScriptStruct, const FString& InDefaults,
+	const FName& InMethodName, const FVector2D& InPosition, const FString& InNodeName, bool bSetupUndoRedo,
+	bool bPrintPythonCommand)
+{
+	if(InScriptStruct == nullptr)
+	{
+		return nullptr;
+	}
+
+	FStructOnScope StructOnScope;
+
+	if(!InDefaults.IsEmpty())
+	{
+		StructOnScope = FStructOnScope(InScriptStruct);
+		
+		FRigVMPinDefaultValueImportErrorContext ErrorPipe;
+		InScriptStruct->ImportText(*InDefaults, StructOnScope.GetStructMemory(), nullptr, PPF_None, &ErrorPipe, FString());
+
+		if(ErrorPipe.NumErrors > 0)
+		{
+			return nullptr;
+		}
+	}
+
+	return AddUnitNodeWithDefaults(InScriptStruct, StructOnScope, InMethodName, InPosition, InNodeName, bSetupUndoRedo, bPrintPythonCommand);
+}
+
+URigVMUnitNode* URigVMController::AddUnitNodeWithDefaults(UScriptStruct* InScriptStruct, const FRigStructScope& InDefaults,
+                                                          const FName& InMethodName, const FVector2D& InPosition, const FString& InNodeName, bool bSetupUndoRedo,
+                                                          bool bPrintPythonCommand)
+{
+	if(InScriptStruct == nullptr)
+	{
+		return nullptr;
+	}
+
+	const bool bSetPinDefaults = InDefaults.IsValid() && (InDefaults.GetScriptStruct() == InScriptStruct); 
+	if(bSetPinDefaults)
+	{
+		static constexpr TCHAR AddUnitNodeTitle[] = TEXT("Add Unit Node");
+		OpenUndoBracket(AddUnitNodeTitle);
+	}
+
+	URigVMUnitNode* Node = AddUnitNode(InScriptStruct, InMethodName, InPosition, InNodeName, bSetupUndoRedo, bPrintPythonCommand);
+	if(Node == nullptr)
+	{
+		if(bSetPinDefaults)
+		{
+			CancelUndoBracket();
+		}
+		return nullptr;
+	}
+
+	if(bSetPinDefaults)
+	{
+		if(!SetUnitNodeDefaults(Node, InDefaults))
+		{
+			CancelUndoBracket();
+		}
+	}
+
+	CloseUndoBracket();
+	return Node;
+}
+
+bool URigVMController::SetUnitNodeDefaults(URigVMUnitNode* InNode, const FRigStructScope& InDefaults,
+	bool bSetupUndoRedo, bool bPrintPythonCommand)
+{
+	if(InNode == nullptr || !InDefaults.IsValid())
+	{
+		return false;
+	}
+
+	if(InNode->GetScriptStruct() != InDefaults.GetScriptStruct())
+	{
+		return false;
+	}
+
+	static constexpr TCHAR SetUnitNodeDefaultsTitle[] = TEXT("Set Unit Node Defaults");
+	OpenUndoBracket(SetUnitNodeDefaultsTitle);
+
+	for(URigVMPin* Pin : InNode->GetPins())
+	{
+		if(Pin->GetDirection() != ERigVMPinDirection::Input &&
+			Pin->GetDirection() != ERigVMPinDirection::IO &&
+			Pin->GetDirection() != ERigVMPinDirection::Visible)
+		{
+			continue;
+		}
+		
+		if(const FProperty* Property = InDefaults.GetScriptStruct()->FindPropertyByName(Pin->GetFName()))
+		{
+			const uint8* MemberMemoryPtr = Property->ContainerPtrToValuePtr<uint8>(InDefaults.GetMemory());
+			const FString NewDefault = FRigVMStruct::ExportToFullyQualifiedText(Property, MemberMemoryPtr);
+			if(NewDefault != Pin->GetDefaultValue())
+			{
+				SetPinDefaultValue(Pin->GetPinPath(), NewDefault, true, bSetupUndoRedo, false, bPrintPythonCommand);
+			}
+		}
+	}
+
+	CloseUndoBracket();
+	return true;
+}
+
 URigVMVariableNode* URigVMController::AddVariableNode(const FName& InVariableName, const FString& InCPPType, UObject* InCPPTypeObject, bool bIsGetter, const FString& InDefaultValue, const FVector2D& InPosition, const FString& InNodeName, bool bSetupUndoRedo, bool bPrintPythonCommand)
 {
 	if (!IsValidGraph())
@@ -7063,7 +7168,7 @@ bool URigVMController::SetPinDefaultValue(URigVMPin* InPin, const FString& InDef
 	
 	check(InPin);
 
-	if(!InPin->IsUObject())
+	if(!InPin->IsUObject() && bValidatePinDefaults)
 	{
 		ensure(!InDefaultValue.IsEmpty());
 	}
@@ -7163,16 +7268,6 @@ bool URigVMController::SetPinDefaultValue(URigVMPin* InPin, const FString& InDef
 		}
 	}
 
-	if(bSetupUndoRedo)
-	{
-		TArray<FRigVMUserWorkflow> WorkflowsOnPinChange = InPin->GetNode()->GetSupportedWorkflows(ERigVMUserWorkflowType::OnPinDefaultChanged, InPin);
-		for(const FRigVMUserWorkflow& Workflow : WorkflowsOnPinChange)
-		{
-			const URigVMUserWorkflowOptions* Options = MakeOptionsForWorkflow(InPin, Workflow);
-			PerformUserWorkflow(Workflow, Options, bSetupUndoRedo);
-		}
-	}
-
 	if (bSetupUndoRedo)
 	{
 		ActionStack->EndAction(Action, bMergeUndoAction);
@@ -10966,71 +11061,33 @@ URigVMUserWorkflowOptions* URigVMController::MakeOptionsForWorkflow(UObject* InS
 	return Options;
 }
 
-bool URigVMController::PerformUserWorkflowAction(const FRigVMUserWorkflow& InWorkflow, const FRigVMUserWorkflowAction& InAction, bool bSetupUndoRedo)
+bool URigVMController::PerformUserWorkflow(const FRigVMUserWorkflow& InWorkflow,
+	const URigVMUserWorkflowOptions* InOptions, bool bSetupUndoRedo)
 {
-	if(!InAction.IsValid())
+	if(!InWorkflow.IsValid() || !ensure(InOptions != nullptr))
 	{
 		return false;
 	}
 
-	static bool bIsPerformingUserWorkflowAction = false;
-	if(bIsPerformingUserWorkflowAction)
-	{
-		return false;
-	}
-	TGuardValue<bool> ReentryGuard(bIsPerformingUserWorkflowAction, true);
-	
-	switch(InAction.GetType())
-	{
-		case ERigVMUserWorkflowActionType::SetPinDefaultValue:
-		{
-			if(URigVMPin* Pin = InAction.GetSubject<URigVMPin>())
-			{
-				return SetPinDefaultValue(Pin->GetPinPath(), InAction.GetData(), true, bSetupUndoRedo, false, false);
-			}
-			break;
-		}
-		case ERigVMUserWorkflowActionType::Invalid:
-		default:
-		{
-			break;
-		}
-	}
-	return false;
-}
-
-bool URigVMController::PerformUserWorkflowActions(const FRigVMUserWorkflow& InWorkflow, const TArray<FRigVMUserWorkflowAction>& InActions, bool bSetupUndoRedo)
-{
-	if(InActions.IsEmpty())
-	{
-		return false;
-	}
-	
-	for(const FRigVMUserWorkflowAction& Action : InActions)
-	{
-		if(!Action.IsValid())
-		{
-			return false;
-		}
-	}
 
 	FRigVMBaseAction Bracket;
 	Bracket.Title = InWorkflow.GetTitle();
 	ActionStack->BeginAction(Bracket);
 
-	for(const FRigVMUserWorkflowAction& Action : InActions)
-	{
-		PerformUserWorkflowAction(InWorkflow, Action, bSetupUndoRedo);
-	}
+	const bool bSuccess = InWorkflow.Perform(InOptions, this);
 
 	ActionStack->EndAction(Bracket);
-	return true;
-}
 
-bool URigVMController::PerformUserWorkflow(const FRigVMUserWorkflow& InWorkflow,
-	const URigVMUserWorkflowOptions* InOptions, bool bSetupUndoRedo)
-{
-	return PerformUserWorkflowActions(InWorkflow, InWorkflow.GetActions(InOptions), bSetupUndoRedo);
+	if(!bSuccess)
+	{
+		// if the workflow was run as the top level action we'll undo
+		if(ActionStack->CurrentActions.IsEmpty())
+		{
+			ActionStack->Undo(this);
+		}
+	}
+
+	return bSuccess;
 }
 
 TArray<TSoftObjectPtr<URigVMFunctionReferenceNode>> URigVMController::GetAffectedReferences(ERigVMControllerBulkEditType InEditType, bool bForceLoad, bool bNotify)
@@ -14397,6 +14454,10 @@ void URigVMController::PostProcessDefaultValue(URigVMPin* Pin, FString& OutDefau
 		{
 			OutDefaultValue = OutDefaultValue.LeftChop(1);
 		}
+		if(OutDefaultValue.IsEmpty() && Pin->GetCPPType() == RigVMTypeUtils::FNameType)
+		{
+			OutDefaultValue = FName(NAME_None).ToString();
+		}
 	}
 }
 
diff --git a/Engine/Source/Developer/RigVMDeveloper/Public/RigVMModel/Nodes/RigVMUnitNode.h b/Engine/Source/Developer/RigVMDeveloper/Public/RigVMModel/Nodes/RigVMUnitNode.h
--- a/Engine/Source/Developer/RigVMDeveloper/Public/RigVMModel/Nodes/RigVMUnitNode.h
+++ b/Engine/Source/Developer/RigVMDeveloper/Public/RigVMModel/Nodes/RigVMUnitNode.h
@@ -60,6 +60,23 @@ public:
 	// @param bUseDefault If set to true the default struct will be created - otherwise the struct will contains the values from the node
 	TSharedPtr<FStructOnScope> ConstructStructInstance(bool bUseDefault = false) const;
 
+	// Returns a copy of the struct with the current values
+	template <
+		typename T,
+		typename TEnableIf<TModels<CRigVMUStruct, T>::Value>::Type* = nullptr
+	>
+	FORCEINLINE T ConstructStructInstance() const
+	{
+		if(!ensure(T::StaticStruct() == GetScriptStruct()))
+		{
+			return T();
+		}
+
+		TSharedPtr<FStructOnScope> Instance = ConstructStructInstance(false);
+		const T& InstanceRef = *(const T*)Instance->GetStructMemory();
+		return InstanceRef;
+	}
+
 	FRigVMStructUpgradeInfo GetUpgradeInfo() const;
 
 protected:
diff --git a/Engine/Source/Developer/RigVMDeveloper/Public/RigVMModel/RigVMController.h b/Engine/Source/Developer/RigVMDeveloper/Public/RigVMModel/RigVMController.h
--- a/Engine/Source/Developer/RigVMDeveloper/Public/RigVMModel/RigVMController.h
+++ b/Engine/Source/Developer/RigVMDeveloper/Public/RigVMModel/RigVMController.h
@@ -97,6 +97,42 @@ class RIGVMDEVELOPER_API URigVMControllerHost : public UInterface
 	GENERATED_BODY()
 };
 
+USTRUCT(BlueprintType)
+struct RIGVMDEVELOPER_API FRigStructScope
+{
+	GENERATED_BODY()
+
+public:
+	
+	FRigStructScope()
+		: ScriptStruct(nullptr)
+		, Memory(nullptr)
+	{}
+
+	template <
+		typename T,
+		typename TEnableIf<TModels<CRigVMUStruct, T>::Value>::Type * = nullptr
+	>
+	FRigStructScope(const T& InInstance)
+		: ScriptStruct(T::StaticStruct())
+		, Memory((const uint8*)&InInstance)
+	{}
+
+	FRigStructScope(const FStructOnScope& InScope)
+		: ScriptStruct(Cast<UScriptStruct>(InScope.GetStruct()))
+		, Memory(InScope.GetStructMemory()) 
+	{}
+
+	const UScriptStruct* GetScriptStruct() const { return ScriptStruct; }
+	const uint8* GetMemory() const { return Memory; }
+	bool IsValid() const { return ScriptStruct != nullptr && Memory != nullptr; }
+
+protected:
+
+	const UScriptStruct* ScriptStruct;
+	const uint8* Memory;
+};
+
 // Interface that allows an object to host a rig VM controller. Used by graph edting code to interact with the controller.
 class RIGVMDEVELOPER_API IRigVMControllerHost
 {
@@ -116,7 +152,6 @@ public:
 	virtual URigVMController* GetOrCreateRigVMController(const UObject* InEditorObject) = 0;
 };
 
-
 /**
  * The Controller is the sole authority to perform changes
  * on the Graph. The Controller itself is stateless.
@@ -208,6 +243,33 @@ public:
 	UFUNCTION(BlueprintCallable, Category = RigVMController)
 	URigVMUnitNode* AddUnitNodeFromStructPath(const FString& InScriptStructPath, const FName& InMethodName = TEXT("Execute"), const FVector2D& InPosition = FVector2D::ZeroVector, const FString& InNodeName = TEXT(""), bool bSetupUndoRedo = true, bool bPrintPythonCommand = false);
 
+	// Adds a unit node using a template
+	template <
+		typename T,
+		typename TEnableIf<TModels<CRigVMUStruct, T>::Value>::Type * = nullptr
+	>
+	URigVMUnitNode* AddUnitNode(const T& InDefaults, const FName& InMethodName = TEXT("Execute"), const FVector2D& InPosition = FVector2D::ZeroVector, const FString& InNodeName = TEXT(""), bool bSetupUndoRedo = true, bool bPrintPythonCommand = false)
+	{
+		return AddUnitNodeWithDefaults(T::StaticStruct(), InDefaults, InMethodName, InPosition, InNodeName, bSetupUndoRedo, bPrintPythonCommand); 
+	}
+
+	// Adds a Function / Struct Node to the edited Graph.
+	// UnitNode represent a RIGVM_METHOD declaration on a USTRUCT.
+	// This causes a NodeAdded modified event.
+	UFUNCTION(BlueprintCallable, Category = RigVMController)
+	URigVMUnitNode* AddUnitNodeWithDefaults(UScriptStruct* InScriptStruct, const FString& InDefaults, const FName& InMethodName = TEXT("Execute"), const FVector2D& InPosition = FVector2D::ZeroVector, const FString& InNodeName = TEXT(""), bool bSetupUndoRedo = true, bool bPrintPythonCommand = false);
+
+	// Adds a Function / Struct Node to the edited Graph.
+	// UnitNode represent a RIGVM_METHOD declaration on a USTRUCT.
+	// This causes a NodeAdded modified event.
+	URigVMUnitNode* AddUnitNodeWithDefaults(UScriptStruct* InScriptStruct, const FRigStructScope& InDefaults, const FName& InMethodName = TEXT("Execute"), const FVector2D& InPosition = FVector2D::ZeroVector, const FString& InNodeName = TEXT(""), bool bSetupUndoRedo = true, bool bPrintPythonCommand = false);
+
+	// Adds a Function / Struct Node to the edited Graph.
+	// UnitNode represent a RIGVM_METHOD declaration on a USTRUCT.
+	// This causes a NodeAdded modified event.
+	UFUNCTION(BlueprintCallable, Category = RigVMController)
+	bool SetUnitNodeDefaults(URigVMUnitNode* InNode, const FRigStructScope& InDefaults, bool bSetupUndoRedo = true, bool bPrintPythonCommand = false);
+
 	// Adds a Variable Node to the edited Graph.
 	// Variables represent local work state for the function and
 	// can be read from and written to. 
@@ -941,8 +1003,6 @@ private:
 	URigVMInjectionInfo* InjectNodeIntoPin(URigVMPin* InPin, bool bAsInput, const FName& InInputPinName, const FName& InOutputPinName, bool bSetupUndoRedo = true);
 	URigVMNode* EjectNodeFromPin(URigVMPin* InPin, bool bSetupUndoRedo = true, bool bPrintPythonCommands = false);
 	bool EjectAllInjectedNodes(URigVMNode* InNode, bool bSetupUndoRedo = true, bool bPrintPythonCommands = false);
-	bool PerformUserWorkflowAction(const FRigVMUserWorkflow& InWorkflow, const FRigVMUserWorkflowAction& InAction, bool bSetupUndoRedo);
-	bool PerformUserWorkflowActions(const FRigVMUserWorkflow& InWorkflow, const TArray<FRigVMUserWorkflowAction>& InActions, bool bSetupUndoRedo);
 
 	// try to reconnect source and target pins after a node deletion
 	void RelinkSourceAndTargetPins(URigVMNode* RigNode, bool bSetupUndoRedo = true);
diff --git a/Engine/Source/Runtime/RigVM/Private/RigVMCore/RigVMUserWorkflow.cpp b/Engine/Source/Runtime/RigVM/Private/RigVMCore/RigVMUserWorkflow.cpp
--- a/Engine/Source/Runtime/RigVM/Private/RigVMCore/RigVMUserWorkflow.cpp
+++ b/Engine/Source/Runtime/RigVM/Private/RigVMCore/RigVMUserWorkflow.cpp
@@ -59,29 +59,23 @@ void URigVMUserWorkflowOptions::Report(EMessageSeverity::Type InSeverity,  const
 
 ////////////////////////////////////////////////////////////////////////////////////////////////////////////////////////
 
-TArray<FRigVMUserWorkflowAction> FRigVMUserWorkflow::GetActions(const URigVMUserWorkflowOptions* InOptions) const
+bool FRigVMUserWorkflow::Perform(const URigVMUserWorkflowOptions* InOptions, UObject* InController) const
 {
-	TArray<FRigVMUserWorkflowAction> Actions;
 	if(!IsValid() || !ValidateOptions(InOptions))
 	{
-		return Actions;
+		return false;
 	}
 
-	if(OnGetActionsDelegate.IsBound())
+	if(PerformDelegate.IsBound())
 	{
-		Actions = OnGetActionsDelegate.Execute(InOptions);
+		return PerformDelegate.Execute(InOptions, InController);
 	}
-	else if(OnGetActionsDynamicDelegate.IsBound())
+	else if(PerformDynamicDelegate.IsBound())
 	{
-		Actions = OnGetActionsDynamicDelegate.Execute(InOptions);
+		return PerformDynamicDelegate.Execute(InOptions, InController);
 	}
 	
-	Actions = Actions.FilterByPredicate([](const FRigVMUserWorkflowAction& Action) -> bool
-	{
-		return Action.IsValid();
-	});
-	
-	return Actions;
+	return false;
 }
 
 bool FRigVMUserWorkflow::ValidateOptions(const URigVMUserWorkflowOptions* InOptions) const
diff --git a/Engine/Source/Runtime/RigVM/Public/RigVMCore/RigVMUserWorkflow.h b/Engine/Source/Runtime/RigVM/Public/RigVMCore/RigVMUserWorkflow.h
--- a/Engine/Source/Runtime/RigVM/Public/RigVMCore/RigVMUserWorkflow.h
+++ b/Engine/Source/Runtime/RigVM/Public/RigVMCore/RigVMUserWorkflow.h
@@ -5,61 +5,14 @@
 #include "CoreMinimal.h"
 #include "Logging/TokenizedMessage.h"
 #include "RigVMCore/RigVMTraits.h"
+#include "UObject/StructOnScope.h"
 #include "RigVMUserWorkflow.generated.h"
 
 class URigVMUserWorkflowOptions;
 
 DECLARE_DELEGATE_ThreeParams(FRigVMReportDelegate, EMessageSeverity::Type, UObject*, const FString&);
-
-// Types of actions within a workflow
-UENUM(BlueprintType)
-enum class ERigVMUserWorkflowActionType : uint8
-{
-	Invalid = 0 UMETA(Hidden),
-	SetPinDefaultValue = 1
-};
-
-USTRUCT(BlueprintType)
-struct RIGVM_API FRigVMUserWorkflowAction
-{
-	GENERATED_BODY()
-
-public:
-
-	FORCEINLINE FRigVMUserWorkflowAction()
-    : Type(ERigVMUserWorkflowActionType::Invalid)
-    , Subject()
-    , Data()
-    {}
-
-	FORCEINLINE FRigVMUserWorkflowAction(
-		ERigVMUserWorkflowActionType InType,
-		UObject* InSubject,
-		const FString& InData)
-	: Type(InType)
-	, Subject(InSubject)
-	, Data(InData)
-	{}
-
-	FORCEINLINE bool IsValid() const { return Type != ERigVMUserWorkflowActionType::Invalid && Subject != nullptr; }
-	FORCEINLINE ERigVMUserWorkflowActionType GetType() const { return Type; }
-	FORCEINLINE UObject* GetSubject() const { return Subject; }
-	FORCEINLINE const FString& GetData() const { return Data; }
-	
-	template<typename T>
-	FORCEINLINE T* GetSubject() const { return Cast<T>(Subject.Get()); }
-
-protected:
-
-	UPROPERTY(BlueprintReadWrite, EditAnywhere, Category = Action, meta = (AllowPrivateAccess = true))
-	ERigVMUserWorkflowActionType Type;
-
-	UPROPERTY(BlueprintReadWrite, EditAnywhere, Category = Action, meta = (AllowPrivateAccess = true))
-	TObjectPtr<UObject> Subject;
-
-	UPROPERTY(BlueprintReadWrite, EditAnywhere, Category = Action, meta = (AllowPrivateAccess = true))
-	FString Data;
-};
+DECLARE_DELEGATE_RetVal_TwoParams(bool, FRigVMPerformUserWorkflowDelegate, const URigVMUserWorkflowOptions*, UObject*);
+DECLARE_DYNAMIC_DELEGATE_RetVal_TwoParams(bool, FRigVMPeformUserWorkflowDynamicDelegate, const URigVMUserWorkflowOptions*, InOptions, UObject*, InController);
 
 // Types of workflows offered by a rigvm struct node
 UENUM(BlueprintType)
@@ -72,9 +25,6 @@ enum class ERigVMUserWorkflowType : uint8
 	All = NodeContext | PinContext | OnPinDefaultChanged
 };
 
-DECLARE_DELEGATE_RetVal_OneParam(TArray<FRigVMUserWorkflowAction>, FRigVMWorkflowGetActionsDelegate, const URigVMUserWorkflowOptions*);
-DECLARE_DYNAMIC_DELEGATE_RetVal_OneParam(TArray<FRigVMUserWorkflowAction>, FRigVMWorkflowGetActionsDynamicDelegate, const URigVMUserWorkflowOptions*, InOptions);
-
 USTRUCT(BlueprintType)
 struct RIGVM_API FRigVMUserWorkflow
 {
@@ -83,11 +33,7 @@ struct RIGVM_API FRigVMUserWorkflow
 public:
 
 	FORCEINLINE FRigVMUserWorkflow()
-	: Title()
-	, Tooltip()
-	, Type(ERigVMUserWorkflowType::Invalid)
-	, OnGetActionsDelegate()
-	, OnGetActionsDynamicDelegate()
+	: Type(ERigVMUserWorkflowType::Invalid)
 	, OptionsClass(nullptr)
 	{}
 
@@ -95,13 +41,12 @@ public:
 		const FString& InTitle,
 		const FString& InTooltip,
 		ERigVMUserWorkflowType InType,
-		FRigVMWorkflowGetActionsDelegate InGetActionsDelegate,
+		FRigVMPerformUserWorkflowDelegate InGetActionsDelegate,
 		UClass* InOptionsClass)
 	: Title(InTitle)
 	, Tooltip(InTooltip)
 	, Type(InType)
-	, OnGetActionsDelegate(InGetActionsDelegate)
-	, OnGetActionsDynamicDelegate()
+	, PerformDelegate(InGetActionsDelegate)
 	, OptionsClass(InOptionsClass)
 	{}
 	
@@ -111,7 +56,7 @@ public:
 	{
 		return Type != ERigVMUserWorkflowType::Invalid &&
 			GetOptionsClass() != nullptr &&
-			(OnGetActionsDelegate.IsBound() || OnGetActionsDynamicDelegate.IsBound());
+			(PerformDelegate.IsBound() || PerformDynamicDelegate.IsBound());
 	}
 
 	FORCEINLINE const FString& GetTitle() const { return Title; }
@@ -119,7 +64,7 @@ public:
 	FORCEINLINE ERigVMUserWorkflowType GetType() const { return Type; }
 	FORCEINLINE UClass* GetOptionsClass() const { return OptionsClass; }
 
-	TArray<FRigVMUserWorkflowAction> GetActions(const URigVMUserWorkflowOptions* InOptions) const;
+	bool Perform(const URigVMUserWorkflowOptions* InOptions, UObject* InController) const;
 
 protected:
 
@@ -134,10 +79,10 @@ protected:
 	UPROPERTY(BlueprintReadWrite, EditAnywhere, Category = Workflow, meta = (AllowPrivateAccess = true))
 	ERigVMUserWorkflowType Type;
 
-	FRigVMWorkflowGetActionsDelegate OnGetActionsDelegate;
+	FRigVMPerformUserWorkflowDelegate PerformDelegate;
 
-	UPROPERTY(BlueprintReadWrite, EditAnywhere, Category = Workflow, meta = (ScriptName = "OnGetActions", AllowPrivateAccess = true))
-	FRigVMWorkflowGetActionsDynamicDelegate OnGetActionsDynamicDelegate;
+	UPROPERTY(BlueprintReadWrite, EditAnywhere, Category = Workflow, meta = (ScriptName = "OnPerformWorkflow", AllowPrivateAccess = true))
+	FRigVMPeformUserWorkflowDynamicDelegate PerformDynamicDelegate;
 
 	UPROPERTY(BlueprintReadWrite, EditAnywhere, Category = Workflow, meta = (AllowPrivateAccess = true))
 	TObjectPtr<UClass> OptionsClass;
