Control Rig: Disable erroneous diagnosis from PVS.

#rnx

diff --git a/Engine/Source/Developer/RigVMDeveloper/Private/RigVMCompiler/RigVMCodeGenerator.cpp b/Engine/Source/Developer/RigVMDeveloper/Private/RigVMCompiler/RigVMCodeGenerator.cpp
--- a/Engine/Source/Developer/RigVMDeveloper/Private/RigVMCompiler/RigVMCodeGenerator.cpp
+++ b/Engine/Source/Developer/RigVMDeveloper/Private/RigVMCompiler/RigVMCodeGenerator.cpp
@@ -143,7 +143,7 @@ FString FRigVMCodeGenerator::DumpExternalVariables(bool bForHeader, bool bLog)
 	for(int32 ExternalVariableIndex = 0; ExternalVariableIndex < VM->GetExternalVariables().Num(); ExternalVariableIndex++)
 	{
 		const FRigVMOperand ExternalVarOperand(ERigVMMemoryType::External, ExternalVariableIndex, INDEX_NONE);
-		const FRigVMExternalVariable& ExternalVariable = VM->GetExternalVariables()[ExternalVariableIndex];
+		const FRigVMExternalVariable& ExternalVariable = VM->GetExternalVariables()[ExternalVariableIndex]; //-V758
 		const FString ExternalVarCPPType = ExternalVariable.GetExtendedCPPType().ToString();
 		FString OperandName = *GetOperandName(ExternalVarOperand, false);
 		if(OperandName.StartsWith(TEXT("(*")) && OperandName.EndsWith(TEXT(")")))
