RigVM: Respect name map during import from text

diff --git a/Engine/Source/Developer/RigVMDeveloper/Private/RigVMModel/RigVMController.cpp b/Engine/Source/Developer/RigVMDeveloper/Private/RigVMModel/RigVMController.cpp
--- a/Engine/Source/Developer/RigVMDeveloper/Private/RigVMModel/RigVMController.cpp
+++ b/Engine/Source/Developer/RigVMDeveloper/Private/RigVMModel/RigVMController.cpp
@@ -3171,6 +3171,7 @@ struct FRigVMControllerObjectFactory : public FCustomizableTextObjectFactory
 public:
 	URigVMController* Controller;
 	TArray<URigVMNode*> CreatedNodes;
+	TArray<FName> CreateNodeNames;
 	TMap<FName, FName> NodeNameMap;
 	TArray<URigVMLink*> CreatedLinks;
 public:
@@ -3200,8 +3201,15 @@ protected:
 	{
 		if (URigVMNode* DefaultNode = Cast<URigVMNode>(ObjectClass->GetDefaultObject()))
 		{
-			FName ValidName = *Controller->GetValidNodeName(InOutObjName.ToString());
+			URigVMGraph* Graph = Controller->GetGraph();
+			check(Graph);
+
+			const FName ValidName = Controller->GetUniqueName(InOutObjName, [Graph, this](const FName& InName) {
+				return !CreateNodeNames.Contains(InName) && Graph->IsNameAvailable(InName.ToString());
+			}, false, true);
+			
 			NodeNameMap.Add(InOutObjName, ValidName);
+			CreateNodeNames.Add(ValidName);
 			InOutObjName = ValidName;
 		}
 	}
