RigVM: When renaming a function update the build data

diff --git a/Engine/Source/Developer/RigVMDeveloper/Private/RigVMModel/RigVMController.cpp b/Engine/Source/Developer/RigVMDeveloper/Private/RigVMModel/RigVMController.cpp
--- a/Engine/Source/Developer/RigVMDeveloper/Private/RigVMModel/RigVMController.cpp
+++ b/Engine/Source/Developer/RigVMDeveloper/Private/RigVMModel/RigVMController.cpp
@@ -5994,6 +5994,7 @@ bool URigVMController::RenameNode(URigVMNode* InNode, const FName& InNewName, bo
 		Notify(ERigVMGraphNotifType::LinkRemoved, Link);
 	}
 
+	const FSoftObjectPath PreviousObjectPath(InNode);
 	InNode->PreviousName = InNode->GetFName();
 	if (!RenameObject(InNode, *ValidNewName.ToString()))
 	{
@@ -6014,6 +6015,25 @@ bool URigVMController::RenameNode(URigVMNode* InNode, const FName& InNewName, bo
 	{
 		if (URigVMFunctionLibrary* FunctionLibrary = Cast<URigVMFunctionLibrary>(LibraryNode->GetGraph()))
 		{
+			// update the table in the build data
+			if(URigVMBuildData* BuildData = GetBuildData())
+			{
+				for(const TPair< TSoftObjectPtr<URigVMLibraryNode>, FRigVMFunctionReferenceArray >& Pair: BuildData->FunctionReferences)
+				{
+					if(Pair.Key.ToSoftObjectPath() == PreviousObjectPath)
+					{
+						const TSoftObjectPtr<URigVMLibraryNode> SoftObjectPtr(InNode);
+						const FRigVMFunctionReferenceArray FunctionReferences = Pair.Value;
+						
+						BuildData->Modify();
+						BuildData->FunctionReferences.Remove(Pair.Key);
+						BuildData->FunctionReferences.Add(SoftObjectPtr, FunctionReferences);
+						BuildData->MarkPackageDirty();
+						break;
+					}
+				}
+			}
+			
 			FunctionLibrary->ForEachReference(LibraryNode->GetFName(), [this, InNewName](URigVMFunctionReferenceNode* ReferenceNode)
 			{
 				FRigVMControllerGraphGuard GraphGuard(this, ReferenceNode->GetGraph(), false);
