Control Rig: Fix reroutes incorrectly identifying cycles

diff --git a/Engine/Source/Developer/RigVMDeveloper/Private/RigVMCompiler/RigVMAST.cpp b/Engine/Source/Developer/RigVMDeveloper/Private/RigVMCompiler/RigVMAST.cpp
--- a/Engine/Source/Developer/RigVMDeveloper/Private/RigVMCompiler/RigVMAST.cpp
+++ b/Engine/Source/Developer/RigVMDeveloper/Private/RigVMCompiler/RigVMAST.cpp
@@ -2214,38 +2214,38 @@ bool FRigVMParserAST::CanLink(URigVMPin* InSourcePin, URigVMPin* InTargetPin, FS
 
 	if(SourceNode->IsA<URigVMRerouteNode>())
 	{
-		TArray<URigVMNode*> LinkedSourceNodes = SourceNode->GetLinkedSourceNodes();
-		for (int32 LinkedSourceNodeIndex = 0; LinkedSourceNodeIndex<LinkedSourceNodes.Num(); LinkedSourceNodeIndex++)
+		TArray<URigVMPin*> SourcePins;
+		for (URigVMPin* Pin : SourceNode->GetPins())
 		{
-			URigVMNode* LinkedSourceNode = LinkedSourceNodes[LinkedSourceNodeIndex];
-			if(LinkedSourceNode->IsA<URigVMRerouteNode>())
-			{
-				LinkedSourceNodes.Append(LinkedSourceNode->GetLinkedSourceNodes());
-			}
-			else
+			SourcePins.Append(Pin->GetLinkedSourcePins(true));
+		}
+		
+		for (URigVMPin* SourcePin : SourcePins)
+		{
+			if (!CanLink(SourcePin, InTargetPin, OutFailureReason))
 			{
-				SourceNode = LinkedSourceNode;
-				break;
+				return false;
 			}
 		}
+		return true;
 	}
 
 	if(TargetNode->IsA<URigVMRerouteNode>())
 	{
-		TArray<URigVMNode*> LinkedTargetNodes = TargetNode->GetLinkedTargetNodes();
-		for (int32 LinkedTargetNodeIndex = 0; LinkedTargetNodeIndex<LinkedTargetNodes.Num(); LinkedTargetNodeIndex++)
+		TArray<URigVMPin*> TargetPins;
+		for (URigVMPin* Pin : TargetNode->GetPins())
 		{
-			URigVMNode* LinkedTargetNode = LinkedTargetNodes[LinkedTargetNodeIndex];
-			if(LinkedTargetNode->IsA<URigVMRerouteNode>())
-			{
-				LinkedTargetNodes.Append(LinkedTargetNode->GetLinkedTargetNodes());
-			}
-			else
+			TargetPins.Append(Pin->GetLinkedTargetPins(true));
+		}
+		
+		for (URigVMPin* TargetPin : TargetPins)
+		{
+			if (!CanLink(InSourcePin, TargetPin, OutFailureReason))
 			{
-				TargetNode = LinkedTargetNode;
-				break;
+				return false;
 			}
 		}
+		return true;
 	}
 
 	FRigVMASTProxy SourceNodeProxy = FRigVMASTProxy::MakeFromUObject(SourceNode);
diff --git a/Engine/Source/Developer/RigVMDeveloper/Public/RigVMModel/Nodes/RigVMRerouteNode.h b/Engine/Source/Developer/RigVMDeveloper/Public/RigVMModel/Nodes/RigVMRerouteNode.h
--- a/Engine/Source/Developer/RigVMDeveloper/Public/RigVMModel/Nodes/RigVMRerouteNode.h
+++ b/Engine/Source/Developer/RigVMDeveloper/Public/RigVMModel/Nodes/RigVMRerouteNode.h
@@ -42,5 +42,6 @@ private:
 
 	friend class URigVMController;
 	friend class URigVMCompiler;
+	friend class FRigVMParserAST;
 };
 
