RigVM: Reroute to inherit template node

diff --git a/Engine/Plugins/Animation/ControlRig/Source/ControlRigDeveloper/Private/Graph/ControlRigGraphSchema.cpp b/Engine/Plugins/Animation/ControlRig/Source/ControlRigDeveloper/Private/Graph/ControlRigGraphSchema.cpp
--- a/Engine/Plugins/Animation/ControlRig/Source/ControlRigDeveloper/Private/Graph/ControlRigGraphSchema.cpp
+++ b/Engine/Plugins/Animation/ControlRig/Source/ControlRigDeveloper/Private/Graph/ControlRigGraphSchema.cpp
@@ -568,6 +568,21 @@ bool UControlRigGraphSchema::TryCreateConnection(UEdGraphPin* PinA, UEdGraphPin*
 				UserLinkDirection = ERigVMPinDirection::Input;
 			}
 
+			const bool bPinAIsWildCard = PinA->PinType.PinSubCategoryObject == RigVMTypeUtils::GetWildCardCPPTypeObject();
+			const bool bPinBIsWildCard = PinB->PinType.PinSubCategoryObject == RigVMTypeUtils::GetWildCardCPPTypeObject();
+			if(bPinAIsWildCard != bPinBIsWildCard)
+			{
+				// switch the user link direction if only one of the pins is a wildcard 
+				if(UserLinkDirection == ERigVMPinDirection::Input && bPinBIsWildCard)
+				{
+					UserLinkDirection = ERigVMPinDirection::Output;
+				}
+				else if(UserLinkDirection == ERigVMPinDirection::Output && bPinAIsWildCard)
+				{
+					UserLinkDirection = ERigVMPinDirection::Input;
+				}
+			}
+
 #if WITH_EDITOR
 
 			// check if we are trying to connect a loop iteration pin to a return
@@ -699,6 +714,19 @@ const FPinConnectionResponse UControlRigGraphSchema::CanCreateConnection(const U
 			UserLinkDirection = ERigVMPinDirection::Input;
 		}
 
+		if(PinA->IsWildCard() != PinB->IsWildCard())
+		{
+			// switch the user link direction if only one of the pins is a wildcard 
+			if(UserLinkDirection == ERigVMPinDirection::Input && PinB->IsWildCard())
+			{
+				UserLinkDirection = ERigVMPinDirection::Output;
+			}
+			else if(UserLinkDirection == ERigVMPinDirection::Output && PinA->IsWildCard())
+			{
+				UserLinkDirection = ERigVMPinDirection::Input;
+			}
+		}
+
 		const FRigVMByteCode* ByteCode = RigNodeA->GetController()->GetCurrentByteCode();
 
 		FString FailureReason;
diff --git a/Engine/Plugins/Animation/ControlRig/Source/ControlRigEditor/Private/Graph/SControlRigGraphChangePinType.cpp b/Engine/Plugins/Animation/ControlRig/Source/ControlRigEditor/Private/Graph/SControlRigGraphChangePinType.cpp
--- a/Engine/Plugins/Animation/ControlRig/Source/ControlRigEditor/Private/Graph/SControlRigGraphChangePinType.cpp
+++ b/Engine/Plugins/Animation/ControlRig/Source/ControlRigEditor/Private/Graph/SControlRigGraphChangePinType.cpp
@@ -176,7 +176,16 @@ void SControlRigChangePinType::FillPinTypeMenu(FMenuBuilder& MenuBuilder)
 			return Info;
 		}
 
+		static FArgumentInfo Make(const FRigVMTemplateArgument::FType& InType)
+		{
+			FArgumentInfo Info;
+			Info.Property = nullptr;
+			Info.Type = InType;
+			return Info;
+		}
+
 		const FProperty* Property;
+		FRigVMTemplateArgument::FType Type;
 	};
 
 	typedef TPair<FRigVMTemplateArgument::FType, FArgumentInfo> FTypePair;
@@ -219,14 +228,21 @@ void SControlRigChangePinType::FillPinTypeMenu(FMenuBuilder& MenuBuilder)
 						{
 							continue;
 						}
-						
-						if(UScriptStruct* FunctionStruct = Template->GetPermutation(PermutationIndex)->Struct)
+
+						if(const FRigVMFunction* Permutation = Template->GetPermutation(PermutationIndex))
 						{
-							if(const FProperty* Property = FunctionStruct->FindPropertyByName(ArgumentName))
+							if(UScriptStruct* FunctionStruct = Permutation->Struct)
 							{
-								Types.Add(ArgumentType, FArgumentInfo::Make(Property));
+								if(const FProperty* Property = FunctionStruct->FindPropertyByName(ArgumentName))
+								{
+									Types.Add(ArgumentType, FArgumentInfo::Make(Property));
+								}
 							}
 						}
+						else
+						{
+							Types.Add(ArgumentType, FArgumentInfo::Make(ArgumentType));
+						}
 					}
 				}
 			}
@@ -280,9 +296,16 @@ void SControlRigChangePinType::FillPinTypeMenu(FMenuBuilder& MenuBuilder)
 		for(int32 TypeIndex=0; TypeIndex < SortedTypes.Num(); TypeIndex++)
 		{
 			const FRigVMTemplateArgument::FType& Type = SortedTypes[TypeIndex].Key;
-			const FProperty* Property = SortedTypes[TypeIndex].Value.Property;
 			const FEdGraphPinType PinType = RigVMTypeUtils::PinTypeFromCPPType(Type.CPPType, Type.CPPTypeObject);
 
+			if(UScriptStruct* ScriptStruct = Cast<UScriptStruct>(Type.CPPTypeObject))
+			{
+				if(ScriptStruct->IsChildOf(FRigVMExecuteContext::StaticStruct()))
+				{
+					continue;
+				}
+			}
+
 			MenuBuilder.AddMenuEntry(
 				FUIAction(FExecuteAction::CreateSP(this, &SControlRigChangePinType::HandlePinTypeChanged, Type)),
 				SNew(SHorizontalBox)
diff --git a/Engine/Source/Developer/RigVMDeveloper/Private/RigVMModel/Nodes/RigVMRerouteNode.cpp b/Engine/Source/Developer/RigVMDeveloper/Private/RigVMModel/Nodes/RigVMRerouteNode.cpp
--- a/Engine/Source/Developer/RigVMDeveloper/Private/RigVMModel/Nodes/RigVMRerouteNode.cpp
+++ b/Engine/Source/Developer/RigVMDeveloper/Private/RigVMModel/Nodes/RigVMRerouteNode.cpp
@@ -3,6 +3,7 @@
 #include "RigVMModel/Nodes/RigVMRerouteNode.h"
 
 const FString URigVMRerouteNode::RerouteName = TEXT("Reroute");
+const FString URigVMRerouteNode::RerouteArrayName = TEXT("RerouteArray");
 const FString URigVMRerouteNode::ValueName = TEXT("Value");
 
 URigVMRerouteNode::URigVMRerouteNode()
@@ -62,3 +63,31 @@ FLinearColor URigVMRerouteNode::GetNodeColor() const
 {
 	return FLinearColor::White;
 }
+
+FName URigVMRerouteNode::GetNotation() const
+{
+	return GetTemplate()->GetNotation();
+}
+
+const FRigVMTemplate* URigVMRerouteNode::GetTemplate() const
+{
+	if(CachedTemplate == nullptr)
+	{
+		bool bIsArray = false;
+		if(const URigVMPin* ValuePin = FindPin(ValueName))
+		{
+			bIsArray = ValuePin->IsArray();
+		}
+		CachedTemplate = FindOrAddTemplate(bIsArray);
+	}
+	return CachedTemplate;
+}
+
+const FRigVMTemplate* URigVMRerouteNode::FindOrAddTemplate(bool bIsArray)
+{
+	TArray<FRigVMTemplateArgument> Arguments;
+	Arguments.Emplace(TEXT("Value"), ERigVMPinDirection::IO, bIsArray ? FRigVMTemplateArgument::FType::Array() : FRigVMTemplateArgument::FType());
+		
+	return FRigVMRegistry::Get().GetOrAddTemplateFromArguments(
+		bIsArray ? *RerouteArrayName : *RerouteName, Arguments, true);
+}
diff --git a/Engine/Source/Developer/RigVMDeveloper/Private/RigVMModel/Nodes/RigVMTemplateNode.cpp b/Engine/Source/Developer/RigVMDeveloper/Private/RigVMModel/Nodes/RigVMTemplateNode.cpp
--- a/Engine/Source/Developer/RigVMDeveloper/Private/RigVMModel/Nodes/RigVMTemplateNode.cpp
+++ b/Engine/Source/Developer/RigVMDeveloper/Private/RigVMModel/Nodes/RigVMTemplateNode.cpp
@@ -72,9 +72,22 @@ FText URigVMTemplateNode::GetToolTipTextForPin(const URigVMPin* InPin) const
 		{
 			if (const FRigVMTemplateArgument* Arg = Template->FindArgument(RootPin->GetFName()))
 			{
+				FString Tooltip;
+
 				const TArray<int32> PermutationIndices = GetResolvedPermutationIndices();
-				const FString SupportedTypesJoined = FString::Join(Arg->GetSupportedTypeStrings(PermutationIndices), TEXT("\n"));
-				FString Tooltip = TEXT("Supported Types:\n\n") + SupportedTypesJoined;
+				if(PermutationIndices.Num() == GetTemplate()->NumPermutations())
+				{
+					if(Arg->GetTypes().Num() > 100)
+					{
+						Tooltip = TEXT("Supports any type.");
+					}
+				}
+				
+				if(Tooltip.IsEmpty())
+				{
+					const FString SupportedTypesJoined = FString::Join(Arg->GetSupportedTypeStrings(PermutationIndices), TEXT("\n"));
+					Tooltip = TEXT("Supported Types:\n\n") + SupportedTypesJoined;
+				}
 
 				if(!SuperToolTip.IsEmpty())
 				{
@@ -315,14 +328,16 @@ FString URigVMTemplateNode::GetInitialDefaultValueForPin(const FName& InRootPinN
 
 	for(const int32 PermutationIndex : PermutationIndices)
 	{
-		const FRigVMFunction* Permutation = GetTemplate()->GetPermutation(PermutationIndex);
-		check(Permutation);
-
-		const TSharedPtr<FStructOnScope> StructOnScope = MakeShareable(new FStructOnScope(Permutation->Struct));
-		const FRigVMStruct* DefaultStruct = (const FRigVMStruct*)StructOnScope->GetStructMemory();
+		FString NewDefaultValue;
+		
+		if(const FRigVMFunction* Permutation = GetTemplate()->GetPermutation(PermutationIndex))
+		{
+			const TSharedPtr<FStructOnScope> StructOnScope = MakeShareable(new FStructOnScope(Permutation->Struct));
+			const FRigVMStruct* DefaultStruct = (const FRigVMStruct*)StructOnScope->GetStructMemory();
 
-		const FString NewDefaultValue = DefaultStruct->ExportToFullyQualifiedText(
-			Cast<UScriptStruct>(StructOnScope->GetStruct()), InRootPinName);
+			NewDefaultValue = DefaultStruct->ExportToFullyQualifiedText(
+				Cast<UScriptStruct>(StructOnScope->GetStruct()), InRootPinName);
+		}
 
 		if(!NewDefaultValue.IsEmpty())
 		{
diff --git a/Engine/Source/Developer/RigVMDeveloper/Private/RigVMModel/RigVMController.cpp b/Engine/Source/Developer/RigVMDeveloper/Private/RigVMModel/RigVMController.cpp
--- a/Engine/Source/Developer/RigVMDeveloper/Private/RigVMModel/RigVMController.cpp
+++ b/Engine/Source/Developer/RigVMDeveloper/Private/RigVMModel/RigVMController.cpp
@@ -714,7 +714,7 @@ URigVMUnitNode* URigVMController::AddUnitNode(UScriptStruct* InScriptStruct, con
 			return nullptr;
 		}
 
-		FullyResolveTemplateNode(TemplateNode, Function, bSetupUndoRedo);
+		FullyResolveTemplateNode(TemplateNode, Template->FindPermutation(Function), bSetupUndoRedo);
 		
 		if(bSetupUndoRedo)
 		{
@@ -11712,10 +11712,23 @@ URigVMTemplateNode* URigVMController::AddTemplateNode(const FName& InNotation, c
 	URigVMTemplateNode* Node = nullptr;
 
 	// determine what kind of node we need to create
-	const UScriptStruct* PotentialUnitStruct = Template->GetPermutation(0)->Struct;
-	if(PotentialUnitStruct && PotentialUnitStruct->IsChildOf(FRigVMStruct::StaticStruct()))
+	if(const FRigVMFunction* FirstFunction = Template->GetPermutation(0))
 	{
-		Node = NewObject<URigVMUnitNode>(Graph, *Name); 
+		const UScriptStruct* PotentialUnitStruct = FirstFunction->Struct;
+		if(PotentialUnitStruct && PotentialUnitStruct->IsChildOf(FRigVMStruct::StaticStruct()))
+		{
+			Node = NewObject<URigVMUnitNode>(Graph, *Name); 
+		}
+	}
+
+	if(Node == nullptr)
+	{
+		const FString TemplateName = Template->GetName().ToString();
+		if(TemplateName== URigVMRerouteNode::RerouteName ||
+			TemplateName== URigVMRerouteNode::RerouteArrayName)
+		{
+			Node = NewObject<URigVMRerouteNode>(Graph, *Name);
+		}
 	}
 
 	if(Node == nullptr)
@@ -14514,8 +14527,7 @@ void URigVMController::ResolveTemplateNodePins(URigVMTemplateNode* InNode, bool
 
 		if(ResolvedPermutationIndices.Num() == 1)
 		{
-			const FRigVMFunction* ResolvedFunction = InNode->GetTemplate()->GetPermutation(ResolvedPermutationIndices[0]);
-			FullyResolveTemplateNode(InNode, ResolvedFunction, bSetupUndoRedo);
+			FullyResolveTemplateNode(InNode, ResolvedPermutationIndices[0], bSetupUndoRedo);
 		}
 
 		// resolve all wires on resolved pins
@@ -14582,7 +14594,7 @@ void URigVMController::ResolveTemplateNodeMetaData(URigVMTemplateNode* InNode, b
 	}
 }
 
-bool URigVMController::FullyResolveTemplateNode(URigVMTemplateNode* InNode, const FRigVMFunction* InFunctionToResolve, bool bSetupUndoRedo)
+bool URigVMController::FullyResolveTemplateNode(URigVMTemplateNode* InNode, int32 InPermutationIndex, bool bSetupUndoRedo)
 {
 	if(bIsFullyResolvingTemplateNode)
 	{
@@ -14591,21 +14603,11 @@ bool URigVMController::FullyResolveTemplateNode(URigVMTemplateNode* InNode, cons
 	TGuardValue<bool> ReentryGuard(bIsFullyResolvingTemplateNode, true);
 	
 	check(InNode);
-	check(InFunctionToResolve);
+	check(InPermutationIndex != INDEX_NONE);
 
 	const FRigVMTemplate* Template = InNode->GetTemplate();
-
-	const int32 PermutationIndex = Template->FindPermutation(InFunctionToResolve);
-	if(PermutationIndex == INDEX_NONE)
-	{
-		ReportErrorf(TEXT("Template %s doesn't support function %s::%s."),
-			*Template->GetNotation().ToString(),
-			*InFunctionToResolve->GetName()
-		);
-		return false;
-	}
-
-	const TArray<int32> PermutationIndices = {PermutationIndex};
+	const FRigVMFunction* ResolvedFunction = Template->GetPermutation(InPermutationIndex);
+	const TArray<int32> PermutationIndices = {InPermutationIndex};
 
 	// find all existing pins that we may need to change
 	TArray<FRigVMTemplateArgument> MissingPins;
@@ -14633,35 +14635,41 @@ bool URigVMController::FullyResolveTemplateNode(URigVMTemplateNode* InNode, cons
 	}
 
 	// find all missing pins which are not arguments for the template
-	for (TFieldIterator<FProperty> It(InFunctionToResolve->Struct); It; ++It)
+	if(ResolvedFunction)
 	{
-		const FRigVMTemplateArgument ExpectedArgument(*It);
-		const FRigVMTemplateArgument::FType ExpectedType = ExpectedArgument.GetSupportedTypes()[0];
-
-		if(URigVMPin* Pin = InNode->FindPin(It->GetFName().ToString()))
+		for (TFieldIterator<FProperty> It(ResolvedFunction->Struct); It; ++It)
 		{
-			if(Pin->GetCPPType() != ExpectedType.CPPType)
+			const FRigVMTemplateArgument ExpectedArgument(*It);
+			const FRigVMTemplateArgument::FType ExpectedType = ExpectedArgument.GetSupportedTypes()[0];
+
+			if(URigVMPin* Pin = InNode->FindPin(It->GetFName().ToString()))
 			{
-				PinTypesToChange.Add(Pin, ExpectedType);
+				if(Pin->GetCPPType() != ExpectedType.CPPType)
+				{
+					PinTypesToChange.Add(Pin, ExpectedType);
+				}
+			}
+			else
+			{
+				MissingPins.Add(ExpectedArgument);
 			}
-		}
-		else
-		{
-			MissingPins.Add(ExpectedArgument);
 		}
 	}
 
 	// find all pins which don't have a matching arg on the function
-	for(URigVMPin* Pin : InNode->GetPins())
+	if(ResolvedFunction)
 	{
-		if(InFunctionToResolve->Struct->FindPropertyByName(Pin->GetFName()) == nullptr)
+		for(URigVMPin* Pin : InNode->GetPins())
 		{
-			PinsToRemove.Add(Pin);
+			if(ResolvedFunction->Struct->FindPropertyByName(Pin->GetFName()) == nullptr)
+			{
+				PinsToRemove.Add(Pin);
+			}
 		}
-	}
 
-	// update the cached resolved function name
-	InNode->ResolvedFunctionName = InFunctionToResolve->Name;
+		// update the cached resolved function name
+		InNode->ResolvedFunctionName = ResolvedFunction->Name;
+	}
 
 	// exit out early if there's nothing to do
 	if(PinTypesToChange.IsEmpty() && MissingPins.IsEmpty() && PinsToRemove.IsEmpty())
@@ -14718,20 +14726,23 @@ bool URigVMController::FullyResolveTemplateNode(URigVMTemplateNode* InNode, cons
 	}
 
 	// add missing pins
-	for(const FRigVMTemplateArgument& MissingPin : MissingPins)
+	if(ResolvedFunction)
 	{
-		check(MissingPin.GetDirection() == ERigVMPinDirection::Hidden);
-		
-		FProperty* Property = InFunctionToResolve->Struct->FindPropertyByName(MissingPin.GetName());
-		check(Property);
+		for(const FRigVMTemplateArgument& MissingPin : MissingPins)
+		{
+			check(MissingPin.GetDirection() == ERigVMPinDirection::Hidden);
+			
+			FProperty* Property = ResolvedFunction->Struct->FindPropertyByName(MissingPin.GetName());
+			check(Property);
 
-		URigVMPin* Pin = NewObject<URigVMPin>(Cast<UObject>(InNode), MissingPin.GetName());
-		ConfigurePinFromProperty(Property, Pin, MissingPin.GetDirection());
+			URigVMPin* Pin = NewObject<URigVMPin>(Cast<UObject>(InNode), MissingPin.GetName());
+			ConfigurePinFromProperty(Property, Pin, MissingPin.GetDirection());
 
-		AddNodePin(InNode, Pin);
-		Notify(ERigVMGraphNotifType::PinAdded, Pin);
+			AddNodePin(InNode, Pin);
+			Notify(ERigVMGraphNotifType::PinAdded, Pin);
 
-		// we don't need to set the default value here since the pin is hidden
+			// we don't need to set the default value here since the pin is hidden
+		}
 	}
 
 	if (UnitNodeCreatedContext.IsValid())
diff --git a/Engine/Source/Developer/RigVMDeveloper/Private/RigVMModel/RigVMPin.cpp b/Engine/Source/Developer/RigVMDeveloper/Private/RigVMModel/RigVMPin.cpp
--- a/Engine/Source/Developer/RigVMDeveloper/Private/RigVMModel/RigVMPin.cpp
+++ b/Engine/Source/Developer/RigVMDeveloper/Private/RigVMModel/RigVMPin.cpp
@@ -1673,11 +1673,11 @@ bool URigVMPin::CanLink(URigVMPin* InSourcePin, URigVMPin* InTargetPin, FString*
 					return false;
 				};
 				
-				if(IsPinValidForTypeChange(InSourcePin, false, InUserLinkDirection) && !InTargetPin->IsExecuteContext())
+				if(IsPinValidForTypeChange(InSourcePin, false, InUserLinkDirection))
 				{
 					return TemplateNodeSupportsType(InSourcePin, InTargetPin->GetCPPType(), OutFailureReason);
 				}
-				else if(IsPinValidForTypeChange(InTargetPin, true, InUserLinkDirection) && !InSourcePin->IsExecuteContext())
+				else if(IsPinValidForTypeChange(InTargetPin, true, InUserLinkDirection))
 				{
 					return TemplateNodeSupportsType(InTargetPin, InSourcePin->GetCPPType(), OutFailureReason);
 				}
diff --git a/Engine/Source/Developer/RigVMDeveloper/Public/RigVMModel/Nodes/RigVMAggregateNode.h b/Engine/Source/Developer/RigVMDeveloper/Public/RigVMModel/Nodes/RigVMAggregateNode.h
--- a/Engine/Source/Developer/RigVMDeveloper/Public/RigVMModel/Nodes/RigVMAggregateNode.h
+++ b/Engine/Source/Developer/RigVMDeveloper/Public/RigVMModel/Nodes/RigVMAggregateNode.h
@@ -42,7 +42,7 @@ public:
 
 private:
 
-	void InvalidateCache();
+	virtual void InvalidateCache() override;
 
 	mutable URigVMNode* FirstInnerNodeCache;
 	mutable URigVMNode* LastInnerNodeCache;
diff --git a/Engine/Source/Developer/RigVMDeveloper/Public/RigVMModel/Nodes/RigVMRerouteNode.h b/Engine/Source/Developer/RigVMDeveloper/Public/RigVMModel/Nodes/RigVMRerouteNode.h
--- a/Engine/Source/Developer/RigVMDeveloper/Public/RigVMModel/Nodes/RigVMRerouteNode.h
+++ b/Engine/Source/Developer/RigVMDeveloper/Public/RigVMModel/Nodes/RigVMRerouteNode.h
@@ -2,7 +2,7 @@
 
 #pragma once
 
-#include "RigVMModel/RigVMNode.h"
+#include "RigVMModel/Nodes/RigVMTemplateNode.h"
 #include "RigVMRerouteNode.generated.h"
 
 /**
@@ -13,7 +13,7 @@
  * displayed as full nodes or as small circles.
  */
 UCLASS(BlueprintType)
-class RIGVMDEVELOPER_API URigVMRerouteNode : public URigVMNode
+class RIGVMDEVELOPER_API URigVMRerouteNode : public URigVMTemplateNode
 {
 	GENERATED_BODY()
 
@@ -32,9 +32,15 @@ public:
 
 	virtual FLinearColor GetNodeColor() const override;
 
+	virtual FName GetNotation() const override;
+	virtual const FRigVMTemplate* GetTemplate() const override;
+
 private:
 
+	static const FRigVMTemplate* FindOrAddTemplate(bool bIsArray);
+
 	static const FString RerouteName;
+	static const FString RerouteArrayName;
 	static const FString ValueName;
 
 	UPROPERTY()
@@ -43,5 +49,6 @@ private:
 	friend class URigVMController;
 	friend class URigVMCompiler;
 	friend class FRigVMParserAST;
+	friend class FRigVMDeveloperModule;
 };
 
diff --git a/Engine/Source/Developer/RigVMDeveloper/Public/RigVMModel/Nodes/RigVMTemplateNode.h b/Engine/Source/Developer/RigVMDeveloper/Public/RigVMModel/Nodes/RigVMTemplateNode.h
--- a/Engine/Source/Developer/RigVMDeveloper/Public/RigVMModel/Nodes/RigVMTemplateNode.h
+++ b/Engine/Source/Developer/RigVMDeveloper/Public/RigVMModel/Nodes/RigVMTemplateNode.h
@@ -34,7 +34,7 @@ public:
 
 	// Returns the notation of the node
 	UFUNCTION(BlueprintPure, Category = Template)
-	FName GetNotation() const;
+	virtual FName GetNotation() const;
 
 	// returns true if a pin supports a given type
 	bool SupportsType(const URigVMPin* InPin, const FString& InCPPType, FString* OutCPPType = nullptr);
@@ -49,7 +49,7 @@ public:
 	TArray<const FRigVMFunction*> GetResolvedPermutations(FRigVMTemplate::FTypeMap* OutTypes = nullptr) const;
 
 	// returns the template used for this node
-	const FRigVMTemplate* GetTemplate() const;
+	virtual const FRigVMTemplate* GetTemplate() const;
 
 	// returns the type map for the currently resolved pins
 	FRigVMTemplate::FTypeMap GetResolvedTypes() const;
@@ -73,7 +73,7 @@ public:
 
 protected:
 
-	void InvalidateCache();
+	virtual void InvalidateCache() override;
 
 	UPROPERTY()
 	FName TemplateNotation;
diff --git a/Engine/Source/Developer/RigVMDeveloper/Public/RigVMModel/RigVMController.h b/Engine/Source/Developer/RigVMDeveloper/Public/RigVMModel/RigVMController.h
--- a/Engine/Source/Developer/RigVMDeveloper/Public/RigVMModel/RigVMController.h
+++ b/Engine/Source/Developer/RigVMDeveloper/Public/RigVMModel/RigVMController.h
@@ -1054,7 +1054,7 @@ private:
 	bool ResolveWildCardPinImpl(URigVMPin* InPinToResolve, const FString& InCPPType, UObject* InCPPTypeObject, bool bSetupUndoRedo, bool bTraverseNode = true, bool bTraverseParentPins = true, bool bTraverseLinks = true, bool bAllowFloatingPointCasts = false);
 	void ResolveTemplateNodePins(URigVMTemplateNode* InNode, bool bSetupUndoRedo);
 	void ResolveTemplateNodeMetaData(URigVMTemplateNode* InNode, bool bSetupUndoRedo);
-	bool FullyResolveTemplateNode(URigVMTemplateNode* InNode, const FRigVMFunction* InFunctionToResolve, bool bSetupUndoRedo);
+	bool FullyResolveTemplateNode(URigVMTemplateNode* InNode, int32 InPermutationIndex, bool bSetupUndoRedo);
 	bool ChangePinType(const FString& InPinPath, const FString& InCPPType, const FName& InCPPTypeObjectPath, bool bSetupUndoRedo, bool bSetupOrphanPins = true, bool bBreakLinks = true, bool bRemoveSubPins = true);
 	bool ChangePinType(URigVMPin* InPin, const FString& InCPPType, const FName& InCPPTypeObjectPath, bool bSetupUndoRedo, bool bSetupOrphanPins = true, bool bBreakLinks = true, bool bRemoveSubPins = true);
 	bool ChangePinType(URigVMPin* InPin, const FString& InCPPType,UObject* InCPPTypeObject, bool bSetupUndoRedo, bool bSetupOrphanPins = true, bool bBreakLinks = true, bool bRemoveSubPins = true);
diff --git a/Engine/Source/Developer/RigVMDeveloper/Public/RigVMModel/RigVMNode.h b/Engine/Source/Developer/RigVMDeveloper/Public/RigVMModel/RigVMNode.h
--- a/Engine/Source/Developer/RigVMDeveloper/Public/RigVMModel/RigVMNode.h
+++ b/Engine/Source/Developer/RigVMDeveloper/Public/RigVMModel/RigVMNode.h
@@ -249,6 +249,7 @@ private:
 
 protected:
 
+	virtual void InvalidateCache() {}
 	virtual TArray<int32> GetInstructionsForVMImpl(URigVM* InVM, const FRigVMASTProxy& InProxy = FRigVMASTProxy()) const; 
 	virtual FText GetToolTipTextForPin(const URigVMPin* InPin) const;
 	virtual bool AllowsLinksOn(const URigVMPin* InPin) const { return true; }
diff --git a/Engine/Source/Runtime/RigVM/Private/RigVMCore/RigVMRegistry.cpp b/Engine/Source/Runtime/RigVM/Private/RigVMCore/RigVMRegistry.cpp
--- a/Engine/Source/Runtime/RigVM/Private/RigVMCore/RigVMRegistry.cpp
+++ b/Engine/Source/Runtime/RigVM/Private/RigVMCore/RigVMRegistry.cpp
@@ -118,7 +118,7 @@ const TArray<FRigVMTemplate>& FRigVMRegistry::GetTemplates() const
 	return Templates;
 }
 
-const FRigVMTemplate* FRigVMRegistry::GetOrAddTemplateFromArguments(const FName& InName, const TArray<FRigVMTemplateArgument>& InArguments)
+const FRigVMTemplate* FRigVMRegistry::GetOrAddTemplateFromArguments(const FName& InName, const TArray<FRigVMTemplateArgument>& InArguments, bool bAllowExecuteContext)
 {
 	FRigVMTemplate Template(InName, InArguments, INDEX_NONE);
 	if(const FRigVMTemplate* ExistingTemplate = FindTemplate(Template.GetNotation()))
@@ -142,7 +142,7 @@ const FRigVMTemplate* FRigVMRegistry::GetOrAddTemplateFromArguments(const FName&
 	{
 		if(Argument.Types[0].IsWildCard())
 		{
-			static TArray<FRigVMTemplateArgument::FType> AllTypes, AllArrayTypes;
+			static TArray<FRigVMTemplateArgument::FType> AllTypes, AllArrayTypes, ExecuteContextTypes;
 			if(AllTypes.IsEmpty())
 			{
 				AllTypes.Add(FRigVMTemplateArgument::FType(RigVMTypeUtils::BoolType));
@@ -163,17 +163,17 @@ const FRigVMTemplate* FRigVMRegistry::GetOrAddTemplateFromArguments(const FName&
 
 				struct FTypeTraverser
 				{
-					static EObjectFlags DisallowedFlags()
+					EObjectFlags DisallowedFlags()
 					{
 						return RF_BeginDestroyed | RF_FinishDestroyed;
 					}
 
-					static EObjectFlags NeededFlags()
+					EObjectFlags NeededFlags()
 					{
 						return RF_Public;
 					}
 					
-					static bool IsAllowedType(const FProperty* InProperty)
+					bool IsAllowedType(const FProperty* InProperty)
 					{
 						if(!InProperty->HasAnyPropertyFlags(
 							CPF_BlueprintVisible |
@@ -225,21 +225,17 @@ const FRigVMTemplate* FRigVMRegistry::GetOrAddTemplateFromArguments(const FName&
 						return false;
 					}
 
-					static bool IsAllowedType(const UEnum* InEnum)
+					bool IsAllowedType(const UEnum* InEnum)
 					{
 						return !InEnum->HasAnyFlags(DisallowedFlags()) && InEnum->HasAllFlags(NeededFlags());
 					}
 
-					static bool IsAllowedType(const UStruct* InStruct)
+					bool IsAllowedType(const UStruct* InStruct)
 					{
 						if(InStruct->HasAnyFlags(DisallowedFlags()) || !InStruct->HasAllFlags(NeededFlags()))
 						{
 							return false;
 						}
-						if(InStruct->IsChildOf(FRigVMExecuteContext::StaticStruct()))
-						{
-							return false;
-						}
 						if(InStruct->IsChildOf(FRigVMStruct::StaticStruct()))
 						{
 							return false;
@@ -254,7 +250,7 @@ const FRigVMTemplate* FRigVMRegistry::GetOrAddTemplateFromArguments(const FName&
 						return true;
 					}
 
-					static bool IsAllowedType(const UClass* InClass)
+					bool IsAllowedType(const UClass* InClass)
 					{
 						if(InClass->HasAnyClassFlags(CLASS_Hidden | CLASS_Abstract))
 						{
@@ -267,24 +263,34 @@ const FRigVMTemplate* FRigVMRegistry::GetOrAddTemplateFromArguments(const FName&
 					}
 				};
 
+				FTypeTraverser Traverser;
+
 				// add all structs
 				for (TObjectIterator<UScriptStruct> ScriptIt; ScriptIt; ++ScriptIt)
 				{
 					UScriptStruct* ScriptStruct = *ScriptIt;
-					if(!FTypeTraverser::IsAllowedType(ScriptStruct))
+					if(!Traverser.IsAllowedType(ScriptStruct))
 					{
 						continue;
 					}
 					const FString CPPType = ScriptStruct->GetStructCPPName();
-					AllTypes.Add(FRigVMTemplateArgument::FType(CPPType, ScriptStruct));
-					AllArrayTypes.Add(FRigVMTemplateArgument::FType(RigVMTypeUtils::ArrayTypeFromBaseType(CPPType), ScriptStruct));
+
+					if(ScriptStruct->IsChildOf(FRigVMExecuteContext::StaticStruct()))
+					{
+						ExecuteContextTypes.Add(FRigVMTemplateArgument::FType(CPPType, ScriptStruct));
+					}
+					else
+					{
+						AllTypes.Add(FRigVMTemplateArgument::FType(CPPType, ScriptStruct));
+						AllArrayTypes.Add(FRigVMTemplateArgument::FType(RigVMTypeUtils::ArrayTypeFromBaseType(CPPType), ScriptStruct));
+					}
 				}
 
 				// add all enums
 				for (TObjectIterator<UEnum> EnumIt; EnumIt; ++EnumIt)
 				{
 					UEnum* Enum = (*EnumIt);
-					if(!FTypeTraverser::IsAllowedType(Enum))
+					if(!Traverser.IsAllowedType(Enum))
 					{
 						continue;
 					}
@@ -297,7 +303,7 @@ const FRigVMTemplate* FRigVMRegistry::GetOrAddTemplateFromArguments(const FName&
 				for (TObjectIterator<UClass> ClassIt; ClassIt; ++ClassIt)
 				{
 					UClass* Class = *ClassIt;
-					if(!FTypeTraverser::IsAllowedType(Class))
+					if(!Traverser.IsAllowedType(Class))
 					{
 						continue;
 					}
@@ -308,7 +314,19 @@ const FRigVMTemplate* FRigVMRegistry::GetOrAddTemplateFromArguments(const FName&
 				}
 			}
 
-			Argument.Types = Argument.Types[0].IsArray() ? AllArrayTypes : AllTypes;
+			if(Argument.Types[0].IsArray())
+			{
+				Argument.Types = AllArrayTypes;
+			}
+			else
+			{
+				Argument.Types = AllTypes;
+				if(bAllowExecuteContext)
+				{
+					Argument.Types.Append(ExecuteContextTypes);
+				}
+			}
+			
 			Argument.bSingleton = false;
 			NumPermutations = Argument.Types.Num(); 
 		}
diff --git a/Engine/Source/Runtime/RigVM/Private/RigVMCore/RigVMTemplate.cpp b/Engine/Source/Runtime/RigVM/Private/RigVMCore/RigVMTemplate.cpp
--- a/Engine/Source/Runtime/RigVM/Private/RigVMCore/RigVMTemplate.cpp
+++ b/Engine/Source/Runtime/RigVM/Private/RigVMCore/RigVMTemplate.cpp
@@ -317,6 +317,10 @@ FName FRigVMTemplate::GetName() const
 	{
 		return *Left;
 	}
+	if (GetNotation().ToString().Split(TEXT("("), &Left, nullptr))
+	{
+		return *Left;
+	}
 	return NAME_None;
 }
 
@@ -360,6 +364,11 @@ FLinearColor FRigVMTemplate::GetColor(const TArray<int32>& InPermutationIndices)
 		FString NodeColorMetadata;
 
 		const FRigVMFunction* ResolvedFunction = GetPermutation(InPermutationIndex);
+		if(ResolvedFunction == nullptr)
+		{
+			return false;
+		}
+
 		ResolvedFunction->Struct->GetStringMetaDataHierarchical(NodeColorName, &NodeColorMetadata);
 		if (!NodeColorMetadata.IsEmpty())
 		{
@@ -411,6 +420,11 @@ FText FRigVMTemplate::GetTooltipText(const TArray<int32>& InPermutationIndices)
 	auto VisitPermutation = [&ResolvedTooltipText, this](int32 InPermutationIndex) -> bool
 	{
 		const FRigVMFunction* ResolvedFunction = GetPermutation(InPermutationIndex);
+		if(ResolvedFunction == nullptr)
+		{
+			return false;
+		}
+		
 		const FText TooltipText = ResolvedFunction->Struct->GetToolTipText();
 		
 		if (!ResolvedTooltipText.IsEmpty())
@@ -461,6 +475,11 @@ FText FRigVMTemplate::GetDisplayNameForArgument(const FName& InArgumentName, con
 		auto VisitPermutation = [InArgumentName, &ResolvedDisplayName, this](int32 InPermutationIndex) -> bool
 		{
 			const FRigVMFunction* ResolvedFunction = GetPermutation(InPermutationIndex);
+			if(ResolvedFunction == nullptr)
+			{
+				return false;
+			}
+			
 			if(const FProperty* Property = ResolvedFunction->Struct->FindPropertyByName(InArgumentName))
 			{
 				const FText DisplayName = Property->GetDisplayNameText();
diff --git a/Engine/Source/Runtime/RigVM/Public/RigVMCore/RigVMRegistry.h b/Engine/Source/Runtime/RigVM/Public/RigVMCore/RigVMRegistry.h
--- a/Engine/Source/Runtime/RigVM/Public/RigVMCore/RigVMRegistry.h
+++ b/Engine/Source/Runtime/RigVM/Public/RigVMCore/RigVMRegistry.h
@@ -45,7 +45,7 @@ public:
 	const TArray<FRigVMTemplate>& GetTemplates() const;
 
 	// Defines and retrieves a template given its arguments
-	const FRigVMTemplate* GetOrAddTemplateFromArguments(const FName& InName, const TArray<FRigVMTemplateArgument>& InArguments);
+	const FRigVMTemplate* GetOrAddTemplateFromArguments(const FName& InName, const TArray<FRigVMTemplateArgument>& InArguments, bool bIncludeExecuteContext = false);
 
 private:
 
