RigVM: Support array type changes while maintaining element pins

diff --git a/Engine/Source/Developer/RigVMDeveloper/Private/RigVMModel/RigVMController.cpp b/Engine/Source/Developer/RigVMDeveloper/Private/RigVMModel/RigVMController.cpp
--- a/Engine/Source/Developer/RigVMDeveloper/Private/RigVMModel/RigVMController.cpp
+++ b/Engine/Source/Developer/RigVMDeveloper/Private/RigVMModel/RigVMController.cpp
@@ -1726,6 +1726,37 @@ bool URigVMController::UnresolveTemplateNodes(const TArray<URigVMTemplateNode*>&
 		FRigVMTemplate::FTypeMap Types;
 		Node->GetTemplate()->FullyResolve(Types, PermutationIndex);
 
+		// if the types' containers don't match the current pin containers try again
+		bool bTypesWereAdjusted = false;
+		FRigVMTemplate::FTypeMap AdjustedTypes = Types;
+		for(TPair<FName, FRigVMTemplateArgument::FType>& TypePair : AdjustedTypes)
+		{
+			if(URigVMPin* Pin = Node->FindPin(TypePair.Key.ToString()))
+			{
+				if(Pin->IsArray() != TypePair.Value.IsArray())
+				{
+					if(Pin->IsArray())
+					{
+						TypePair.Value.CPPType = RigVMTypeUtils::ArrayTypeFromBaseType(TypePair.Value.CPPType);
+					}
+					else
+					{
+						TypePair.Value.CPPType = RigVMTypeUtils::BaseTypeFromArrayType(TypePair.Value.CPPType);
+					}
+					bTypesWereAdjusted = true;
+				}
+			}
+		}
+
+		if(bTypesWereAdjusted)
+		{
+			Node->GetTemplate()->FullyResolve(AdjustedTypes, PermutationIndex);
+			if(AdjustedTypes.Num() == Types.Num())
+			{
+				Types = AdjustedTypes;
+			}
+		}
+
 		for(const TPair<FName, FRigVMTemplateArgument::FType>& TypePair : Types)
 		{
 			const FName& PinName = TypePair.Key;
@@ -1742,9 +1773,21 @@ bool URigVMController::UnresolveTemplateNodes(const TArray<URigVMTemplateNode*>&
 				return false;
 			}
 
-			if(ChangePinType(Pin, ExpectedType.CPPType, ExpectedType.CPPTypeObject, bSetupUndoRedo, false, true, true))
+			if(ChangePinType(Pin, ExpectedType.CPPType, ExpectedType.CPPTypeObject, bSetupUndoRedo, false, true, !Pin->IsArray()))
 			{
 				bChangedAnyPin = true;
+
+				if(Pin->IsArray())
+				{
+					FRigVMTemplateArgument::FType ElementType = ExpectedType;
+					ElementType.CPPType = RigVMTypeUtils::BaseTypeFromArrayType(ElementType.CPPType);
+					
+					TArray<URigVMPin*> SubPins = Pin->GetSubPins();
+					for (URigVMPin* SubPin : SubPins)
+					{
+						ChangePinType(SubPin, ElementType.CPPType, ElementType.CPPTypeObject, bSetupUndoRedo, false, true, true);
+					}
+				}
 			}
 		}
 
@@ -14993,12 +15036,17 @@ bool URigVMController::ResolveWildCardPinImpl(URigVMPin* InPinToResolve, const F
 	{
 		for(URigVMPin* SubPin : InPinToResolve->GetSubPins())
 		{
+			const FString BaseCPPType = RigVMTypeUtils::IsArrayType(CPPType) ? RigVMTypeUtils::BaseTypeFromArrayType(CPPType) : CPPType;
+			const FString TypeToResolveAgainst = SubPin->IsArray() ? RigVMTypeUtils::ArrayTypeFromBaseType(BaseCPPType) : BaseCPPType;
+			
 			if(SubPin->IsWildCard())
 			{
-				const FString BaseCPPType = RigVMTypeUtils::IsArrayType(CPPType) ? RigVMTypeUtils::BaseTypeFromArrayType(CPPType) : CPPType;
-				const FString TypeToResolveAgainst = SubPin->IsArray() ? RigVMTypeUtils::ArrayTypeFromBaseType(BaseCPPType) : BaseCPPType; 
 				ResolveWildCardPinImpl(SubPin, TypeToResolveAgainst, CPPTypeObject, bSetupUndoRedo, false, false, bTraverseLinks);
 			}
+			else
+			{
+				ChangePinType(SubPin, TypeToResolveAgainst, InCPPTypeObject, bSetupUndoRedo, false, true, true);
+			}
 		}
 	}
 
diff --git a/Engine/Source/Runtime/RigVM/Private/RigVMCore/RigVMTemplate.cpp b/Engine/Source/Runtime/RigVM/Private/RigVMCore/RigVMTemplate.cpp
--- a/Engine/Source/Runtime/RigVM/Private/RigVMCore/RigVMTemplate.cpp
+++ b/Engine/Source/Runtime/RigVM/Private/RigVMCore/RigVMTemplate.cpp
@@ -1001,8 +1001,21 @@ bool FRigVMTemplate::Resolve(FTypeMap& InOutTypes, TArray<int32>& OutPermutation
 				continue;
 			}
 		}
-		
-		if(Argument.GetArrayType() == FRigVMTemplateArgument::EArrayType_ArrayValue)
+
+		const FRigVMTemplateArgument::EArrayType ArrayType = Argument.GetArrayType();
+		if(ArrayType == FRigVMTemplateArgument::EArrayType_Mixed)
+		{
+			InOutTypes.Add(Argument.Name, FRigVMTemplateArgument::FType());
+
+			if(const FRigVMTemplateArgument::FType* InputType = InputTypes.Find(Argument.Name))
+			{
+				if(InputType->IsArray())
+				{
+					InOutTypes.FindChecked(Argument.Name) = FRigVMTemplateArgument::FType::Array();
+				}
+			}
+		}
+		else if(ArrayType == FRigVMTemplateArgument::EArrayType_ArrayValue)
 		{
 			InOutTypes.Add(Argument.Name, FRigVMTemplateArgument::FType::Array());
 		}
