Control Rig: Only display workflows on desired nodes

diff --git a/Engine/Source/Developer/RigVMDeveloper/Private/RigVMUserWorkflowRegistry.cpp b/Engine/Source/Developer/RigVMDeveloper/Private/RigVMUserWorkflowRegistry.cpp
--- a/Engine/Source/Developer/RigVMDeveloper/Private/RigVMUserWorkflowRegistry.cpp
+++ b/Engine/Source/Developer/RigVMDeveloper/Private/RigVMUserWorkflowRegistry.cpp
@@ -35,7 +35,7 @@ TArray<FRigVMUserWorkflow> URigVMUserWorkflowRegistry::GetWorkflows(ERigVMUserWo
 
 	for(const TTuple<int32, const UScriptStruct*, FRigVMUserWorkflowProvider>& Provider : Providers)
 	{
-		if(Provider.Get<1>() == InStruct || Provider.Get<1>() == nullptr || InStruct == nullptr)
+		if(Provider.Get<1>() == InStruct || Provider.Get<1>() == nullptr)
 		{
 			Workflows.Append(Provider.Get<2>().Execute(InSubject));
 		}
