RigVM: If, select and array nodes to inherit template node

diff --git a/Engine/Plugins/Animation/ControlRig/Source/ControlRigDeveloper/Private/ControlRigBlueprint.cpp b/Engine/Plugins/Animation/ControlRig/Source/ControlRigDeveloper/Private/ControlRigBlueprint.cpp
--- a/Engine/Plugins/Animation/ControlRig/Source/ControlRigDeveloper/Private/ControlRigBlueprint.cpp
+++ b/Engine/Plugins/Animation/ControlRig/Source/ControlRigDeveloper/Private/ControlRigBlueprint.cpp
@@ -112,6 +112,9 @@ UControlRigBlueprint::UControlRigBlueprint(const FObjectInitializer& ObjectIniti
 
 	Model->SetDefaultFunctionLibrary(FunctionLibrary);
 
+	Model->SetExecuteContextStruct(FControlRigExecuteContext::StaticStruct());
+	FunctionLibrary->SetExecuteContextStruct(FControlRigExecuteContext::StaticStruct());
+
 	Validator = ObjectInitializer.CreateDefaultSubobject<UControlRigValidator>(this, TEXT("ControlRigValidator"));
 
 	DebugBoneRadius = 1.f;
@@ -146,6 +149,8 @@ void UControlRigBlueprint::InitializeModelIfRequired(bool bRecompileVM)
 
 	if (Controllers.Num() == 0)
 	{
+		Model->SetExecuteContextStruct(FControlRigExecuteContext::StaticStruct());
+		FunctionLibrary->SetExecuteContextStruct(FControlRigExecuteContext::StaticStruct());
 		GetOrCreateController(Model);
 		GetOrCreateController(FunctionLibrary);
 
@@ -1427,7 +1432,6 @@ URigVMController* UControlRigBlueprint::GetOrCreateController(URigVMGraph* InGra
 	}
 
 	URigVMController* Controller = NewObject<URigVMController>(this);
-	Controller->SetExecuteContextStruct(FControlRigExecuteContext::StaticStruct());
 	Controller->SetGraph(InGraph);
 	Controller->OnModified().AddUObject(this, &UControlRigBlueprint::HandleModifiedEvent);
 
@@ -1856,6 +1860,7 @@ URigVMGraph* UControlRigBlueprint::GetTemplateModel()
 	{
 		TemplateModel = NewObject<URigVMGraph>(this, TEXT("TemplateModel"));
 		TemplateModel->SetFlags(RF_Transient);
+		TemplateModel->SetExecuteContextStruct(FControlRigExecuteContext::StaticStruct());
 	}
 	return TemplateModel;
 #else
@@ -1869,7 +1874,6 @@ URigVMController* UControlRigBlueprint::GetTemplateController()
 	if (TemplateController == nullptr)
 	{
 		TemplateController = NewObject<URigVMController>(this, TEXT("TemplateController"));
-		TemplateController->SetExecuteContextStruct(FControlRigExecuteContext::StaticStruct());
 		TemplateController->SetGraph(GetTemplateModel());
 		TemplateController->EnableReporting(false);
 		TemplateController->SetFlags(RF_Transient);
diff --git a/Engine/Plugins/Animation/ControlRig/Source/ControlRigEditor/Private/Graph/SControlRigGraphChangePinType.cpp b/Engine/Plugins/Animation/ControlRig/Source/ControlRigEditor/Private/Graph/SControlRigGraphChangePinType.cpp
--- a/Engine/Plugins/Animation/ControlRig/Source/ControlRigEditor/Private/Graph/SControlRigGraphChangePinType.cpp
+++ b/Engine/Plugins/Animation/ControlRig/Source/ControlRigEditor/Private/Graph/SControlRigGraphChangePinType.cpp
@@ -300,6 +300,11 @@ void SControlRigChangePinType::FillPinTypeMenu(FMenuBuilder& MenuBuilder)
 		for(int32 TypeIndex=0; TypeIndex < SortedTypes.Num(); TypeIndex++)
 		{
 			const FRigVMTemplateArgument::FType& Type = SortedTypes[TypeIndex].Key;
+			if (Type.CPPTypeObject != nullptr && !IsValid(Type.CPPTypeObject))
+			{
+				continue;
+			}
+			
 			const FEdGraphPinType PinType = RigVMTypeUtils::PinTypeFromCPPType(Type.CPPType, Type.CPPTypeObject);
 
 			if(bHasAllTypes && Type.CPPTypeObject)
diff --git a/Engine/Plugins/Experimental/DataInterfaceGraph/Source/DataInterfaceGraphUncookedOnly/Private/DataInterfaceGraph_EditorData.cpp b/Engine/Plugins/Experimental/DataInterfaceGraph/Source/DataInterfaceGraphUncookedOnly/Private/DataInterfaceGraph_EditorData.cpp
--- a/Engine/Plugins/Experimental/DataInterfaceGraph/Source/DataInterfaceGraphUncookedOnly/Private/DataInterfaceGraph_EditorData.cpp
+++ b/Engine/Plugins/Experimental/DataInterfaceGraph/Source/DataInterfaceGraphUncookedOnly/Private/DataInterfaceGraph_EditorData.cpp
@@ -18,6 +18,9 @@ UDataInterfaceGraph_EditorData::UDataInterfaceGraph_EditorData(const FObjectInit
 	RigVMGraph = ObjectInitializer.CreateDefaultSubobject<URigVMGraph>(this, TEXT("RigVMGraph"));
 	RigVMFunctionLibrary = ObjectInitializer.CreateDefaultSubobject<URigVMFunctionLibrary>(this, TEXT("RigVMFunctionLibrary"));
 	RigVMGraph->SetDefaultFunctionLibrary(RigVMFunctionLibrary);
+
+	RigVMGraph->SetExecuteContextStruct(FRigVMExecuteContext::StaticStruct());
+	RigVMFunctionLibrary->SetExecuteContextStruct(FRigVMExecuteContext::StaticStruct());
 	
 	auto MakeEdGraph = [this, &ObjectInitializer](FName InName) -> UDataInterfaceGraph_EdGraph*
 	{
@@ -57,6 +60,9 @@ void UDataInterfaceGraph_EditorData::Initialize(bool bRecompileVM)
 		}
 	}
 
+	RigVMGraph->SetExecuteContextStruct(FRigVMExecuteContext::StaticStruct());
+	RigVMFunctionLibrary->SetExecuteContextStruct(FRigVMExecuteContext::StaticStruct());
+	
 	RootGraph->Initialize(this);
 	FunctionLibraryEdGraph->Initialize(this);
 	if(EntryPointGraph)
@@ -181,7 +187,6 @@ URigVMController* UDataInterfaceGraph_EditorData::GetOrCreateRigVMController(URi
 	}
 
 	URigVMController* Controller = NewObject<URigVMController>(this);
-	Controller->SetExecuteContextStruct(FRigVMExecuteContext::StaticStruct());
 	Controller->SetGraph(InRigVMGraph);
 	Controller->OnModified().AddUObject(this, &UDataInterfaceGraph_EditorData::HandleModifiedEvent);
 
diff --git a/Engine/Source/Developer/RigVMDeveloper/Private/RigVMModel/Nodes/RigVMArrayNode.cpp b/Engine/Source/Developer/RigVMDeveloper/Private/RigVMModel/Nodes/RigVMArrayNode.cpp
--- a/Engine/Source/Developer/RigVMDeveloper/Private/RigVMModel/Nodes/RigVMArrayNode.cpp
+++ b/Engine/Source/Developer/RigVMDeveloper/Private/RigVMModel/Nodes/RigVMArrayNode.cpp
@@ -2,6 +2,7 @@
 
 #include "RigVMModel/Nodes/RigVMArrayNode.h"
 #include "RigVMModel/RigVMGraph.h"
+#include "RigVMCore/RigVMStruct.h"
 
 const FString URigVMArrayNode::ArrayName = TEXT("Array");
 const FString URigVMArrayNode::NumName = TEXT("Num");
@@ -131,6 +132,242 @@ FText URigVMArrayNode::GetToolTipText() const
 	return Super::GetToolTipText();
 }
 
+FName URigVMArrayNode::GetNotation() const
+{
+    const FString TemplateName = StaticEnum<ERigVMOpCode>()->GetDisplayNameTextByValue((int64)OpCode).ToString();
+	switch(OpCode)
+    {
+    	case ERigVMOpCode::ArrayReset:
+    	{
+    		static constexpr TCHAR Format[] = TEXT("%s(io %s, io %s)");
+    		const FName Notation = *FString::Printf(Format, *TemplateName, *FRigVMStruct::ExecuteContextName.ToString(), *ArrayName);
+            return Notation;
+    	}
+    	case ERigVMOpCode::ArrayGetNum:
+    	{
+    		static constexpr TCHAR Format[] = TEXT("%s(in %s, out %s)");
+    		const FName Notation = *FString::Printf(Format, *TemplateName, *ArrayName, *NumName);
+    		return Notation;
+    	} 
+    	case ERigVMOpCode::ArraySetNum:
+    	{
+    		static constexpr TCHAR Format[] = TEXT("%s(io %s, io %s, in %s)");
+    		const FName Notation = *FString::Printf(Format, *TemplateName, *FRigVMStruct::ExecuteContextName.ToString(), *ArrayName, *NumName);
+    		return Notation;
+    	}
+    	case ERigVMOpCode::ArrayGetAtIndex:
+    	{
+    		static constexpr TCHAR Format[] = TEXT("%s(in %s, in %s, out %s)");
+    		const FName Notation = *FString::Printf(Format, *TemplateName, *ArrayName, *IndexName, *ElementName);
+    		return Notation;
+    	}  
+    	case ERigVMOpCode::ArraySetAtIndex:
+    	{
+    		static constexpr TCHAR Format[] = TEXT("%s(io %s, io %s, in %s, in %s)");
+    		const FName Notation = *FString::Printf(Format, *TemplateName, *FRigVMStruct::ExecuteContextName.ToString(), *ArrayName, *IndexName, *ElementName);
+    		return Notation;
+    	}
+    	case ERigVMOpCode::ArrayAdd:
+    	{
+    		static constexpr TCHAR Format[] = TEXT("%s(io %s, io %s, in %s, out %s)");
+    		const FName Notation = *FString::Printf(Format, *TemplateName, *FRigVMStruct::ExecuteContextName.ToString(), *ArrayName, *ElementName, *IndexName);
+    		return Notation;
+    	}
+    	case ERigVMOpCode::ArrayInsert:
+    	{
+    		static constexpr TCHAR Format[] = TEXT("%s(io %s, io %s, in %s, in %s)");
+    		const FName Notation = *FString::Printf(Format, *TemplateName, *FRigVMStruct::ExecuteContextName.ToString(), *ArrayName, *IndexName, *ElementName);
+    		return Notation;
+    	}
+    	case ERigVMOpCode::ArrayRemove:
+    	{
+    		static constexpr TCHAR Format[] = TEXT("%s(io %s, io %s, in %s)");
+    		const FName Notation = *FString::Printf(Format, *TemplateName, *FRigVMStruct::ExecuteContextName.ToString(), *ArrayName, *IndexName);
+    		return Notation;
+    	}
+    	case ERigVMOpCode::ArrayFind:
+    	{
+    		static constexpr TCHAR Format[] = TEXT("%s(in %s, in %s, out %s, out %s)");
+    		const FName Notation = *FString::Printf(Format, *TemplateName, *ArrayName, *ElementName, *IndexName, *SuccessName);
+    		return Notation;
+    	}
+    	case ERigVMOpCode::ArrayAppend:
+    	{
+    		static constexpr TCHAR Format[] = TEXT("%s(io %s, io %s, in %s)");
+    		const FName Notation = *FString::Printf(Format, *TemplateName, *FRigVMStruct::ExecuteContextName.ToString(), *ArrayName, *OtherName);
+    		return Notation;
+    	}
+    	case ERigVMOpCode::ArrayClone:
+    	{
+    		static constexpr TCHAR Format[] = TEXT("%s(in %s, out %s)");
+    		const FName Notation = *FString::Printf(Format, *TemplateName, *ArrayName, *CloneName);
+    		return Notation;
+    	}
+    	case ERigVMOpCode::ArrayIterator:
+    	{
+    		static constexpr TCHAR Format[] = TEXT("%s(io %s, in %s, out %s, out %s, out %s, out %s, out %s)");
+    		const FName Notation = *FString::Printf(Format, *TemplateName, *FRigVMStruct::ExecuteContextName.ToString(),
+    			*ArrayName, *ElementName, *IndexName, *CountName, *RatioName, *CompletedName);
+    		return Notation;
+    	}
+    	case ERigVMOpCode::ArrayUnion:
+    	{
+    		static constexpr TCHAR Format[] = TEXT("%s(io %s, io %s, in %s)");
+    		const FName Notation = *FString::Printf(Format, *TemplateName, *FRigVMStruct::ExecuteContextName.ToString(), *ArrayName, *OtherName);
+    		return Notation;
+    	}
+    	case ERigVMOpCode::ArrayDifference:
+    	{
+    		static constexpr TCHAR Format[] = TEXT("%s(in %s, in %s, out %s)");
+    		const FName Notation = *FString::Printf(Format, *TemplateName, *ArrayName, *OtherName, *ResultName);
+    		return Notation;
+    	}
+    	case ERigVMOpCode::ArrayIntersection:
+    	{
+    		static constexpr TCHAR Format[] = TEXT("%s(in %s, in %s, out %s)");
+    		const FName Notation = *FString::Printf(Format, *TemplateName, *ArrayName, *OtherName, *ResultName);
+    		return Notation;
+    	}
+    	case ERigVMOpCode::ArrayReverse:
+    	{
+    		static constexpr TCHAR Format[] = TEXT("%s(io %s, io %s)");
+    		const FName Notation = *FString::Printf(Format, *TemplateName, *FRigVMStruct::ExecuteContextName.ToString(), *ArrayName);
+    		return Notation;
+    	}
+    	default:
+    	{
+    		ensure(false);
+    		break;
+    	}
+    }
+    return Super::GetNotation();
+}
+
+const FRigVMTemplate* URigVMArrayNode::GetTemplate() const
+{
+	if(const FRigVMTemplate* SuperTemplate = Super::GetTemplate())
+	{
+		return SuperTemplate;
+	}
+	
+	if(CachedTemplate == nullptr)
+	{		
+		const TArray<FRigVMTemplateArgument::FType>& SingleTypes = FRigVMTemplateArgument::GetCompatibleTypes(FRigVMTemplateArgument::ETypeCategory_SingleAnyValue);
+		const TArray<FRigVMTemplateArgument::FType>& ArrayTypes = FRigVMTemplateArgument::GetCompatibleTypes(FRigVMTemplateArgument::ETypeCategory_ArrayAnyValue);
+
+		UScriptStruct* ExecuteStruct = GetGraph()->GetExecuteContextStruct();
+		FRigVMTemplateArgument::FType ExecuteType(ExecuteStruct->GetStructCPPName(), ExecuteStruct);
+		
+		TArray<FRigVMTemplateArgument> Arguments;
+		
+		switch(OpCode)
+	    {
+    		case ERigVMOpCode::ArrayReset:
+			case ERigVMOpCode::ArrayReverse:
+    		{
+    			Arguments.Emplace(*FRigVMStruct::ExecuteContextName.ToString(), ERigVMPinDirection::IO, ExecuteType);
+    			Arguments.Emplace(*ArrayName, ERigVMPinDirection::IO, ArrayTypes);
+    			break;
+    		}
+    		case ERigVMOpCode::ArrayGetNum:
+    		{
+    			Arguments.Emplace(*ArrayName, ERigVMPinDirection::Input, ArrayTypes);
+    			Arguments.Emplace(*NumName, ERigVMPinDirection::Output, FRigVMTemplateArgument::FType(RigVMTypeUtils::Int32Type, nullptr));
+    			break;
+    		} 
+    		case ERigVMOpCode::ArraySetNum:
+    		{
+    			Arguments.Emplace(*FRigVMStruct::ExecuteContextName.ToString(), ERigVMPinDirection::IO, ExecuteType);
+    			Arguments.Emplace(*ArrayName, ERigVMPinDirection::IO, ArrayTypes);
+    			Arguments.Emplace(*NumName, ERigVMPinDirection::Input, FRigVMTemplateArgument::FType(RigVMTypeUtils::Int32Type, nullptr));
+    			break;
+    		}
+    		case ERigVMOpCode::ArrayGetAtIndex:
+    		{
+    			Arguments.Emplace(*ArrayName, ERigVMPinDirection::Input, ArrayTypes);
+    			Arguments.Emplace(*IndexName, ERigVMPinDirection::Input, FRigVMTemplateArgument::FType(RigVMTypeUtils::Int32Type, nullptr));
+    			Arguments.Emplace(*ElementName, ERigVMPinDirection::Output, SingleTypes);
+    			break;
+    		}  
+    		case ERigVMOpCode::ArraySetAtIndex:
+			case ERigVMOpCode::ArrayInsert:
+    		{
+    			Arguments.Emplace(*FRigVMStruct::ExecuteContextName.ToString(), ERigVMPinDirection::IO, ExecuteType);
+    			Arguments.Emplace(*ArrayName, ERigVMPinDirection::IO, ArrayTypes);
+    			Arguments.Emplace(*IndexName, ERigVMPinDirection::Input, FRigVMTemplateArgument::FType(RigVMTypeUtils::Int32Type, nullptr));
+    			Arguments.Emplace(*ElementName, ERigVMPinDirection::Input, SingleTypes);
+    			break;
+    		}
+    		case ERigVMOpCode::ArrayAdd:
+    		{
+    			Arguments.Emplace(*FRigVMStruct::ExecuteContextName.ToString(), ERigVMPinDirection::IO, ExecuteType);
+    			Arguments.Emplace(*ArrayName, ERigVMPinDirection::IO, ArrayTypes);
+    			Arguments.Emplace(*ElementName, ERigVMPinDirection::Input, SingleTypes);
+    			Arguments.Emplace(*IndexName, ERigVMPinDirection::Output, FRigVMTemplateArgument::FType(RigVMTypeUtils::Int32Type, nullptr));
+    			break;
+    		}
+    		case ERigVMOpCode::ArrayRemove:
+    		{
+    			Arguments.Emplace(*FRigVMStruct::ExecuteContextName.ToString(), ERigVMPinDirection::IO, ExecuteType);
+    			Arguments.Emplace(*ArrayName, ERigVMPinDirection::IO, ArrayTypes);
+    			Arguments.Emplace(*IndexName, ERigVMPinDirection::Input, FRigVMTemplateArgument::FType(RigVMTypeUtils::Int32Type, nullptr));
+    			break;
+    		}
+    		case ERigVMOpCode::ArrayFind:
+    		{
+    			Arguments.Emplace(*ArrayName, ERigVMPinDirection::Input, ArrayTypes);
+    			Arguments.Emplace(*ElementName, ERigVMPinDirection::Input, SingleTypes);
+    			Arguments.Emplace(*IndexName, ERigVMPinDirection::Output, FRigVMTemplateArgument::FType(RigVMTypeUtils::Int32Type, nullptr));
+    			Arguments.Emplace(*SuccessName, ERigVMPinDirection::Output, FRigVMTemplateArgument::FType(RigVMTypeUtils::BoolType, nullptr));
+    			break;
+    		}
+    		case ERigVMOpCode::ArrayAppend:
+			case ERigVMOpCode::ArrayUnion:
+    		{
+    			Arguments.Emplace(*FRigVMStruct::ExecuteContextName.ToString(), ERigVMPinDirection::IO, ExecuteType);
+    			Arguments.Emplace(*ArrayName, ERigVMPinDirection::IO, ArrayTypes);
+    			Arguments.Emplace(*OtherName, ERigVMPinDirection::Input, ArrayTypes);
+    			break;
+    		}
+    		case ERigVMOpCode::ArrayClone:
+    		{
+    			Arguments.Emplace(*ArrayName, ERigVMPinDirection::Input, ArrayTypes);
+    			Arguments.Emplace(*CloneName, ERigVMPinDirection::Output, ArrayTypes);
+    			break;    				
+    		}
+    		case ERigVMOpCode::ArrayIterator:
+    		{
+    			Arguments.Emplace(*FRigVMStruct::ExecuteContextName.ToString(), ERigVMPinDirection::IO, ExecuteType);
+    			Arguments.Emplace(*ArrayName, ERigVMPinDirection::Input, ArrayTypes);
+    			Arguments.Emplace(*ElementName, ERigVMPinDirection::Output, SingleTypes);
+    			Arguments.Emplace(*IndexName, ERigVMPinDirection::Output, FRigVMTemplateArgument::FType(RigVMTypeUtils::Int32Type, nullptr));
+    			Arguments.Emplace(*CountName, ERigVMPinDirection::Output, FRigVMTemplateArgument::FType(RigVMTypeUtils::Int32Type, nullptr));
+    			Arguments.Emplace(*RatioName, ERigVMPinDirection::Output, FRigVMTemplateArgument::FType(RigVMTypeUtils::FloatType, nullptr));
+    			Arguments.Emplace(*ContinueName, ERigVMPinDirection::Hidden, ExecuteType);
+    			Arguments.Emplace(*CompletedName, ERigVMPinDirection::Output, ExecuteType);
+    			break;
+    		}
+    		case ERigVMOpCode::ArrayDifference:
+    		case ERigVMOpCode::ArrayIntersection:
+    		{
+    			Arguments.Emplace(*ArrayName, ERigVMPinDirection::Input, ArrayTypes);
+    			Arguments.Emplace(*OtherName, ERigVMPinDirection::Input, ArrayTypes);
+    			Arguments.Emplace(*ResultName, ERigVMPinDirection::Output, ArrayTypes);
+    			break;
+    		}
+    		default:
+    		{
+    			ensure(false);
+    			break;
+    		}
+	    }
+
+		const FString TemplateName = StaticEnum<ERigVMOpCode>()->GetDisplayNameTextByValue((int64)OpCode).ToString();
+		CachedTemplate = FRigVMRegistry::Get().GetOrAddTemplateFromArguments(*TemplateName, Arguments);
+	}
+	return CachedTemplate;
+}
+
 bool URigVMArrayNode::IsLoopNode() const
 {
 	return GetOpCode() == ERigVMOpCode::ArrayIterator;
diff --git a/Engine/Source/Developer/RigVMDeveloper/Private/RigVMModel/Nodes/RigVMIfNode.cpp b/Engine/Source/Developer/RigVMDeveloper/Private/RigVMModel/Nodes/RigVMIfNode.cpp
--- a/Engine/Source/Developer/RigVMDeveloper/Private/RigVMModel/Nodes/RigVMIfNode.cpp
+++ b/Engine/Source/Developer/RigVMDeveloper/Private/RigVMModel/Nodes/RigVMIfNode.cpp
@@ -7,3 +7,36 @@ const FString URigVMIfNode::ConditionName = TEXT("Condition");
 const FString URigVMIfNode::TrueName = TEXT("True");
 const FString URigVMIfNode::FalseName = TEXT("False");
 const FString URigVMIfNode::ResultName = TEXT("Result");
+
+
+FName URigVMIfNode::GetNotation() const
+{
+	static constexpr TCHAR Format[] = TEXT("%s(in %s, in %s, in %s, out %s)");
+	static const FName Notation = *FString::Printf(Format, *IfName, *ConditionName, *TrueName, *FalseName, *ResultName);
+	return Notation;
+}
+
+const FRigVMTemplate* URigVMIfNode::GetTemplate() const
+{
+	if(const FRigVMTemplate* SuperTemplate = Super::GetTemplate())
+	{
+		return SuperTemplate;
+	}
+	
+	if(CachedTemplate == nullptr)
+	{
+		TArray<FRigVMTemplateArgument::FType> Types;
+		Types.Append(FRigVMTemplateArgument::GetCompatibleTypes(FRigVMTemplateArgument::ETypeCategory_SingleAnyValue));
+		Types.Append(FRigVMTemplateArgument::GetCompatibleTypes(FRigVMTemplateArgument::ETypeCategory_ArrayAnyValue));
+		
+		TArray<FRigVMTemplateArgument> Arguments;
+		Arguments.Emplace(*ConditionName, ERigVMPinDirection::Input, FRigVMTemplateArgument::FType(RigVMTypeUtils::BoolType, nullptr));
+		Arguments.Emplace(*TrueName, ERigVMPinDirection::Input, Types);
+		Arguments.Emplace(*FalseName, ERigVMPinDirection::Input, Types);
+		Arguments.Emplace(*ResultName, ERigVMPinDirection::Output, Types);
+		
+		CachedTemplate = FRigVMRegistry::Get().GetOrAddTemplateFromArguments(*IfName, Arguments);
+	}
+	return CachedTemplate;
+}
+
diff --git a/Engine/Source/Developer/RigVMDeveloper/Private/RigVMModel/Nodes/RigVMRerouteNode.cpp b/Engine/Source/Developer/RigVMDeveloper/Private/RigVMModel/Nodes/RigVMRerouteNode.cpp
--- a/Engine/Source/Developer/RigVMDeveloper/Private/RigVMModel/Nodes/RigVMRerouteNode.cpp
+++ b/Engine/Source/Developer/RigVMDeveloper/Private/RigVMModel/Nodes/RigVMRerouteNode.cpp
@@ -1,6 +1,7 @@
 // Copyright Epic Games, Inc. All Rights Reserved.
 
 #include "RigVMModel/Nodes/RigVMRerouteNode.h"
+#include "RigVMModel/RigVMGraph.h"
 
 const FString URigVMRerouteNode::RerouteName = TEXT("Reroute");
 const FString URigVMRerouteNode::ValueName = TEXT("Value");
@@ -82,7 +83,9 @@ const FRigVMTemplate* URigVMRerouteNode::GetTemplate() const
 		TArray<FRigVMTemplateArgument::FType> Types;
 		Types.Append(FRigVMTemplateArgument::GetCompatibleTypes(FRigVMTemplateArgument::ETypeCategory_SingleAnyValue));
 		Types.Append(FRigVMTemplateArgument::GetCompatibleTypes(FRigVMTemplateArgument::ETypeCategory_ArrayAnyValue));
-		Types.Append(FRigVMTemplateArgument::GetCompatibleTypes(FRigVMTemplateArgument::ETypeCategory_ExecuteContext));
+
+		UScriptStruct* ExecuteStruct = GetGraph()->GetExecuteContextStruct();
+		Types.Add(FRigVMTemplateArgument::FType(ExecuteStruct->GetStructCPPName(), ExecuteStruct));
 		
 		TArray<FRigVMTemplateArgument> Arguments;
 		Arguments.Emplace(TEXT("Value"), ERigVMPinDirection::IO, Types);
diff --git a/Engine/Source/Developer/RigVMDeveloper/Private/RigVMModel/Nodes/RigVMSelectNode.cpp b/Engine/Source/Developer/RigVMDeveloper/Private/RigVMModel/Nodes/RigVMSelectNode.cpp
--- a/Engine/Source/Developer/RigVMDeveloper/Private/RigVMModel/Nodes/RigVMSelectNode.cpp
+++ b/Engine/Source/Developer/RigVMDeveloper/Private/RigVMModel/Nodes/RigVMSelectNode.cpp
@@ -19,3 +19,32 @@ bool URigVMSelectNode::AllowsLinksOn(const URigVMPin* InPin) const
 
 	return true;
 }
+
+FName URigVMSelectNode::GetNotation() const
+{
+	static constexpr TCHAR Format[] = TEXT("%s(in %s, in %s, out %s)");
+	static const FName Notation = *FString::Printf(Format, *SelectName, *IndexName, *ValueName, *ResultName);
+	return Notation;
+}
+
+const FRigVMTemplate* URigVMSelectNode::GetTemplate() const
+{
+	if(const FRigVMTemplate* SuperTemplate = Super::GetTemplate())
+	{
+		return SuperTemplate;
+	}
+	
+	if(CachedTemplate == nullptr)
+	{
+		const TArray<FRigVMTemplateArgument::FType>& SingleTypes = FRigVMTemplateArgument::GetCompatibleTypes(FRigVMTemplateArgument::ETypeCategory_SingleAnyValue);
+		const TArray<FRigVMTemplateArgument::FType>& ArrayTypes = FRigVMTemplateArgument::GetCompatibleTypes(FRigVMTemplateArgument::ETypeCategory_ArrayAnyValue);
+		
+		TArray<FRigVMTemplateArgument> Arguments;
+		Arguments.Emplace(*IndexName, ERigVMPinDirection::Input, FRigVMTemplateArgument::FType(RigVMTypeUtils::Int32Type, nullptr));
+		Arguments.Emplace(*ValueName, ERigVMPinDirection::Input, ArrayTypes);
+		Arguments.Emplace(*ResultName, ERigVMPinDirection::Output, SingleTypes);
+		
+		CachedTemplate = FRigVMRegistry::Get().GetOrAddTemplateFromArguments(*SelectName, Arguments);
+	}
+	return CachedTemplate;
+}
\ No newline at end of file
diff --git a/Engine/Source/Developer/RigVMDeveloper/Private/RigVMModel/Nodes/RigVMTemplateNode.cpp b/Engine/Source/Developer/RigVMDeveloper/Private/RigVMModel/Nodes/RigVMTemplateNode.cpp
--- a/Engine/Source/Developer/RigVMDeveloper/Private/RigVMModel/Nodes/RigVMTemplateNode.cpp
+++ b/Engine/Source/Developer/RigVMDeveloper/Private/RigVMModel/Nodes/RigVMTemplateNode.cpp
@@ -367,6 +367,64 @@ FString URigVMTemplateNode::GetInitialDefaultValueForPin(const FName& InRootPinN
 			NewDefaultValue = DefaultStruct->ExportToFullyQualifiedText(
 				Cast<UScriptStruct>(StructOnScope->GetStruct()), InRootPinName);
 		}
+		else
+		{
+			const FRigVMTemplate* Template = GetTemplate();
+			const FRigVMTemplateArgument* Argument = Template->FindArgument(InRootPinName);
+			const FRigVMTemplateArgument::FType Type = Argument->GetTypes()[PermutationIndex];
+
+			if (Type.IsArray())
+			{
+				NewDefaultValue = TEXT("()");
+			}
+			else
+			{
+				if (UScriptStruct* ScriptStruct = Cast<UScriptStruct>(Type.CPPTypeObject))
+				{
+					TArray<uint8, TAlignedHeapAllocator<16>> TempBuffer;
+					TempBuffer.AddUninitialized(ScriptStruct->GetStructureSize());
+
+					// call the struct constructor to initialize the struct
+					ScriptStruct->InitializeDefaultValue(TempBuffer.GetData());
+
+					ScriptStruct->ExportText(NewDefaultValue, TempBuffer.GetData(), nullptr, nullptr, PPF_None, nullptr);
+					ScriptStruct->DestroyStruct(TempBuffer.GetData());				
+				}
+				else if (UEnum* Enum = Cast<UEnum>(Type.CPPTypeObject))
+				{
+					NewDefaultValue = Enum->GetNameStringByValue(0);
+				}
+				else if(UClass* Class = Cast<UClass>(Type.CPPTypeObject))
+				{
+					// not supporting objects yet
+					ensure(false);
+				}
+				else if (Type.CPPType == RigVMTypeUtils::FloatType)
+				{
+					NewDefaultValue = TEXT("0.000000");				
+				}
+				else if (Type.CPPType == RigVMTypeUtils::DoubleType)
+				{
+					NewDefaultValue = TEXT("0.000000");
+				}
+				else if (Type.CPPType == RigVMTypeUtils::Int32Type)
+				{
+					NewDefaultValue = TEXT("0");
+				}
+				else if (Type.CPPType == RigVMTypeUtils::BoolType)
+				{
+					NewDefaultValue = TEXT("False");
+				}
+				else if (Type.CPPType == RigVMTypeUtils::FStringType)
+				{
+					NewDefaultValue = TEXT("");
+				}
+				else if (Type.CPPType == RigVMTypeUtils::FNameType)
+				{
+					NewDefaultValue = TEXT("");
+				}
+			}			
+		}
 
 		if(!NewDefaultValue.IsEmpty())
 		{
diff --git a/Engine/Source/Developer/RigVMDeveloper/Private/RigVMModel/RigVMController.cpp b/Engine/Source/Developer/RigVMDeveloper/Private/RigVMModel/RigVMController.cpp
--- a/Engine/Source/Developer/RigVMDeveloper/Private/RigVMModel/RigVMController.cpp
+++ b/Engine/Source/Developer/RigVMDeveloper/Private/RigVMModel/RigVMController.cpp
@@ -70,7 +70,6 @@
 	, bIsRunningUnitTest(false)
 	, bIsFullyResolvingTemplateNode(false)
 {
-	SetExecuteContextStruct(FRigVMExecuteContext::StaticStruct());
 }
 
 URigVMController::URigVMController(const FObjectInitializer& ObjectInitializer)
@@ -85,7 +84,6 @@
 	, bIsFullyResolvingTemplateNode(false)
 {
 	ActionStack = CreateDefaultSubobject<URigVMActionStack>(TEXT("ActionStack"));
-	SetExecuteContextStruct(FRigVMExecuteContext::StaticStruct());
 
 	ActionStack->OnModified().AddLambda([&](ERigVMGraphNotifType NotifType, URigVMGraph* InGraph, UObject* InSubject) -> void {
 		Notify(NotifType, InSubject);
@@ -511,6 +509,47 @@
 						*RigVMPythonUtils::Vector2DToPythonString(SelectNode->GetPosition()),
 						*NodeName));
 	}
+	else if (const URigVMRerouteNode* RerouteNode = Cast<URigVMRerouteNode>(Node))
+	{
+		// add_free_reroute_node(bool bShowAsFullNode, const FString& InCPPType, const FName& InCPPTypeObjectPath, bool bIsConstant, const FName& InCustomWidgetName, const FString& InDefaultValue, const FVector2D& InPosition = FVector2D::ZeroVector, const FString& InNodeName = TEXT(""), bool bSetupUndoRedo = true);
+		Commands.Add(FString::Printf(TEXT("blueprint.get_controller_by_name('%s').add_free_reroute_node(%s, '%s', '%s', %s, '%s', '%s', %s, '%s')"),
+				*GraphName,
+				RerouteNode->GetShowsAsFullNode() ? TEXT("True") : TEXT("False"),
+				*RerouteNode->GetPins()[0]->GetCPPType(),
+				*RerouteNode->GetPins()[0]->GetCPPTypeObject()->GetPathName(),
+				RerouteNode->GetPins()[0]->IsDefinedAsConstant() ? TEXT("True") : TEXT("False"),
+				*RerouteNode->GetPins()[0]->GetCustomWidgetName().ToString(),
+				*RerouteNode->GetPins()[0]->GetDefaultValue(),
+				*RigVMPythonUtils::Vector2DToPythonString(RerouteNode->GetPosition()),
+				*NodeName));
+	}
+	else if (const URigVMArrayNode* ArrayNode = Cast<URigVMArrayNode>(Node))
+	{
+		FString OpCodeString = StaticEnum<ERigVMOpCode>()->GetNameStringByValue((int64)ArrayNode->GetOpCode());
+		OpCodeString = RigVMPythonUtils::NameToPep8(OpCodeString);
+		OpCodeString.ToUpperInline();
+		
+		// add_array_node(opcode, cpp_type, cpp_type_object, position=[0.0, 0.0], node_name='', undo=True)
+		if (ArrayNode->GetCPPTypeObject())
+		{
+			Commands.Add(FString::Printf(TEXT("blueprint.get_controller_by_name('%s').add_array_node_from_object_path(unreal.RigVMOpCode.%s, '%s', '%s', %s, '%s')"),
+					*GraphName,
+					*OpCodeString,
+					*ArrayNode->GetCPPType(),
+					*ArrayNode->GetCPPTypeObject()->GetPathName(),
+					*RigVMPythonUtils::Vector2DToPythonString(ArrayNode->GetPosition()),
+					*NodeName));	
+		}
+		else
+		{
+			Commands.Add(FString::Printf(TEXT("blueprint.get_controller_by_name('%s').add_array_node(unreal.RigVMOpCode.%s, '%s', None, %s, '%s')"),
+					*GraphName,
+					*OpCodeString,
+					*ArrayNode->GetCPPType(),
+					*RigVMPythonUtils::Vector2DToPythonString(ArrayNode->GetPosition()),
+					*NodeName));	
+		}
+	}
 	else if (const URigVMTemplateNode* TemplateNode = Cast<URigVMTemplateNode>(Node))
 	{
 		// add_template_node(notation, position=[0.0, 0.0], node_name='', undo=True)
@@ -563,51 +602,10 @@
 					*ContainedGraphName));
 		}
 	}
-	else if (const URigVMRerouteNode* RerouteNode = Cast<URigVMRerouteNode>(Node))
-	{
-		// add_free_reroute_node(bool bShowAsFullNode, const FString& InCPPType, const FName& InCPPTypeObjectPath, bool bIsConstant, const FName& InCustomWidgetName, const FString& InDefaultValue, const FVector2D& InPosition = FVector2D::ZeroVector, const FString& InNodeName = TEXT(""), bool bSetupUndoRedo = true);
-		Commands.Add(FString::Printf(TEXT("blueprint.get_controller_by_name('%s').add_free_reroute_node(%s, '%s', '%s', %s, '%s', '%s', %s, '%s')"),
-				*GraphName,
-				RerouteNode->GetShowsAsFullNode() ? TEXT("True") : TEXT("False"),
-				*RerouteNode->GetPins()[0]->GetCPPType(),
-				*RerouteNode->GetPins()[0]->GetCPPTypeObject()->GetPathName(),
-				RerouteNode->GetPins()[0]->IsDefinedAsConstant() ? TEXT("True") : TEXT("False"),
-				*RerouteNode->GetPins()[0]->GetCustomWidgetName().ToString(),
-				*RerouteNode->GetPins()[0]->GetDefaultValue(),
-				*RigVMPythonUtils::Vector2DToPythonString(RerouteNode->GetPosition()),
-				*NodeName));
-	}
 	else if (Node->IsA<URigVMFunctionEntryNode>() || Node->IsA<URigVMFunctionReturnNode>())
 	{
 		
 		
 	}
-	else if (const URigVMArrayNode* ArrayNode = Cast<URigVMArrayNode>(Node))
-	{
-		FString OpCodeString = StaticEnum<ERigVMOpCode>()->GetNameStringByValue((int64)ArrayNode->GetOpCode());
-		OpCodeString = RigVMPythonUtils::NameToPep8(OpCodeString);
-		OpCodeString.ToUpperInline();
-		
-		// add_array_node(opcode, cpp_type, cpp_type_object, position=[0.0, 0.0], node_name='', undo=True)
-		if (ArrayNode->GetCPPTypeObject())
-		{
-			Commands.Add(FString::Printf(TEXT("blueprint.get_controller_by_name('%s').add_array_node_from_object_path(unreal.RigVMOpCode.%s, '%s', '%s', %s, '%s')"),
-					*GraphName,
-					*OpCodeString,
-					*ArrayNode->GetCPPType(),
-					*ArrayNode->GetCPPTypeObject()->GetPathName(),
-					*RigVMPythonUtils::Vector2DToPythonString(ArrayNode->GetPosition()),
-					*NodeName));	
-		}
-		else
-		{
-			Commands.Add(FString::Printf(TEXT("blueprint.get_controller_by_name('%s').add_array_node(unreal.RigVMOpCode.%s, '%s', None, %s, '%s')"),
-					*GraphName,
-					*OpCodeString,
-					*ArrayNode->GetCPPType(),
-					*RigVMPythonUtils::Vector2DToPythonString(ArrayNode->GetPosition()),
-					*NodeName));	
-		}
-	}
 	else
 	{
@@ -1053,6 +1051,7 @@
 
 	if (!bIsGetter)
 	{
+		UScriptStruct* ExecuteContextStruct = Graph->GetExecuteContextStruct();
 		URigVMPin* ExecutePin = NewObject<URigVMPin>(Node, FRigVMStruct::ExecuteContextName);
 		ExecutePin->CPPType = ExecuteContextStruct->GetStructCPPName();
 		ExecutePin->CPPTypeObject = ExecuteContextStruct;
@@ -7194,7 +7193,10 @@
 	
 	check(InPin);
 
-	if(!InPin->IsUObject() && bValidatePinDefaults)
+	if(!InPin->IsUObject()
+		&& InPin->GetCPPType() != RigVMTypeUtils::FStringType
+		&& InPin->GetCPPType() != RigVMTypeUtils::FNameType
+		&& bValidatePinDefaults)
 	{
 		ensure(!InDefaultValue.IsEmpty());
 	}
@@ -10366,6 +10368,7 @@
 	
 	if (bMutable)
 	{
+		UScriptStruct* ExecuteContextStruct = Graph->GetExecuteContextStruct();
 		URigVMPin* ExecutePin = NewObject<URigVMPin>(CollapseNode, FRigVMStruct::ExecuteContextName);
 		ExecutePin->CPPType = FString::Printf(TEXT("F%s"), *ExecuteContextStruct->GetName());
 		ExecutePin->CPPTypeObject = ExecuteContextStruct;
@@ -11450,6 +11453,7 @@
 	URigVMBranchNode* Node = NewObject<URigVMBranchNode>(Graph, *Name);
 	Node->Position = InPosition;
 
+	UScriptStruct* ExecuteContextStruct = Graph->GetExecuteContextStruct();
 	URigVMPin* ExecutePin = NewObject<URigVMPin>(Node, FRigVMStruct::ExecuteContextName);
 	ExecutePin->DisplayName = FRigVMStruct::ExecuteName;
 	ExecutePin->CPPType = FString::Printf(TEXT("F%s"), *ExecuteContextStruct->GetName());
@@ -12093,7 +12097,8 @@
 		static URigVMPin* AddExecutePin(URigVMController* InController, URigVMNode* InNode, ERigVMPinDirection InDirection = ERigVMPinDirection::IO, const FName& InName = NAME_None)
 		{
 			const FName PinName = InName.IsNone() ? FRigVMStruct::ExecuteContextName : InName;
-			URigVMPin* Pin = AddPin(InController, InNode, PinName, InDirection, false, FString::Printf(TEXT("F%s"), *InController->ExecuteContextStruct->GetName()), InController->ExecuteContextStruct);
+			UScriptStruct* ExecuteContextStruct = InController->GetGraph()->GetExecuteContextStruct();
+			URigVMPin* Pin = AddPin(InController, InNode, PinName, InDirection, false, FString::Printf(TEXT("F%s"), *ExecuteContextStruct->GetName()), ExecuteContextStruct);
 			if(PinName == FRigVMStruct::ExecuteContextName)
 			{
 				Pin->DisplayName = FRigVMStruct::ExecuteName;
@@ -12325,13 +12330,6 @@
 	}
 }
 
-void URigVMController::SetExecuteContextStruct(UScriptStruct* InExecuteContextStruct)
-{
-	check(InExecuteContextStruct);
-	ensure(InExecuteContextStruct->IsChildOf(FRigVMExecuteContext::StaticStruct()));
-	ExecuteContextStruct = InExecuteContextStruct;
-}
-
 FString URigVMController::GetValidNodeName(const FString& InPrefix)
 {
 	URigVMGraph* Graph = GetGraph();
@@ -14934,11 +14932,22 @@
 	
 	if(URigVMTemplateNode* TemplateNode = Cast<URigVMTemplateNode>(InPinToResolve->GetNode()))
 	{
-		const FString BaseCPPType = RigVMTypeUtils::IsArrayType(CPPType) ? RigVMTypeUtils::BaseTypeFromArrayType(CPPType) : CPPType;
-		const FString TypeToResolveAgainst = InPinToResolve->GetRootPin()->IsArray() ? RigVMTypeUtils::ArrayTypeFromBaseType(BaseCPPType) : BaseCPPType; 
-		if(!TemplateNode->SupportsType(InPinToResolve->GetRootPin(), TypeToResolveAgainst, &CPPType))
+		URigVMPin* RootPin = InPinToResolve->GetRootPin();
+		if (RootPin != InPinToResolve)
 		{
-			return false;
+			const FString BaseCPPType = RigVMTypeUtils::IsArrayType(CPPType) ? RigVMTypeUtils::BaseTypeFromArrayType(CPPType) : CPPType;
+			const FString TypeToResolveAgainst = InPinToResolve->GetRootPin()->IsArray() ? RigVMTypeUtils::ArrayTypeFromBaseType(BaseCPPType) : BaseCPPType; 
+			if(!TemplateNode->SupportsType(InPinToResolve->GetRootPin(), TypeToResolveAgainst, &CPPType))
+			{
+				return false;
+			}			
 		}
+		else
+		{
+			if(!TemplateNode->SupportsType(InPinToResolve->GetRootPin(), CPPType, &CPPType))
+			{
+				return false;
+			}
+		}
 	}
 
diff --git a/Engine/Source/Developer/RigVMDeveloper/Private/RigVMModel/RigVMGraph.cpp b/Engine/Source/Developer/RigVMDeveloper/Private/RigVMModel/RigVMGraph.cpp
--- a/Engine/Source/Developer/RigVMDeveloper/Private/RigVMModel/RigVMGraph.cpp
+++ b/Engine/Source/Developer/RigVMDeveloper/Private/RigVMModel/RigVMGraph.cpp
@@ -14,6 +14,7 @@ URigVMGraph::URigVMGraph()
 , RuntimeAST(nullptr)
 , bEditable(true)
 {
+	SetExecuteContextStruct(FRigVMExecuteContext::StaticStruct());
 }
 
 const TArray<URigVMNode*>& URigVMGraph::GetNodes() const
@@ -333,6 +334,22 @@ FRigVMGraphModifiedEvent& URigVMGraph::OnModified()
 	return ModifiedEvent;
 }
 
+void URigVMGraph::SetExecuteContextStruct(UScriptStruct* InExecuteContextStruct)
+{
+	check(InExecuteContextStruct);
+	ensure(InExecuteContextStruct->IsChildOf(FRigVMExecuteContext::StaticStruct()));
+	ExecuteContextStruct = InExecuteContextStruct;
+}
+
+UScriptStruct* URigVMGraph::GetExecuteContextStruct() const
+{
+	if (URigVMGraph* RootGraph = GetRootGraph())
+	{
+		return RootGraph->ExecuteContextStruct.Get();
+	}
+	return nullptr;
+}
+
 void URigVMGraph::Notify(ERigVMGraphNotifType InNotifType, UObject* InSubject)
 {
 	ModifiedEvent.Broadcast(InNotifType, this, InSubject);
diff --git a/Engine/Source/Developer/RigVMDeveloper/Public/RigVMModel/Nodes/RigVMArrayNode.h b/Engine/Source/Developer/RigVMDeveloper/Public/RigVMModel/Nodes/RigVMArrayNode.h
--- a/Engine/Source/Developer/RigVMDeveloper/Public/RigVMModel/Nodes/RigVMArrayNode.h
+++ b/Engine/Source/Developer/RigVMDeveloper/Public/RigVMModel/Nodes/RigVMArrayNode.h
@@ -2,8 +2,8 @@
 
 #pragma once
 
+#include "RigVMTemplateNode.h"
 #include "RigVMCore/RigVMExternalVariable.h"
-#include "RigVMModel/RigVMNode.h"
 #include "RigVMArrayNode.generated.h"
 
 /**
@@ -11,7 +11,7 @@
  * array operations such as SetNum, GetAtIndex etc.
  */
 UCLASS(BlueprintType)
-class RIGVMDEVELOPER_API URigVMArrayNode : public URigVMNode
+class RIGVMDEVELOPER_API URigVMArrayNode : public URigVMTemplateNode
 {
 	GENERATED_BODY()
 
@@ -32,6 +32,12 @@ public:
 	// Override of the node color
 	virtual FLinearColor GetNodeColor() const override;
 
+    // Override of the template notation
+	virtual FName GetNotation() const override;
+	
+	// Override of the template
+	virtual const FRigVMTemplate* GetTemplate() const override;
+
 	// Returns the op code of this node
 	UFUNCTION(BlueprintCallable, Category = RigVMArrayNode)
 	ERigVMOpCode GetOpCode() const;
diff --git a/Engine/Source/Developer/RigVMDeveloper/Public/RigVMModel/Nodes/RigVMIfNode.h b/Engine/Source/Developer/RigVMDeveloper/Public/RigVMModel/Nodes/RigVMIfNode.h
--- a/Engine/Source/Developer/RigVMDeveloper/Public/RigVMModel/Nodes/RigVMIfNode.h
+++ b/Engine/Source/Developer/RigVMDeveloper/Public/RigVMModel/Nodes/RigVMIfNode.h
@@ -2,23 +2,26 @@
 
 #pragma once
 
-#include "RigVMModel/RigVMNode.h"
+#include "RigVMTemplateNode.h"
 #include "RigVMIfNode.generated.h"
 
 /**
  * A if node is used to pick between two values
  */
 UCLASS(BlueprintType)
-class RIGVMDEVELOPER_API URigVMIfNode : public URigVMNode
+class RIGVMDEVELOPER_API URigVMIfNode : public URigVMTemplateNode
 {
 	GENERATED_BODY()
 
 public:
 
+	// Override from URigVMTemplateNode
+	virtual FName GetNotation() const override;
+	virtual const FRigVMTemplate* GetTemplate() const override;
+
 	// Override from URigVMNode
 	virtual FString GetNodeTitle() const override { return IfName; }
 	virtual FLinearColor GetNodeColor() const override { return FLinearColor::Black; }
-
 private:
 
 	static const FString IfName;
diff --git a/Engine/Source/Developer/RigVMDeveloper/Public/RigVMModel/Nodes/RigVMSelectNode.h b/Engine/Source/Developer/RigVMDeveloper/Public/RigVMModel/Nodes/RigVMSelectNode.h
--- a/Engine/Source/Developer/RigVMDeveloper/Public/RigVMModel/Nodes/RigVMSelectNode.h
+++ b/Engine/Source/Developer/RigVMDeveloper/Public/RigVMModel/Nodes/RigVMSelectNode.h
@@ -2,23 +2,27 @@
 
 #pragma once
 
-#include "RigVMModel/RigVMNode.h"
+#include "RigVMTemplateNode.h"
 #include "RigVMSelectNode.generated.h"
 
 /**
  * A select node is used to select between multiple values
  */
 UCLASS(BlueprintType)
-class RIGVMDEVELOPER_API URigVMSelectNode : public URigVMNode
+class RIGVMDEVELOPER_API URigVMSelectNode : public URigVMTemplateNode
 {
 	GENERATED_BODY()
 
 public:
 
+	// Override from URigVMTemplateNode
+	virtual FName GetNotation() const override;
+	virtual const FRigVMTemplate* GetTemplate() const override;
+
 	// Override from URigVMNode
 	virtual FString GetNodeTitle() const override { return SelectName; }
 	virtual FLinearColor GetNodeColor() const override { return FLinearColor::Black; }
-
+	
 protected:
 
 	virtual bool AllowsLinksOn(const URigVMPin* InPin) const override;
diff --git a/Engine/Source/Developer/RigVMDeveloper/Public/RigVMModel/RigVMController.h b/Engine/Source/Developer/RigVMDeveloper/Public/RigVMModel/RigVMController.h
--- a/Engine/Source/Developer/RigVMDeveloper/Public/RigVMModel/RigVMController.h
+++ b/Engine/Source/Developer/RigVMDeveloper/Public/RigVMModel/RigVMController.h
@@ -834,9 +834,6 @@ public:
 	// Determine affected assets for a potential bulk edit on a library node
 	TArray<FAssetData> GetAffectedAssets(ERigVMControllerBulkEditType InEditType, bool bForceLoad = false, bool bNotify = true);
 
-	// Sets the execute context struct type to use
-	void SetExecuteContextStruct(UScriptStruct* InExecuteContextStruct);
-
 	// A delegate that can be set to change the struct unfolding behaviour
 	FRigVMController_ShouldStructUnfoldDelegate UnfoldStructDelegate;
 
@@ -1098,9 +1095,6 @@ private:
 	UPROPERTY(transient)
 	TArray<TObjectPtr<URigVMGraph>> Graphs;
 
-	UPROPERTY()
-	TObjectPtr<UScriptStruct> ExecuteContextStruct;
-
 	UPROPERTY(transient)
 	TObjectPtr<URigVMActionStack> ActionStack;
 
diff --git a/Engine/Source/Developer/RigVMDeveloper/Public/RigVMModel/RigVMGraph.h b/Engine/Source/Developer/RigVMDeveloper/Public/RigVMModel/RigVMGraph.h
--- a/Engine/Source/Developer/RigVMDeveloper/Public/RigVMModel/RigVMGraph.h
+++ b/Engine/Source/Developer/RigVMDeveloper/Public/RigVMModel/RigVMGraph.h
@@ -149,6 +149,11 @@ public:
 	// subscribe to changes happening within the Graph.
 	FRigVMGraphModifiedEvent& OnModified();
 
+	// Sets the execute context struct type to use
+	void SetExecuteContextStruct(UScriptStruct* InExecuteContextStruct);
+	
+	UScriptStruct* GetExecuteContextStruct() const;
+
 	void PrepareCycleChecking(URigVMPin* InPin, bool bAsInput);
 	virtual bool CanLink(URigVMPin* InSourcePin, URigVMPin* InTargetPin, FString* OutFailureReason, const FRigVMByteCode* InByteCode, ERigVMPinDirection InUserLinkDirection = ERigVMPinDirection::IO);
 	
@@ -173,6 +178,9 @@ private:
 	UPROPERTY()
 	TWeakObjectPtr<URigVMGraph> DefaultFunctionLibraryPtr;
 
+	UPROPERTY()
+	TObjectPtr<UScriptStruct> ExecuteContextStruct;
+
 	TSharedPtr<FRigVMParserAST> DiagnosticsAST;
 	TSharedPtr<FRigVMParserAST> RuntimeAST;
 
diff --git a/Engine/Source/Runtime/RigVM/Private/RigVMCore/RigVMTemplate.cpp b/Engine/Source/Runtime/RigVM/Private/RigVMCore/RigVMTemplate.cpp
--- a/Engine/Source/Runtime/RigVM/Private/RigVMCore/RigVMTemplate.cpp
+++ b/Engine/Source/Runtime/RigVM/Private/RigVMCore/RigVMTemplate.cpp
@@ -187,8 +187,7 @@ const TArray<FRigVMTemplateArgument::FType> FRigVMTemplateArgument::GetCompatibl
 		CompatibleTypes.Add(ETypeCategory_ArrayEnumValue, TArray<FType>());
 		CompatibleTypes.Add(ETypeCategory_SingleObjectValue, TArray<FType>());
 		CompatibleTypes.Add(ETypeCategory_ArrayObjectValue, TArray<FType>());
-		CompatibleTypes.Add(ETypeCategory_ExecuteContext, TArray<FType>());
-
+		
 		TArray<FType>& SingleAnyValueTypes = CompatibleTypes.FindChecked(ETypeCategory_SingleAnyValue);
 		TArray<FType>& ArrayAnyValueTypes = CompatibleTypes.FindChecked(ETypeCategory_ArrayAnyValue);
 		TArray<FType>& SingleSimpleValueTypes = CompatibleTypes.FindChecked(ETypeCategory_SingleSimpleValue);
@@ -201,8 +200,7 @@ const TArray<FRigVMTemplateArgument::FType> FRigVMTemplateArgument::GetCompatibl
 		TArray<FType>& ArrayEnumValueTypes = CompatibleTypes.FindChecked(ETypeCategory_ArrayEnumValue);
 		TArray<FType>& SingleObjectValueTypes = CompatibleTypes.FindChecked(ETypeCategory_SingleObjectValue);
 		TArray<FType>& ArrayObjectValueTypes = CompatibleTypes.FindChecked(ETypeCategory_ArrayObjectValue);
-		TArray<FType>& ExecuteContextTypes = CompatibleTypes.FindChecked(ETypeCategory_ExecuteContext);
-
+		
 		SingleSimpleValueTypes.Add(FType(RigVMTypeUtils::BoolType));
 		SingleSimpleValueTypes.Add(FType(RigVMTypeUtils::Int32Type));
 		SingleSimpleValueTypes.Add(FType(RigVMTypeUtils::UInt8Type));
@@ -355,11 +353,7 @@ const TArray<FRigVMTemplateArgument::FType> FRigVMTemplateArgument::GetCompatibl
 
 			const FString CPPType = ScriptStruct->GetStructCPPName();
 
-			if(ScriptStruct->IsChildOf(FRigVMExecuteContext::StaticStruct()))
-			{
-				ExecuteContextTypes.Add(FType(CPPType, ScriptStruct));
-			}
-			else
+			if(!ScriptStruct->IsChildOf(FRigVMExecuteContext::StaticStruct()))
 			{
 				SingleAnyValueTypes.Add(FType(CPPType, ScriptStruct));
 				ArrayAnyValueTypes.Add(FType(RigVMTypeUtils::ArrayTypeFromBaseType(CPPType), ScriptStruct));
diff --git a/Engine/Source/Runtime/RigVM/Public/RigVMCore/RigVMTemplate.h b/Engine/Source/Runtime/RigVM/Public/RigVMCore/RigVMTemplate.h
--- a/Engine/Source/Runtime/RigVM/Public/RigVMCore/RigVMTemplate.h
+++ b/Engine/Source/Runtime/RigVM/Public/RigVMCore/RigVMTemplate.h
@@ -134,7 +134,6 @@ struct RIGVM_API FRigVMTemplateArgument
 		ETypeCategory_ArrayEnumValue,
 		ETypeCategory_SingleObjectValue,
 		ETypeCategory_ArrayObjectValue,
-		ETypeCategory_ExecuteContext,
 		ETypeCategory_Invalid
 	};
 
