Fixing Control Rig enum redirector

diff --git a/Engine/Plugins/Animation/ControlRig/Config/BaseControlRig.ini b/Engine/Plugins/Animation/ControlRig/Config/BaseControlRig.ini
--- a/Engine/Plugins/Animation/ControlRig/Config/BaseControlRig.ini
+++ b/Engine/Plugins/Animation/ControlRig/Config/BaseControlRig.ini
@@ -26,6 +26,7 @@
 
 +EnumRedirects=(OldName="ERigElementType",ValueChanges=(("Space","Null"))
 +EnumRedirects=(OldName="/Script/ControlRig.EControlRigRotationOrder",NewName="/Script/AnimationCore.EEulerRotationOrder")
++EnumRedirects=(OldName="EControlRigRotationOrder",NewName="/Script/AnimationCore.EEulerRotationOrder")
 
 
 ; renaming gizmo to shape
diff --git a/Engine/Source/Developer/RigVMDeveloper/Private/RigVMModel/RigVMController.cpp b/Engine/Source/Developer/RigVMDeveloper/Private/RigVMModel/RigVMController.cpp
--- a/Engine/Source/Developer/RigVMDeveloper/Private/RigVMModel/RigVMController.cpp
+++ b/Engine/Source/Developer/RigVMDeveloper/Private/RigVMModel/RigVMController.cpp
@@ -15565,21 +15565,26 @@ static UObject* FindObjectGloballyWithRedirectors(const TCHAR* InObjectName)
 	{
 		return Object;
 	}
-	
-	// Apply redirectors and do another lookup using the redirected name. GetRedirectedName
-	// will always return a valid name.
-	const FCoreRedirectObjectName NewObjectName = FCoreRedirects::GetRedirectedName(
+
+	FCoreRedirectObjectName NewObjectName;
+	const bool bFoundRedirect = FCoreRedirects::RedirectNameAndValues(
 		ECoreRedirectFlags::Type_Class | ECoreRedirectFlags::Type_Struct | ECoreRedirectFlags::Type_Enum,
-		FCoreRedirectObjectName(InObjectName));
+		FCoreRedirectObjectName(InObjectName),
+		NewObjectName,
+		nullptr,
+		ECoreRedirectMatchFlags::None);
 
-	// If there was no package name, then there was no redirect set up for this type.
-	if (NewObjectName.PackageName.IsNone())
+	if (!bFoundRedirect)
 	{
 		return nullptr;
 	}
 
 	const FString RedirectedObjectName = NewObjectName.ObjectName.ToString();
-	UPackage *Package = FindPackage(nullptr, *NewObjectName.PackageName.ToString());
+	UPackage *Package = nullptr;
+	if (!NewObjectName.PackageName.IsNone())
+	{
+		Package = FindPackage(nullptr, *NewObjectName.PackageName.ToString());
+	}
 	if (Package != nullptr)
 	{
 		Object = FindObject<UField>(Package, *RedirectedObjectName);
