Control Rig: Fix collapsing nodes modifies behaviour

diff --git a/Engine/Source/Developer/RigVMDeveloper/Private/RigVMModel/Nodes/RigVMUnitNode.cpp b/Engine/Source/Developer/RigVMDeveloper/Private/RigVMModel/Nodes/RigVMUnitNode.cpp
--- a/Engine/Source/Developer/RigVMDeveloper/Private/RigVMModel/Nodes/RigVMUnitNode.cpp
+++ b/Engine/Source/Developer/RigVMDeveloper/Private/RigVMModel/Nodes/RigVMUnitNode.cpp
@@ -19,9 +19,9 @@ void URigVMUnitNode::PostLoad()
 			if(Function->TemplateIndex != INDEX_NONE)
 			{
 				const FRigVMTemplate& Template = FRigVMRegistry::Get().GetTemplates()[Function->TemplateIndex];
-				TemplateNotation = Template.GetNotation();
-				ResolvedFunctionName = Function->GetName();
-			}	
+				TemplateNotation = Template.GetNotation();				
+			}
+			ResolvedFunctionName = Function->GetName();
 		}
 	}
 }
