Control Rig: Fix input argument variable nodes being identified as external variables to be remapped

diff --git a/Engine/Source/Developer/RigVMDeveloper/Private/RigVMModel/Nodes/RigVMVariableNode.cpp b/Engine/Source/Developer/RigVMDeveloper/Private/RigVMModel/Nodes/RigVMVariableNode.cpp
--- a/Engine/Source/Developer/RigVMDeveloper/Private/RigVMModel/Nodes/RigVMVariableNode.cpp
+++ b/Engine/Source/Developer/RigVMDeveloper/Private/RigVMModel/Nodes/RigVMVariableNode.cpp
@@ -40,6 +40,11 @@ bool URigVMVariableNode::IsGetter() const
 	return ValuePin->GetDirection() == ERigVMPinDirection::Output;
 }
 
+bool URigVMVariableNode::IsExternalVariable() const
+{
+	return !IsLocalVariable() && !IsInputArgument();
+}
+
 bool URigVMVariableNode::IsLocalVariable() const
 {
 	const FName CurrentVariableName = GetVariableName();
diff --git a/Engine/Source/Developer/RigVMDeveloper/Private/RigVMModel/RigVMGraph.cpp b/Engine/Source/Developer/RigVMDeveloper/Private/RigVMModel/RigVMGraph.cpp
--- a/Engine/Source/Developer/RigVMDeveloper/Private/RigVMModel/RigVMGraph.cpp
+++ b/Engine/Source/Developer/RigVMDeveloper/Private/RigVMModel/RigVMGraph.cpp
@@ -261,18 +261,8 @@ TArray<FRigVMExternalVariable> URigVMGraph::GetExternalVariables() const
 		}
 		else if(URigVMVariableNode* VariableNode = Cast<URigVMVariableNode>(Node))
 		{
-			// Make sure it is not a local variable
-			bool bFoundLocalVariable = false;
-			for (FRigVMGraphVariableDescription& LocalVariable : Node->GetGraph()->LocalVariables)
-			{
-				if (LocalVariable.Name == VariableNode->GetVariableName())
-				{
-					bFoundLocalVariable = true;
-					break;
-				}
-			}
-
-			if (!bFoundLocalVariable)
+			// Make sure it is not a local variable or input argument
+			if (VariableNode->IsExternalVariable())
 			{
 				FRigVMExternalVariable::MergeExternalVariable(Variables, VariableNode->GetVariableDescription().ToExternalVariable());
 			}
diff --git a/Engine/Source/Developer/RigVMDeveloper/Public/RigVMModel/Nodes/RigVMVariableNode.h b/Engine/Source/Developer/RigVMDeveloper/Public/RigVMModel/Nodes/RigVMVariableNode.h
--- a/Engine/Source/Developer/RigVMDeveloper/Public/RigVMModel/Nodes/RigVMVariableNode.h
+++ b/Engine/Source/Developer/RigVMDeveloper/Public/RigVMModel/Nodes/RigVMVariableNode.h
@@ -34,6 +34,10 @@ public:
 	UFUNCTION(BlueprintCallable, Category = RigVMVariableNode)
 	bool IsGetter() const;
 
+	// Returns true if this variable is an external variable
+	UFUNCTION(BlueprintCallable, Category = RigVMVariableNode)
+	bool IsExternalVariable() const;
+
 	// Returns true if this variable is a local variable
 	UFUNCTION(BlueprintCallable, Category = RigVMVariableNode)
 	bool IsLocalVariable() const;
