Control Rig: Enabling aggregate nodes + implementing sequence as aggregate

diff --git a/Engine/Plugins/Animation/ControlRig/Source/ControlRig/Private/Units/Execution/RigUnit_SequenceExecution.cpp b/Engine/Plugins/Animation/ControlRig/Source/ControlRig/Private/Units/Execution/RigUnit_SequenceExecution.cpp
--- a/Engine/Plugins/Animation/ControlRig/Source/ControlRig/Private/Units/Execution/RigUnit_SequenceExecution.cpp
+++ b/Engine/Plugins/Animation/ControlRig/Source/ControlRig/Private/Units/Execution/RigUnit_SequenceExecution.cpp
@@ -7,3 +7,21 @@ FRigUnit_SequenceExecution_Execute()
 {
 	A = B = C = D = ExecuteContext;
 }
+
+FRigVMStructUpgradeInfo FRigUnit_SequenceExecution::GetUpgradeInfo() const
+{
+	FRigUnit_SequenceAggregate NewNode;
+	
+	FRigVMStructUpgradeInfo Info(*this, NewNode);
+
+	// add two more pins
+	Info.AddAggregatePin(); // "C"
+	Info.AddAggregatePin(); // "D"
+
+	return Info;
+}
+
+FRigUnit_SequenceAggregate_Execute()
+{
+	A = B = ExecuteContext;
+}
diff --git a/Engine/Plugins/Animation/ControlRig/Source/ControlRig/Private/Units/Execution/RigUnit_SequenceExecution.h b/Engine/Plugins/Animation/ControlRig/Source/ControlRig/Private/Units/Execution/RigUnit_SequenceExecution.h
--- a/Engine/Plugins/Animation/ControlRig/Source/ControlRig/Private/Units/Execution/RigUnit_SequenceExecution.h
+++ b/Engine/Plugins/Animation/ControlRig/Source/ControlRig/Private/Units/Execution/RigUnit_SequenceExecution.h
@@ -9,7 +9,7 @@
 /**
  * Allows for a single execution pulse to trigger a series of events in order.
  */
-USTRUCT(meta=(DisplayName="Sequence", Category="Execution", TitleColor="1 0 0", NodeColor="1 1 1", Icon="EditorStyle|GraphEditor.Sequence_16x"))
+USTRUCT(meta=(DisplayName="Sequence", Category="Execution", TitleColor="1 0 0", NodeColor="1 1 1", Icon="EditorStyle|GraphEditor.Sequence_16x", Deprecated="5.1.0"))
 struct CONTROLRIG_API FRigUnit_SequenceExecution : public FRigUnit
 {
 	GENERATED_BODY()
@@ -36,4 +36,31 @@ struct CONTROLRIG_API FRigUnit_SequenceExecution : public FRigUnit
 	// The execution result D
 	UPROPERTY(EditAnywhere, Transient, Category = "SequenceExecution", meta = (Output))
 	FControlRigExecuteContext D;
+
+	RIGVM_METHOD()
+	virtual FRigVMStructUpgradeInfo GetUpgradeInfo() const override;
+};
+
+/**
+ * Allows for a single execution pulse to trigger a series of events in order.
+ */
+USTRUCT(meta=(DisplayName="Sequence", Category="Execution", TitleColor="1 0 0", NodeColor="1 1 1", Icon="EditorStyle|GraphEditor.Sequence_16x"))
+struct CONTROLRIG_API FRigUnit_SequenceAggregate : public FRigUnit
+{
+	GENERATED_BODY()
+
+	RIGVM_METHOD()
+	virtual void Execute(const FRigUnitContext& Context) override;
+
+	// The execution input
+	UPROPERTY(EditAnywhere, Transient, DisplayName = "Execute", Category = "SequenceExecution", meta = (Input, Aggregate))
+	FControlRigExecuteContext ExecuteContext;
+
+	// The execution result A
+	UPROPERTY(EditAnywhere, Transient, Category = "SequenceExecution", meta = (Output, Aggregate))
+	FControlRigExecuteContext A;
+
+	// The execution result B
+	UPROPERTY(EditAnywhere, Transient, Category = "SequenceExecution", meta = (Output, Aggregate))
+	FControlRigExecuteContext B;
 };
diff --git a/Engine/Plugins/Animation/ControlRig/Source/ControlRigDeveloper/Private/ControlRigBlueprint.cpp b/Engine/Plugins/Animation/ControlRig/Source/ControlRigDeveloper/Private/ControlRigBlueprint.cpp
--- a/Engine/Plugins/Animation/ControlRig/Source/ControlRigDeveloper/Private/ControlRigBlueprint.cpp
+++ b/Engine/Plugins/Animation/ControlRig/Source/ControlRigDeveloper/Private/ControlRigBlueprint.cpp
@@ -1799,6 +1799,8 @@ TArray<FString> UControlRigBlueprint::GeneratePythonCommands(const FString InNew
 					
 					// Add Exposed Pins
 					{
+
+						bool bHitFirstExecute = false;
 						for (auto Pin : EntryNode->GetPins())
 						{
 							if (Pin->GetDirection() != ERigVMPinDirection::Output)
@@ -1806,6 +1808,15 @@ TArray<FString> UControlRigBlueprint::GeneratePythonCommands(const FString InNew
 								continue;
 							}
 
+							if(Pin->IsExecuteContext())
+							{
+								if(!bHitFirstExecute)
+								{
+									bHitFirstExecute = true;
+									continue;
+								}
+							}
+
 							// FName AddExposedPin(const FName& InPinName, ERigVMPinDirection InDirection, const FString& InCPPType, const FName& InCPPTypeObjectPath, const FString& InDefaultValue, bool bSetupUndoRedo = true);
 							InternalCommands.Add(FString::Printf(TEXT("blueprint.get_controller_by_name('%s').add_exposed_pin('%s', unreal.RigVMPinDirection.INPUT, '%s', '%s', '%s')"),
 									*Graph->GetGraphName(),
@@ -1815,6 +1826,7 @@ TArray<FString> UControlRigBlueprint::GeneratePythonCommands(const FString InNew
 									*Pin->GetDefaultValue()));
 						}
 
+						bHitFirstExecute = false;
 						for (auto Pin : ReturnNode->GetPins())
 						{
 							if (Pin->GetDirection() != ERigVMPinDirection::Input)
@@ -1822,6 +1834,15 @@ TArray<FString> UControlRigBlueprint::GeneratePythonCommands(const FString InNew
 								continue;
 							}
 
+							if(Pin->IsExecuteContext())
+							{
+								if(!bHitFirstExecute)
+								{
+									bHitFirstExecute = true;
+									continue;
+								}
+							}
+
 							// FName AddExposedPin(const FName& InPinName, ERigVMPinDirection InDirection, const FString& InCPPType, const FName& InCPPTypeObjectPath, const FString& InDefaultValue, bool bSetupUndoRedo = true);
 							InternalCommands.Add(FString::Printf(TEXT("blueprint.get_controller_by_name('%s').add_exposed_pin('%s', unreal.RigVMPinDirection.OUTPUT, '%s', '%s', '%s')"),
 									*Graph->GetGraphName(),
diff --git a/Engine/Plugins/Animation/ControlRig/Source/ControlRigDeveloper/Private/Graph/ControlRigGraphNode.cpp b/Engine/Plugins/Animation/ControlRig/Source/ControlRigDeveloper/Private/Graph/ControlRigGraphNode.cpp
--- a/Engine/Plugins/Animation/ControlRig/Source/ControlRigDeveloper/Private/Graph/ControlRigGraphNode.cpp
+++ b/Engine/Plugins/Animation/ControlRig/Source/ControlRigDeveloper/Private/Graph/ControlRigGraphNode.cpp
@@ -25,6 +25,7 @@
 #include "RigVMCore/RigVMExecuteContext.h"
 #include "ControlRigDeveloper.h"
 #include "ControlRigObjectVersion.h"
+#include "RigVMModel/Nodes/RigVMAggregateNode.h"
 #include "RigVMModel/Nodes/RigVMFunctionReferenceNode.h"
 #include "RigVMModel/Nodes/RigVMFunctionEntryNode.h"
 #include "RigVMModel/Nodes/RigVMFunctionReturnNode.h"
@@ -90,8 +91,11 @@ FText UControlRigGraphNode::GetNodeTitle(ENodeTitleType::Type TitleType) const
 
 			else if(URigVMCollapseNode* CollapseNode = Cast<URigVMCollapseNode>(ModelNode))
 			{
-				static const FString CollapseNodeString = TEXT("Collapsed Graph");
-				SubTitle = CollapseNodeString;
+				if(!CollapseNode->IsA<URigVMAggregateNode>())
+				{
+					static const FString CollapseNodeString = TEXT("Collapsed Graph");
+					SubTitle = CollapseNodeString;
+				}
 			}
 
 			else if(URigVMVariableNode* VariableNode = Cast<URigVMVariableNode>(ModelNode))
@@ -799,11 +803,16 @@ FSlateIcon UControlRigGraphNode::GetIconAndTint(FLinearColor& OutColor) const
 			return EventIcon;
 		}
 
+		while(const URigVMAggregateNode* AggregateNode = Cast<URigVMAggregateNode>(ModelNode))
+		{
+			ModelNode = AggregateNode->GetFirstInnerNode();
+		}
+
 		if (ModelNode->IsA<URigVMFunctionReferenceNode>())
 		{ 
 			return FunctionIcon;
 		}
-		
+
 		if (ModelNode->IsA<URigVMCollapseNode>())
 		{
 			return CollapsedNodeIcon;
diff --git a/Engine/Plugins/Animation/ControlRig/Source/ControlRigEditor/Private/Editor/ControlRigEditor.cpp b/Engine/Plugins/Animation/ControlRig/Source/ControlRigEditor/Private/Editor/ControlRigEditor.cpp
--- a/Engine/Plugins/Animation/ControlRig/Source/ControlRigEditor/Private/Editor/ControlRigEditor.cpp
+++ b/Engine/Plugins/Animation/ControlRig/Source/ControlRigEditor/Private/Editor/ControlRigEditor.cpp
@@ -2300,7 +2300,7 @@ FReply FControlRigEditor::OnSpawnGraphNodeByShortcut(FInputChord InChord, const
 				}
 				else if(InChord.Key == EKeys::S)
 				{
-					Controller->AddUnitNode(FRigUnit_SequenceExecution::StaticStruct(), FRigUnit::GetMethodName(), InPosition, FString(), true, true);
+					Controller->AddUnitNode(FRigUnit_SequenceAggregate::StaticStruct(), FRigUnit::GetMethodName(), InPosition, FString(), true, true);
 				}
 				else if(InChord.Key == EKeys::One)
 				{
diff --git a/Engine/Plugins/Animation/ControlRig/Source/ControlRigEditor/Private/Graph/SControlRigGraphNode.cpp b/Engine/Plugins/Animation/ControlRig/Source/ControlRigEditor/Private/Graph/SControlRigGraphNode.cpp
--- a/Engine/Plugins/Animation/ControlRig/Source/ControlRigEditor/Private/Graph/SControlRigGraphNode.cpp
+++ b/Engine/Plugins/Animation/ControlRig/Source/ControlRigEditor/Private/Graph/SControlRigGraphNode.cpp
@@ -32,6 +32,7 @@
 #include "DetailLayoutBuilder.h"
 #include "EditorStyleSet.h"
 #include "SControlRigGraphPinVariableBinding.h"
+#include "RigVMModel/Nodes/RigVMAggregateNode.h"
 #include "Slate/SlateTextures.h"
 
 #if WITH_EDITOR
@@ -513,6 +514,8 @@ void SControlRigGraphNode::Construct( const FArguments& InArgs )
 			];
 		}
 	}
+
+	CreateAggregateAddPinButton();
 	
 	// add spacer widget at the end
 	LeftNodeBox->AddSlot()
@@ -580,56 +583,32 @@ bool SControlRigGraphNode::UseLowDetailPinNames() const
 	return false;
 }
 
-void SControlRigGraphNode::CreateInputSideAddButton(TSharedPtr<SVerticalBox> InputBox)
-{
-	if (ModelNode.IsValid() && ModelNode->IsAggregate() && ModelNode->IsInputAggregate())
-	{
-		TSharedRef<SWidget> AddPinButton = AddPinButtonContent(
-		   NSLOCTEXT("ControlRigAggregateNode", "ControlRigAggregateNodeAddPinButton", "Add pin"),
-		   NSLOCTEXT("ControlRigAggregateNode", "ControlRigAggregateNodeAddPinButton_Tooltip", "Adds a input pin to the node"),
-		   false);
-
-		FMargin AddPinPadding = Settings->GetInputPinPadding();
-		AddPinPadding.Top += 6.0f;
-
-		InputBox->AddSlot()
-			.AutoHeight()
-			.VAlign(VAlign_Center)
-			.Padding(AddPinPadding)
-			[
-				AddPinButton
-			];
-	}
-	else
-	{
-		SGraphNode::CreateInputSideAddButton(InputBox);
-	}
-}
-
-void SControlRigGraphNode::CreateOutputSideAddButton(TSharedPtr<SVerticalBox> InputBox)
+void SControlRigGraphNode::CreateAggregateAddPinButton()
 {
-	if (ModelNode.IsValid() && ModelNode->IsAggregate() && !ModelNode->IsInputAggregate())
+	if (LeftNodeBox.IsValid() && ModelNode.IsValid() && (ModelNode->IsAggregate() || ModelNode->IsA<URigVMAggregateNode>()))
 	{
-		TSharedRef<SWidget> AddPinButton = AddPinButtonContent(
-		   NSLOCTEXT("LayeredBoneBlendNode", "LayeredBoneBlendNodeAddPinButton", "Add pin"),
-		   NSLOCTEXT("LayeredBoneBlendNode", "LayeredBoneBlendNodeAddPinButton_Tooltip", "Adds a output pin to the node"),
-		   false);
-
-		FMargin AddPinPadding = Settings->GetOutputPinPadding();
-		AddPinPadding.Top += 6.0f;
-
-		InputBox->AddSlot()
+		const bool bInputAggregate = ModelNode->IsInputAggregate();
+		const TSharedRef<SWidget> AddPinButton = AddPinButtonContent(
+		   LOCTEXT("ControlRigAggregateNodeAddPinButton", "Add pin"),
+		   bInputAggregate ? 
+			LOCTEXT("ControlRigAggregateNodeAddInputPinButton_Tooltip", "Adds an input pin to the node") :
+			LOCTEXT("ControlRigAggregateNodeAddOutputPinButton_Tooltip", "Adds an output pin to the node"),
+		   !bInputAggregate);
+
+		FMargin AddPinPadding = bInputAggregate ? Settings->GetInputPinPadding() : Settings->GetOutputPinPadding();
+		AddPinPadding.Top += 2.0f;
+		AddPinPadding.Left -= bInputAggregate ? 2.f : 0.f;
+		AddPinPadding.Right -= bInputAggregate ? 0.f : 2.f;
+
+		LeftNodeBox->AddSlot()
 			.AutoHeight()
 			.VAlign(VAlign_Center)
+			.HAlign(bInputAggregate ? HAlign_Left : HAlign_Right)
 			.Padding(AddPinPadding)
 			[
 				AddPinButton
 			];
 	}
-	else
-	{
-		SGraphNode::CreateOutputSideAddButton(InputBox);
-	}
 }
 
 FReply SControlRigGraphNode::OnAddPin()
diff --git a/Engine/Plugins/Animation/ControlRig/Source/ControlRigEditor/Private/Graph/SControlRigGraphNode.h b/Engine/Plugins/Animation/ControlRig/Source/ControlRigEditor/Private/Graph/SControlRigGraphNode.h
--- a/Engine/Plugins/Animation/ControlRig/Source/ControlRigEditor/Private/Graph/SControlRigGraphNode.h
+++ b/Engine/Plugins/Animation/ControlRig/Source/ControlRigEditor/Private/Graph/SControlRigGraphNode.h
@@ -53,8 +53,7 @@ public:
 	virtual bool IsHidingPinWidgets() const override { return UseLowDetailNodeContent(); }
 	virtual bool UseLowDetailPinNames() const override;
 
-	virtual void CreateInputSideAddButton(TSharedPtr<SVerticalBox> InputBox) override;
-	virtual void CreateOutputSideAddButton(TSharedPtr<SVerticalBox> InputBox) override;
+	void CreateAggregateAddPinButton();
 
 	/** Callback function executed when Add pin button is clicked */
 	virtual FReply OnAddPin() override;
diff --git a/Engine/Source/Developer/RigVMDeveloper/Private/RigVMModel/Nodes/RigVMAggregateNode.cpp b/Engine/Source/Developer/RigVMDeveloper/Private/RigVMModel/Nodes/RigVMAggregateNode.cpp
--- a/Engine/Source/Developer/RigVMDeveloper/Private/RigVMModel/Nodes/RigVMAggregateNode.cpp
+++ b/Engine/Source/Developer/RigVMDeveloper/Private/RigVMModel/Nodes/RigVMAggregateNode.cpp
@@ -13,11 +13,11 @@ URigVMAggregateNode::URigVMAggregateNode()
 
 bool URigVMAggregateNode::IsInputAggregate() const
 {
-	for (URigVMNode* Node : GetContainedNodes())
+	for (const URigVMNode* Node : GetContainedNodes())
 	{
 		if (!Node->IsA<URigVMFunctionEntryNode>() && !Node->IsA<URigVMFunctionReturnNode>())
 		{
-			if (URigVMPin* FirstAggregatePin = Node->GetFirstAggregatePin())
+			if (const URigVMPin* FirstAggregatePin = Node->GetFirstAggregatePin())
 			{
 				return FirstAggregatePin->GetDirection() == ERigVMPinDirection::Input;
 			}
@@ -40,13 +40,13 @@ URigVMNode* URigVMAggregateNode::GetFirstInnerNode() const
 		if (IsInputAggregate())
 		{
 			// Find node connected twice to the entry (through aggregate arguments)
-			FString Arg1Name = GetFirstAggregatePin()->GetName();
-			FString Arg2Name = GetSecondAggregatePin()->GetName();
-			URigVMFunctionEntryNode* EntryNode = GetEntryNode();
+			const FString Arg1Name = GetFirstAggregatePin()->GetName();
+			const FString Arg2Name = GetSecondAggregatePin()->GetName();
+			const URigVMFunctionEntryNode* EntryNode = GetEntryNode();
 			TArray<URigVMNode*> ConnectedNodes;
-			for (URigVMPin* EntryPin : EntryNode->GetPins())
+			for (const URigVMPin* EntryPin : EntryNode->GetPins())
 			{
-				TArray<URigVMPin*> TargetPins = EntryPin->GetLinkedTargetPins();
+				const TArray<URigVMPin*> TargetPins = EntryPin->GetLinkedTargetPins();
 				if (TargetPins.Num() > 0)
 				{
 					if (TargetPins[0]->GetName() == Arg1Name || TargetPins[0]->GetName() == Arg2Name)
@@ -66,11 +66,11 @@ URigVMNode* URigVMAggregateNode::GetFirstInnerNode() const
 		else
 		{
 			// Find node connected to entry throught the opposite aggregate argument 
-			FString ArgOppositeName = GetOppositeAggregatePin()->GetName();
-			URigVMFunctionEntryNode* EntryNode = GetEntryNode();
-			if (URigVMPin* EntryPin = EntryNode->FindPin(ArgOppositeName))
+			const FString ArgOppositeName = GetOppositeAggregatePin()->GetName();
+			const URigVMFunctionEntryNode* EntryNode = GetEntryNode();
+			if (const URigVMPin* EntryPin = EntryNode->FindPin(ArgOppositeName))
 			{
-				TArray<URigVMPin*> TargetPins = EntryPin->GetLinkedTargetPins();
+				const TArray<URigVMPin*> TargetPins = EntryPin->GetLinkedTargetPins();
 				if (TargetPins.Num() > 0)
 				{
 					FirstInnerNodeCache = TargetPins[0]->GetNode();
@@ -89,46 +89,14 @@ URigVMNode* URigVMAggregateNode::GetLastInnerNode() const
 
 	if (LastInnerNodeCache == nullptr)
 	{
-		FString ArgOppositeName = GetOppositeAggregatePin()->GetName();
+		const FString ArgOppositeName = GetOppositeAggregatePin()->GetName();
 		if (IsInputAggregate())
 		{
-			// Find node connected to return throught the opposite aggregate argument
-			URigVMFunctionReturnNode* ReturnNode = GetReturnNode();
-			if (URigVMPin* ReturnPin = ReturnNode->FindPin(ArgOppositeName))
-			{
-				TArray<URigVMPin*> SourcePins = ReturnPin->GetLinkedSourcePins();
-				if(SourcePins.Num() > 0)
-				{
-					LastInnerNodeCache = SourcePins[0]->GetNode();
-					return LastInnerNodeCache;
-				}
-			}
+			LastInnerNodeCache = GetAggregateInputs().Last()->GetLinkedTargetPins()[0]->GetNode();
 		}
 		else
 		{
-			// Find node connected twice to the return (through aggregate arguments)
-			FString Arg1Name = GetFirstAggregatePin()->GetName();
-			FString Arg2Name = GetSecondAggregatePin()->GetName();
-			URigVMFunctionReturnNode* ReturnNode = GetReturnNode();
-			TArray<URigVMNode*> ConnectedNodes;
-			for (URigVMPin* ReturnPin : ReturnNode->GetPins())
-			{
-				TArray<URigVMPin*> SourcePins = ReturnPin->GetLinkedSourcePins();
-				if (SourcePins.Num() > 0)
-				{
-					if (SourcePins[0]->GetName() == Arg1Name || SourcePins[0]->GetName() == Arg2Name)
-					{
-						URigVMNode* SourceNode = SourcePins[0]->GetNode();
-						if (ConnectedNodes.Contains(SourceNode))
-						{
-							LastInnerNodeCache = SourceNode;
-							return LastInnerNodeCache;
-						}
-
-						ConnectedNodes.Add(SourceNode);
-					}
-				}
-			}
+			LastInnerNodeCache = GetAggregateOutputs().Last()->GetLinkedSourcePins()[0]->GetNode();
 		}
 	}
 #endif
@@ -138,7 +106,7 @@ URigVMNode* URigVMAggregateNode::GetLastInnerNode() const
 URigVMPin* URigVMAggregateNode::GetFirstAggregatePin() const
 {
 #if UE_RIGVM_AGGREGATE_NODES_ENABLED
-	for (URigVMNode* Node : GetContainedNodes())
+	for (const URigVMNode* Node : GetContainedNodes())
 	{
 		if (!Node->IsA<URigVMFunctionEntryNode>() && !Node->IsA<URigVMFunctionReturnNode>())
 		{
@@ -152,7 +120,7 @@ URigVMPin* URigVMAggregateNode::GetFirstAggregatePin() const
 URigVMPin* URigVMAggregateNode::GetSecondAggregatePin() const
 {
 #if UE_RIGVM_AGGREGATE_NODES_ENABLED
-	for (URigVMNode* Node : GetContainedNodes())
+	for (const URigVMNode* Node : GetContainedNodes())
 	{
 		if (!Node->IsA<URigVMFunctionEntryNode>() && !Node->IsA<URigVMFunctionReturnNode>())
 		{
@@ -166,7 +134,7 @@ URigVMPin* URigVMAggregateNode::GetSecondAggregatePin() const
 URigVMPin* URigVMAggregateNode::GetOppositeAggregatePin() const
 {
 #if UE_RIGVM_AGGREGATE_NODES_ENABLED
-	for (URigVMNode* Node : GetContainedNodes())
+	for (const URigVMNode* Node : GetContainedNodes())
 	{
 		if (!Node->IsA<URigVMFunctionEntryNode>() && !Node->IsA<URigVMFunctionReturnNode>())
 		{
@@ -177,15 +145,74 @@ URigVMPin* URigVMAggregateNode::GetOppositeAggregatePin() const
 	return nullptr;
 }
 
+TArray<URigVMPin*> URigVMAggregateNode::GetAggregateInputs() const
+{
+#if UE_RIGVM_AGGREGATE_NODES_ENABLED
+	
+	TArray<URigVMPin*> Result;
+	const FString Arg1Name = GetFirstAggregatePin()->GetName();
+	const FString Arg2Name = GetSecondAggregatePin()->GetName();
+
+	// Find pins connected to an aggregate pin
+	const URigVMFunctionEntryNode* EntryNode = GetEntryNode();
+	for(URigVMPin* EntryPin : EntryNode->GetPins())
+	{
+		const TArray<URigVMPin*> SourcePins = EntryPin->GetLinkedTargetPins();
+		if(SourcePins.Num() > 0)
+		{
+			const URigVMPin* SourcePin = SourcePins[0];
+			if(SourcePin->GetName() == Arg1Name || SourcePin->GetName() == Arg2Name)
+			{
+				Result.Add(EntryPin);
+			}
+		}
+	}
+	return Result;
+	
+#else
+	return Super::GetAggregateInputs();
+#endif
+}
+
+TArray<URigVMPin*> URigVMAggregateNode::GetAggregateOutputs() const
+{
+#if UE_RIGVM_AGGREGATE_NODES_ENABLED
+
+	TArray<URigVMPin*> Result;
+	const FString Arg1Name = GetFirstAggregatePin()->GetName();
+	const FString Arg2Name = GetSecondAggregatePin()->GetName();
+
+	// Find pins connected to an aggregate pin
+	const URigVMFunctionReturnNode* ReturnNode = GetReturnNode();
+	for(URigVMPin* ReturnPin : ReturnNode->GetPins())
+	{
+		const TArray<URigVMPin*> SourcePins = ReturnPin->GetLinkedSourcePins();
+		if(SourcePins.Num() > 0)
+		{
+			const URigVMPin* SourcePin = SourcePins[0];
+			if(SourcePin->GetName() == Arg1Name || SourcePin->GetName() == Arg2Name)
+			{
+				Result.Add(ReturnPin);
+			}
+		}
+	}
+	return Result;
+	
+#else
+	return Super::GetAggregateOutputs();
+#endif
+}
+
 void URigVMAggregateNode::InvalidateCache()
 {
+	Super::InvalidateCache();
 	FirstInnerNodeCache = nullptr;
 	LastInnerNodeCache = nullptr;
 }
 
 FString URigVMAggregateNode::GetNodeTitle() const
 {
-	if (URigVMNode* InnerNode = GetFirstInnerNode())
+	if (const URigVMNode* InnerNode = GetFirstInnerNode())
 	{
 		return InnerNode->GetNodeTitle();
 	}
@@ -193,9 +220,19 @@ FString URigVMAggregateNode::GetNodeTitle() const
 	return Super::GetNodeTitle();
 }
 
+FLinearColor URigVMAggregateNode::GetNodeColor() const
+{
+	if (const URigVMNode* InnerNode = GetFirstInnerNode())
+	{
+		return InnerNode->GetNodeColor();
+	}
+	
+	return Super::GetNodeColor();
+}
+
 FName URigVMAggregateNode::GetMethodName() const
 {
-	if (URigVMUnitNode* InnerNode = Cast<URigVMUnitNode>(GetFirstInnerNode()))
+	if (const URigVMUnitNode* InnerNode = Cast<URigVMUnitNode>(GetFirstInnerNode()))
 	{
 		return InnerNode->GetMethodName();
 	}
@@ -205,7 +242,7 @@ FName URigVMAggregateNode::GetMethodName() const
 
 FText URigVMAggregateNode::GetToolTipText() const
 {
-	if (URigVMNode* InnerNode = GetFirstInnerNode())
+	if (const URigVMNode* InnerNode = GetFirstInnerNode())
 	{
 		return InnerNode->GetToolTipText();
 	}
@@ -215,7 +252,7 @@ FText URigVMAggregateNode::GetToolTipText() const
 
 FText URigVMAggregateNode::GetToolTipTextForPin(const URigVMPin* InPin) const
 {
-	if (URigVMPin* EntryPin = GetContainedGraph()->GetEntryNode()->FindPin(InPin->GetName()))
+	if (const URigVMPin* EntryPin = GetContainedGraph()->GetEntryNode()->FindPin(InPin->GetName()))
 	{
 		const TArray<URigVMPin*> Targets = EntryPin->GetLinkedTargetPins();
 		if (Targets.Num() > 0)
@@ -223,7 +260,7 @@ FText URigVMAggregateNode::GetToolTipTextForPin(const URigVMPin* InPin) const
 			return Targets[0]->GetToolTipText();
 		}
 	}
-	else if(URigVMPin* ReturnPin = GetContainedGraph()->GetReturnNode()->FindPin(InPin->GetName()))
+	else if(const URigVMPin* ReturnPin = GetContainedGraph()->GetReturnNode()->FindPin(InPin->GetName()))
 	{
 		const TArray<URigVMPin*> Sources = ReturnPin->GetLinkedSourcePins();
 		if (Sources.Num() > 0)
diff --git a/Engine/Source/Developer/RigVMDeveloper/Private/RigVMModel/Nodes/RigVMUnitNode.cpp b/Engine/Source/Developer/RigVMDeveloper/Private/RigVMModel/Nodes/RigVMUnitNode.cpp
--- a/Engine/Source/Developer/RigVMDeveloper/Private/RigVMModel/Nodes/RigVMUnitNode.cpp
+++ b/Engine/Source/Developer/RigVMDeveloper/Private/RigVMModel/Nodes/RigVMUnitNode.cpp
@@ -157,8 +157,8 @@ TArray<FRigVMUserWorkflow> URigVMUnitNode::GetSupportedWorkflows(ERigVMUserWorkf
 bool URigVMUnitNode::IsAggregate() const
 {
 #if UE_RIGVM_AGGREGATE_NODES_ENABLED
-	TArray<URigVMPin*> AggregateInputs = GetAggregateInputs();
-	TArray<URigVMPin*> AggregateOutputs = GetAggregateOutputs();
+	const TArray<URigVMPin*> AggregateInputs = GetAggregateInputs();
+	const TArray<URigVMPin*> AggregateOutputs = GetAggregateOutputs();
 
 	if ((AggregateInputs.Num() == 2 && AggregateOutputs.Num() == 1) ||
 		(AggregateInputs.Num() == 1 && AggregateOutputs.Num() == 2))
@@ -184,8 +184,8 @@ bool URigVMUnitNode::IsAggregate() const
 URigVMPin* URigVMUnitNode::GetFirstAggregatePin() const
 {
 #if UE_RIGVM_AGGREGATE_NODES_ENABLED
-	TArray<URigVMPin*> Inputs = GetAggregateInputs();
-	TArray<URigVMPin*> Outputs = GetAggregateOutputs();
+	const TArray<URigVMPin*> Inputs = GetAggregateInputs();
+	const TArray<URigVMPin*> Outputs = GetAggregateOutputs();
 	if (Inputs.Num() == 2 && Outputs.Num() == 1)
 	{
 		return Inputs[0];
@@ -201,8 +201,8 @@ URigVMPin* URigVMUnitNode::GetFirstAggregatePin() const
 URigVMPin* URigVMUnitNode::GetSecondAggregatePin() const
 {
 #if UE_RIGVM_AGGREGATE_NODES_ENABLED
-	TArray<URigVMPin*> Inputs = GetAggregateInputs();
-	TArray<URigVMPin*> Outputs = GetAggregateOutputs();
+	const TArray<URigVMPin*> Inputs = GetAggregateInputs();
+	const TArray<URigVMPin*> Outputs = GetAggregateOutputs();
 	if (Inputs.Num() == 2 && Outputs.Num() == 1)
 	{
 		return Inputs[1];
@@ -218,8 +218,8 @@ URigVMPin* URigVMUnitNode::GetSecondAggregatePin() const
 URigVMPin* URigVMUnitNode::GetOppositeAggregatePin() const
 {
 #if UE_RIGVM_AGGREGATE_NODES_ENABLED
-	TArray<URigVMPin*> Inputs = GetAggregateInputs();
-	TArray<URigVMPin*> Outputs = GetAggregateOutputs();
+	const TArray<URigVMPin*> Inputs = GetAggregateInputs();
+	const TArray<URigVMPin*> Outputs = GetAggregateOutputs();
 	if (Inputs.Num() == 2 && Outputs.Num() == 1)
 	{
 		return Outputs[0];
@@ -241,13 +241,13 @@ TArray<URigVMPin*> URigVMUnitNode::GetAggregateInputs() const
 {
 	TArray<URigVMPin*> AggregateInputs;
 #if UE_RIGVM_AGGREGATE_NODES_ENABLED
-	if (UScriptStruct* Struct = GetScriptStruct())
+	if (const UScriptStruct* Struct = GetScriptStruct())
 	{
 		for (URigVMPin* Pin : GetPins())
 		{
 			if (Pin->GetDirection() == ERigVMPinDirection::Input)
 			{
-				if (FProperty* Property = Struct->FindPropertyByName(Pin->GetFName()))
+				if (const FProperty* Property = Struct->FindPropertyByName(Pin->GetFName()))
 				{
 					if (Property->HasMetaData(FRigVMStruct::AggregateMetaName))
 					{
@@ -265,13 +265,13 @@ TArray<URigVMPin*> URigVMUnitNode::GetAggregateOutputs() const
 {
 	TArray<URigVMPin*> AggregateOutputs;
 #if UE_RIGVM_AGGREGATE_NODES_ENABLED	
-	if (UScriptStruct* Struct = GetScriptStruct())
+	if (const UScriptStruct* Struct = GetScriptStruct())
 	{
 		for (URigVMPin* Pin : GetPins())
 		{
 			if (Pin->GetDirection() == ERigVMPinDirection::Output)
 			{
-				if (FProperty* Property = Struct->FindPropertyByName(Pin->GetFName()))
+				if (const FProperty* Property = Struct->FindPropertyByName(Pin->GetFName()))
 				{
 					if (Property->HasMetaData(FRigVMStruct::AggregateMetaName))
 					{
@@ -288,7 +288,7 @@ TArray<URigVMPin*> URigVMUnitNode::GetAggregateOutputs() const
 FName URigVMUnitNode::GetNextAggregateName(const FName& InLastAggregatePinName) const
 {
 #if UE_RIGVM_AGGREGATE_NODES_ENABLED	
-	if(UScriptStruct* Struct = GetScriptStruct())
+	if(const UScriptStruct* Struct = GetScriptStruct())
 	{
 		check(Struct->IsChildOf(FRigVMStruct::StaticStruct()));
 
@@ -311,7 +311,7 @@ UScriptStruct* URigVMUnitNode::GetScriptStruct() const
 
 bool URigVMUnitNode::IsLoopNode() const
 {
-	TSharedPtr<FStructOnScope> StructOnScope = ConstructStructInstance(true);
+	const TSharedPtr<FStructOnScope> StructOnScope = ConstructStructInstance(true);
 	if (StructOnScope.IsValid())
 	{
 		const FRigVMStruct* StructMemory = (FRigVMStruct*)StructOnScope->GetStructMemory();
diff --git a/Engine/Source/Developer/RigVMDeveloper/Private/RigVMModel/RigVMController.cpp b/Engine/Source/Developer/RigVMDeveloper/Private/RigVMModel/RigVMController.cpp
--- a/Engine/Source/Developer/RigVMDeveloper/Private/RigVMModel/RigVMController.cpp
+++ b/Engine/Source/Developer/RigVMDeveloper/Private/RigVMModel/RigVMController.cpp
@@ -436,6 +436,24 @@ TArray<FString> URigVMController::GetAddNodePythonCommands(URigVMNode* Node) con
 					*NodeName));
 		}
 	}
+	else if (const URigVMAggregateNode* AggregateNode = Cast<URigVMAggregateNode>(Node))
+	{
+		TArray<FString> InnerNodeCommands = GetAddNodePythonCommands(AggregateNode->GetFirstInnerNode());
+		Commands.Append(InnerNodeCommands);
+
+		// add commands for any additional aggregate pin
+		const TArray<URigVMPin*> AggregatePins = AggregateNode->IsInputAggregate() ? AggregateNode->GetAggregateInputs() : AggregateNode->GetAggregateOutputs();
+
+		for(int32 Index = 2; Index < AggregatePins.Num(); Index++)
+		{
+			// add_aggregate_pin(node_name, pin_name)
+			Commands.Add(FString::Printf(TEXT("blueprint.get_controller_by_name('%s').add_aggregate_pin('%s', '%s')"),
+					*GraphName,
+					*AggregateNode->GetName(),
+					*AggregatePins[Index]->GetName()
+				));
+		}
+	}
 	else if (const URigVMVariableNode* VariableNode = Cast<URigVMVariableNode>(Node))
 	{
 		if (!VariableNode->IsInjected())
@@ -2069,8 +2087,8 @@ URigVMNode* URigVMController::UpgradeNode(URigVMNode* InNode, bool bSetupUndoRed
 			return nullptr;
 		}
 
-		URigVMUnitNode* NewUnitNode = AddUnitNode(UpgradeInfo.GetNewStruct(), MethodName, NodePosition, NodeName, bSetupUndoRedo, false);
-		if(NewUnitNode == nullptr)
+		URigVMNode* NewNode = AddUnitNode(UpgradeInfo.GetNewStruct(), MethodName, NodePosition, NodeName, bSetupUndoRedo, false);
+		if(NewNode == nullptr)
 		{
 			if(bSetupUndoRedo)
 			{
@@ -2080,7 +2098,15 @@ URigVMNode* URigVMController::UpgradeNode(URigVMNode* InNode, bool bSetupUndoRed
 			return nullptr;
 		}
 
-		for(URigVMPin* Pin : NewUnitNode->GetPins())
+		const TArray<FString>& AggregatePins = UpgradeInfo.GetAggregatePins();
+		for(const FString& AggregatePinName : AggregatePins)
+		{
+			const FName PreviousName = NewNode->GetFName();
+			AddAggregatePin(PreviousName.ToString(), AggregatePinName, FString(), bSetupUndoRedo, false);
+			NewNode = GetGraph()->FindNodeByName(PreviousName);
+		}
+
+		for(URigVMPin* Pin : NewNode->GetPins())
 		{
 			const FString DefaultValue = UpgradeInfo.GetDefaultValueForPin(Pin->GetFName());
 			if(!DefaultValue.IsEmpty())
@@ -2106,7 +2132,7 @@ URigVMNode* URigVMController::UpgradeNode(URigVMNode* InNode, bool bSetupUndoRed
 		}
 
 
-		UpgradedNode = NewUnitNode;
+		UpgradedNode = NewNode;
 	}
 	else
 	{
@@ -3969,35 +3995,15 @@ URigVMCollapseNode* URigVMController::CollapseNodes(const TArray<URigVMNode*>& I
 #if UE_RIGVM_AGGREGATE_NODES_ENABLED
 	if (bIsAggregate)
 	{
-		if (URigVMUnitNode* UnitNode = Cast<URigVMUnitNode>(InNodes[0]))
+		if(InNodes.Num() != 1)
 		{
-			for (int32 i=1; i<InNodes.Num(); ++i)
-			{
-				if (URigVMUnitNode* OtherUnitNode = Cast<URigVMUnitNode>(InNodes[i]))
-				{
-					if (UnitNode->GetScriptStruct() == OtherUnitNode->GetScriptStruct())
-					{
-						continue;
-					}
-				}
-				ReportError(TEXT("Cannot aggregate nodes of different types."));
-				return nullptr;
-			}
+			return nullptr;
 		}
-		else if (URigVMFunctionReferenceNode* FunctionReferenceNode = Cast<URigVMFunctionReferenceNode>(InNodes[0]))
+
+		if(!InNodes[0]->IsAggregate())
 		{
-			for (int32 i=1; i<InNodes.Num(); ++i)
-			{
-				if (URigVMFunctionReferenceNode* OtherFunctionReferenceNode = Cast<URigVMFunctionReferenceNode>(InNodes[i]))
-				{
-					if(FunctionReferenceNode->GetReferencedNode() == OtherFunctionReferenceNode->GetReferencedNode())
-					{
-						continue;
-					}
-				}
-				ReportError(TEXT("Cannot aggregate nodes of different types."));
-				return nullptr;
-			}
+			ReportError(TEXT("Cannot aggregate the given node."));
+			return nullptr;
 		}
 	}
 #endif
@@ -4300,7 +4306,7 @@ URigVMCollapseNode* URigVMController::CollapseNodes(const TArray<URigVMNode*>& I
 		}, false, true);
 
 		URigVMPin* CollapsedPin = NewObject<URigVMPin>(CollapseNode, PinName);
-		ConfigurePinFromPin(CollapsedPin, PinToCollapse);
+		ConfigurePinFromPin(CollapsedPin, PinToCollapse, true);
 
 		if (CollapsedPin->IsExecuteContext())
 		{
@@ -7583,8 +7589,13 @@ FString URigVMController::AddAggregatePin(URigVMNode* InNode, const FString& InP
 	}
 
 	URigVMAggregateNode* AggregateNode = Cast<URigVMAggregateNode>(InNode);
-	URigVMUnitNode* UnitNode = Cast<URigVMUnitNode>(InNode);
-	URigVMFunctionReferenceNode* FunctionReferenceNode = Cast<URigVMFunctionReferenceNode>(InNode);
+	if (AggregateNode == nullptr)
+	{
+		if(!InNode->IsAggregate())
+		{
+			return FString();
+		}
+	}
 
 	FRigVMControllerCompileBracketScope CompileScope(this);
 	FRigVMBaseAction Action;
@@ -7594,52 +7605,37 @@ FString URigVMController::AddAggregatePin(URigVMNode* InNode, const FString& InP
 		ActionStack->BeginAction(Action);
 	}
 
-	if (UnitNode || FunctionReferenceNode)
+	if (AggregateNode == nullptr)
 	{
-		if (UnitNode && !UnitNode->IsAggregate())
-		{
-			if (bSetupUndoRedo)
-			{
-				ActionStack->CancelAction(Action, this);
-			}
-			return FString();
-		}
-
 		bool bAggregateInputs = false;
 		URigVMPin* Arg1 = nullptr;
 		URigVMPin* Arg2 = nullptr;
 		URigVMPin* ArgOpposite = nullptr;
-		if (UnitNode)
-		{
-			const TArray<URigVMPin*> AggregateInputs = UnitNode->GetAggregateInputs();
-			const TArray<URigVMPin*> AggregateOutputs = UnitNode->GetAggregateOutputs();
 
-			if (AggregateInputs.Num() == 2 && AggregateOutputs.Num() == 1)
-			{
-				bAggregateInputs = true;
-				Arg1 = AggregateInputs[0];
-				Arg2 = AggregateInputs[1];
-				ArgOpposite = AggregateOutputs[0];
-			}
-			else if (AggregateInputs.Num() == 1 && AggregateOutputs.Num() == 2)
-			{
-				bAggregateInputs = false;
-				Arg1 = AggregateOutputs[0];
-				Arg2 = AggregateOutputs[1];
-				ArgOpposite = AggregateInputs[0];
-			}
-			else
-			{
-				if (bSetupUndoRedo)
-				{
-					ActionStack->CancelAction(Action, this);
-				}
-				return FString();	
-			}
+		const TArray<URigVMPin*> AggregateInputs = InNode->GetAggregateInputs();
+		const TArray<URigVMPin*> AggregateOutputs = InNode->GetAggregateOutputs();
+
+		if (AggregateInputs.Num() == 2 && AggregateOutputs.Num() == 1)
+		{
+			bAggregateInputs = true;
+			Arg1 = AggregateInputs[0];
+			Arg2 = AggregateInputs[1];
+			ArgOpposite = AggregateOutputs[0];
 		}
-		else if (FunctionReferenceNode)
+		else if (AggregateInputs.Num() == 1 && AggregateOutputs.Num() == 2)
 		{
-			
+			bAggregateInputs = false;
+			Arg1 = AggregateOutputs[0];
+			Arg2 = AggregateOutputs[1];
+			ArgOpposite = AggregateInputs[0];
+		}
+		else
+		{
+			if (bSetupUndoRedo)
+			{
+				ActionStack->CancelAction(Action, this);
+			}
+			return FString();	
 		}
 
 		if (!Arg1 || !Arg2 || !ArgOpposite)
@@ -7664,9 +7660,9 @@ FString URigVMController::AddAggregatePin(URigVMNode* InNode, const FString& InP
 		const FString AggregateArg1 = Arg1->GetName();
 		const FString AggregateArg2 = Arg2->GetName();
 		const FString AggregateOppositeArg = ArgOpposite->GetName();
-	
-		TArray<URigVMNode*> Nodes = {InNode};
-		URigVMCollapseNode* CollapseNode = CollapseNodes(Nodes, InNode->GetName(), bSetupUndoRedo, true);
+
+		const FName PreviousNodeName = InNode->GetFName();
+		URigVMCollapseNode* CollapseNode = CollapseNodes({InNode}, InNode->GetName(), bSetupUndoRedo, true);
 		if (!CollapseNode)
 		{
 			if (bSetupUndoRedo)
@@ -7676,40 +7672,53 @@ FString URigVMController::AddAggregatePin(URigVMNode* InNode, const FString& InP
 			return FString();
 		}
 
+		InNode = CollapseNode->GetContainedGraph()->FindNodeByName(PreviousNodeName);
+
 		AggregateNode = Cast<URigVMAggregateNode>(CollapseNode);
 		if (AggregateNode)
 		{
 			FRigVMControllerGraphGuard GraphGuard(this, AggregateNode->GetContainedGraph(), bSetupUndoRedo);
 			TGuardValue<bool> GuardEditGraph(GetGraph()->bEditable, true);
-			if (UnitNode)
+
+			for(int32 Index = 0; Index < InNode->GetPins().Num(); Index++)
 			{
-				uint16 Index = 0;
-				for (TFieldIterator<FProperty> It(UnitNode->GetScriptStruct()); It; ++It, ++Index)
+				URigVMPin* Pin = InNode->GetPins()[Index];
+				const FName PinName = Pin->GetFName();
+				
+				if (URigVMPin* AggregatePin = AggregateNode->FindPin(PinName.ToString()))
 				{
-					FName PropertyName = It->GetFName();
-					if (URigVMPin* AggregatePin = AggregateNode->FindPin(PropertyName.ToString()))
-					{
-						SetExposedPinIndex(PropertyName, Index, bSetupUndoRedo);
-						continue;
-					}
+					SetExposedPinIndex(PinName, Index, bSetupUndoRedo);
+					continue;
+				}
 
-					URigVMPin* Pin = UnitNode->FindPin(PropertyName.ToString());
-					AddExposedPin(PropertyName, Pin->GetDirection(), Pin->GetCPPType(), *Pin->GetCPPTypeObject()->GetPathName(), Pin->GetDefaultValue());
+				const FName ExposedPinName = AddExposedPin(PinName, Pin->GetDirection(), Pin->GetCPPType(), *Pin->GetCPPTypeObject()->GetPathName(), Pin->GetDefaultValue());
 
-					FString PropertyNameStr = PropertyName.ToString();
-					if (Pin->GetDirection() == ERigVMPinDirection::Input)
-					{
-						AddLink(FString::Printf(TEXT("Entry.%s"), *PropertyNameStr), FString::Printf(TEXT("%s.%s"), *InNode->GetName(), *PropertyNameStr), bSetupUndoRedo);
-					}
-					else
-					{
-						AddLink(FString::Printf(TEXT("%s.%s"), *InNode->GetName(), *PropertyNameStr), FString::Printf(TEXT("Return.%s"), *PropertyNameStr), bSetupUndoRedo);
-					}
+				const FString PinNameStr = PinName.ToString();
+				const FString ExposedPinNameStr = ExposedPinName.ToString();
+
+				if(URigVMPin* ExposedPin = AggregateNode->FindPin(ExposedPinNameStr))
+				{
+					ExposedPin->SetDisplayName(Pin->GetDisplayName());
+				}
+
+				if(URigVMPin* ExposedPin = AggregateNode->GetEntryNode()->FindPin(ExposedPinNameStr))
+				{
+					ExposedPin->SetDisplayName(Pin->GetDisplayName());
+				}
+
+				if(URigVMPin* ExposedPin = AggregateNode->GetReturnNode()->FindPin(ExposedPinNameStr))
+				{
+					ExposedPin->SetDisplayName(Pin->GetDisplayName());
+				}
+
+				if (Pin->GetDirection() == ERigVMPinDirection::Input)
+				{
+					AddLink(FString::Printf(TEXT("Entry.%s"), *ExposedPinNameStr), FString::Printf(TEXT("%s.%s"), *InNode->GetName(), *PinNameStr), bSetupUndoRedo);
+				}
+				else
+				{
+					AddLink(FString::Printf(TEXT("%s.%s"), *InNode->GetName(), *PinNameStr), FString::Printf(TEXT("Return.%s"), *ExposedPinNameStr), bSetupUndoRedo);
 				}
-			}
-			else if (FunctionReferenceNode)
-			{
-				
 			}
 		}
 		else
@@ -7736,19 +7745,12 @@ FString URigVMController::AddAggregatePin(URigVMNode* InNode, const FString& InP
 		FRigVMControllerGraphGuard GraphGuard(this, AggregateNode->GetContainedGraph(), bSetupUndoRedo);
 		TGuardValue<bool> GuardEditGraph(GetGraph()->bEditable, true);
 
-		URigVMNode* InnerNode = (AggregateNode == nullptr) ?
-									((UnitNode == nullptr) ? (URigVMNode*) FunctionReferenceNode : UnitNode)
-								: AggregateNode->GetFirstInnerNode();
-		URigVMNode* NewNode = nullptr;
-		if (URigVMUnitNode* InnerUnitNode = Cast<URigVMUnitNode>(InnerNode))
-		{		
-			NewNode = AddUnitNode(InnerUnitNode->GetScriptStruct(), InnerUnitNode->GetMethodName(), FVector2D::ZeroVector, FString(), bSetupUndoRedo);		
-		}
-		else if (URigVMFunctionReferenceNode* InnerReferenceNode = Cast<URigVMFunctionReferenceNode>(InnerNode))
-		{
-			NewNode = AddFunctionReferenceNode(FunctionReferenceNode->GetReferencedNode(), FVector2D::ZeroVector, FString(), bSetupUndoRedo);
-		}
-		else
+		URigVMNode* InnerNode = (AggregateNode == nullptr) ? InNode : AggregateNode->GetFirstInnerNode();
+
+		const FString InnerNodeContent = ExportNodesToText({InnerNode->GetFName()});
+		const TArray<FName> NewNodeNames = ImportNodesFromText(InnerNodeContent);
+		
+		if(NewNodeNames.IsEmpty())
 		{
 			if (bSetupUndoRedo)
 			{
@@ -7757,8 +7759,10 @@ FString URigVMController::AddAggregatePin(URigVMNode* InNode, const FString& InP
 			return FString();
 		}
 
-		FName NewName = *InPinName;
-		if (NewName.IsNone())
+		URigVMNode* NewNode = AggregateNode->GetContainedGraph()->FindNodeByName(NewNodeNames[0]);
+
+		FName NewPinName = *InPinName;
+		if (NewPinName.IsNone())
 		{
 			URigVMNode* LastInnerNode = AggregateNode->GetLastInnerNode();
 			URigVMPin* SecondAggregateInnerPin = LastInnerNode->GetSecondAggregatePin();
@@ -7780,19 +7784,17 @@ FString URigVMController::AddAggregatePin(URigVMNode* InNode, const FString& InP
 				}
 			}
 
-			if (URigVMUnitNode* InnerUnitNode = Cast<URigVMUnitNode>(LastInnerNode))
-			{
-				NewName = InnerUnitNode->GetNextAggregateName(*LastAggregateName);
-			}
+			NewPinName = InnerNode->GetNextAggregateName(*LastAggregateName);
 		}
-		if (NewName.IsNone())
+		
+		if (NewPinName.IsNone())
 		{
-			NewName = InnerNode->GetSecondAggregatePin()->GetFName();
+			NewPinName = InnerNode->GetSecondAggregatePin()->GetFName();
 		}
 		
 		const URigVMPin* Arg1 = AggregateNode->GetFirstAggregatePin();
-		FName NewPinName = AddExposedPin(NewName, Arg1->GetDirection(), Arg1->GetCPPType(), *Arg1->GetCPPTypeObject()->GetPathName(), InDefaultValue, bSetupUndoRedo);
-		NewPin = AggregateNode->FindPin(NewPinName.ToString());
+		FName NewExposedPinName = AddExposedPin(NewPinName, Arg1->GetDirection(), Arg1->GetCPPType(), *Arg1->GetCPPTypeObject()->GetPathName(), InDefaultValue, bSetupUndoRedo);
+		NewPin = AggregateNode->FindPin(NewExposedPinName.ToString());
 		URigVMPin* NewUnitPinArg1 = NewNode->GetFirstAggregatePin();
 		URigVMPin* NewUnitPinArg2 = NewNode->GetSecondAggregatePin();
 		URigVMPin* NewUnitPinOppositeArg = NewNode->GetOppositeAggregatePin();
@@ -7800,7 +7802,7 @@ FString URigVMController::AddAggregatePin(URigVMNode* InNode, const FString& InP
 		if(AggregateNode->IsInputAggregate())
 		{		
 			URigVMFunctionEntryNode* EntryNode = AggregateNode->GetEntryNode();
-			URigVMPin* EntryPin = EntryNode->FindPin(NewPinName.ToString());
+			URigVMPin* EntryPin = EntryNode->FindPin(NewExposedPinName.ToString());
 			URigVMPin* ReturnPin = AggregateNode->GetReturnNode()->FindPin(NewUnitPinOppositeArg->GetName());
 			URigVMPin* PreviousReturnPin = ReturnPin->GetLinkedSourcePins()[0];
 			PreviousNode = PreviousReturnPin->GetNode();
@@ -7813,7 +7815,7 @@ FString URigVMController::AddAggregatePin(URigVMNode* InNode, const FString& InP
 		else
 		{
 			URigVMFunctionReturnNode* ReturnNode = AggregateNode->GetReturnNode();
-			URigVMPin* NewReturnPin = ReturnNode->FindPin(NewPinName.ToString());
+			URigVMPin* NewReturnPin = ReturnNode->FindPin(NewExposedPinName.ToString());
 			URigVMPin* OldReturnPin = ReturnNode->GetPins()[ReturnNode->GetPins().Num()-2];
 			URigVMPin* PreviousReturnPin = OldReturnPin->GetLinkedSourcePins()[0];
 			PreviousNode = PreviousReturnPin->GetNode();
@@ -9484,34 +9486,21 @@ FName URigVMController::AddExposedPin(const FName& InPinName, ERigVMPinDirection
 			ReportError(TEXT("Cannot expose pins of wildcard type."));
 			return NAME_None;
 		}
-		
-		if (UScriptStruct* ScriptStruct = Cast<UScriptStruct>(CPPTypeObject))
-		{
-			if (ScriptStruct->IsChildOf(FRigVMExecuteContext::StaticStruct()))
-			{
-				for (URigVMPin* ExistingPin : LibraryNode->Pins)
-				{
-					if (ExistingPin->IsExecuteContext())
-					{
-						return NAME_None;
-					}
-				}
-
-				InDirection = ERigVMPinDirection::IO;
-			}
-		}
 	}
 
-	// only allow one exposed pin of type execute context per direction
-	if(const UScriptStruct* CPPTypeStruct = Cast<UScriptStruct>(CPPTypeObject))
+	// only allow one IO / input exposed pin of type execute context per direction
+	if(InDirection != ERigVMPinDirection::Output)
 	{
-		if(CPPTypeStruct->IsChildOf(FRigVMExecuteContext::StaticStruct()))
+		if(const UScriptStruct* CPPTypeStruct = Cast<UScriptStruct>(CPPTypeObject))
 		{
-			for(URigVMPin* ExistingPin : LibraryNode->Pins)
+			if(CPPTypeStruct->IsChildOf(FRigVMExecuteContext::StaticStruct()))
 			{
-				if(ExistingPin->IsExecuteContext())
+				for(URigVMPin* ExistingPin : LibraryNode->Pins)
 				{
-					return NAME_None;
+					if(ExistingPin->IsExecuteContext())
+					{
+						return NAME_None;
+					}
 				}
 			}
 		}
@@ -10391,6 +10380,7 @@ URigVMLibraryNode* URigVMController::AddFunctionToLibrary(const FName& InFunctio
 	{
 		UScriptStruct* ExecuteContextStruct = Graph->GetExecuteContextStruct();
 		URigVMPin* ExecutePin = NewObject<URigVMPin>(CollapseNode, FRigVMStruct::ExecuteContextName);
+		ExecutePin->DisplayName = FRigVMStruct::ExecuteName;
 		ExecutePin->CPPType = FString::Printf(TEXT("F%s"), *ExecuteContextStruct->GetName());
 		ExecutePin->CPPTypeObject = ExecuteContextStruct;
 		ExecutePin->CPPTypeObjectPath = *ExecutePin->CPPTypeObject->GetPathName();
@@ -13039,7 +13029,7 @@ void URigVMController::ConfigurePinFromProperty(FProperty* InProperty, URigVMPin
 	InOutPin->CPPType = RigVMTypeUtils::PostProcessCPPType(InOutPin->CPPType, InOutPin->GetCPPTypeObject());
 }
 
-void URigVMController::ConfigurePinFromPin(URigVMPin* InOutPin, URigVMPin* InPin)
+void URigVMController::ConfigurePinFromPin(URigVMPin* InOutPin, URigVMPin* InPin, bool bCopyDisplayName)
 {
 	// it is important we copy things that define the identity of the pin
 	// things that defines the state of the pin is copied during GetPinState()
@@ -13051,6 +13041,10 @@ void URigVMController::ConfigurePinFromPin(URigVMPin* InOutPin, URigVMPin* InPin
 	InOutPin->CPPTypeObject = InPin->CPPTypeObject;
 	InOutPin->DefaultValue = InPin->DefaultValue;
 	InOutPin->bIsDynamicArray = InPin->bIsDynamicArray;
+	if(bCopyDisplayName)
+	{
+		InOutPin->SetDisplayName(InPin->GetDisplayName());
+	}
 }
 
 bool URigVMController::ShouldStructBeUnfolded(const UStruct* Struct)
@@ -13844,7 +13838,7 @@ void URigVMController::RepopulatePinsOnNode(URigVMNode* InNode, bool bFollowCore
 				}
 
 				URigVMPin* ExposedPin = NewObject<URigVMPin>(InNode, LibraryPin->GetFName());
-				ConfigurePinFromPin(ExposedPin, LibraryPin);
+				ConfigurePinFromPin(ExposedPin, LibraryPin, true);
 
 				if (bIsEntryNode)
 				{
@@ -13881,7 +13875,7 @@ void URigVMController::RepopulatePinsOnNode(URigVMNode* InNode, bool bFollowCore
 		for (URigVMPin* RootPin : InNode->Pins)
 		{
 			URigVMPin* NewRootPin = NewObject<URigVMPin>(InNode);
-			ConfigurePinFromPin(NewRootPin, RootPin);
+			ConfigurePinFromPin(NewRootPin, RootPin, true);
 			EnsurePinValidity(NewRootPin, false);
 		
 			NewRootPinInfos.Add(TTuple<URigVMPin*, FName>(NewRootPin, RootPin->GetFName()));
@@ -13934,7 +13928,7 @@ void URigVMController::RepopulatePinsOnNode(URigVMNode* InNode, bool bFollowCore
 			for (URigVMPin* ReferencedPin : ReferencedNode->Pins)
 			{
 				URigVMPin* NewPin = NewObject<URigVMPin>(InNode, ReferencedPin->GetFName());
-				ConfigurePinFromPin(NewPin, ReferencedPin);
+				ConfigurePinFromPin(NewPin, ReferencedPin, true);
 				EnsurePinValidity(NewPin, false);
 				
 				AddNodePin(InNode, NewPin);
@@ -13960,7 +13954,7 @@ void URigVMController::RepopulatePinsOnNode(URigVMNode* InNode, bool bFollowCore
 		for (URigVMPin* RootPin : InNode->Pins)
 		{
 			URigVMPin* NewRootPin = NewObject<URigVMPin>(InNode);
-			ConfigurePinFromPin(NewRootPin, RootPin);
+			ConfigurePinFromPin(NewRootPin, RootPin, true);
 			EnsurePinValidity(NewRootPin, false);
 		
 			NewRootPinInfos.Add(TTuple<URigVMPin*, FName>(NewRootPin, RootPin->GetFName()));
diff --git a/Engine/Source/Developer/RigVMDeveloper/Private/RigVMModel/RigVMPin.cpp b/Engine/Source/Developer/RigVMDeveloper/Private/RigVMModel/RigVMPin.cpp
--- a/Engine/Source/Developer/RigVMDeveloper/Private/RigVMModel/RigVMPin.cpp
+++ b/Engine/Source/Developer/RigVMDeveloper/Private/RigVMModel/RigVMPin.cpp
@@ -571,6 +571,18 @@ void URigVMPin::SetNameFromIndex()
 	LowLevelRename(*FString::FormatAsNumber(GetPinIndex()));
 }
 
+void URigVMPin::SetDisplayName(const FName& InDisplayName)
+{
+	if(InDisplayName == GetFName())
+	{
+		DisplayName = NAME_None;
+	}
+	else
+	{
+		DisplayName = InDisplayName;
+	}
+}
+
 int32 URigVMPin::GetArraySize() const
 {
 	return SubPins.Num();
diff --git a/Engine/Source/Developer/RigVMDeveloper/Public/RigVMModel/Nodes/RigVMAggregateNode.h b/Engine/Source/Developer/RigVMDeveloper/Public/RigVMModel/Nodes/RigVMAggregateNode.h
--- a/Engine/Source/Developer/RigVMDeveloper/Public/RigVMModel/Nodes/RigVMAggregateNode.h
+++ b/Engine/Source/Developer/RigVMDeveloper/Public/RigVMModel/Nodes/RigVMAggregateNode.h
@@ -25,20 +25,24 @@ public:
 
 	// URigVMUnitNode interface
 	virtual FString GetNodeTitle() const override;
+	virtual FLinearColor GetNodeColor() const override;
 	virtual FName GetMethodName() const;
 	virtual  FText GetToolTipText() const override;
 	virtual FText GetToolTipTextForPin(const URigVMPin* InPin) const override;
 
 #if UE_RIGVM_AGGREGATE_NODES_ENABLED
-	virtual bool IsAggregate() const override { return true; };
+	virtual bool IsAggregate() const override { return false; };
 #endif
 	virtual bool IsInputAggregate() const override;
 
 	URigVMNode* GetFirstInnerNode() const;
 	URigVMNode* GetLastInnerNode() const;
-	URigVMPin* GetFirstAggregatePin() const;
-	URigVMPin* GetSecondAggregatePin() const;
-	URigVMPin* GetOppositeAggregatePin() const;
+	
+	virtual URigVMPin* GetFirstAggregatePin() const override;
+	virtual URigVMPin* GetSecondAggregatePin() const override;
+	virtual URigVMPin* GetOppositeAggregatePin() const override;
+	virtual TArray<URigVMPin*> GetAggregateInputs() const override;
+	virtual TArray<URigVMPin*> GetAggregateOutputs() const override;
 
 private:
 
diff --git a/Engine/Source/Developer/RigVMDeveloper/Public/RigVMModel/Nodes/RigVMUnitNode.h b/Engine/Source/Developer/RigVMDeveloper/Public/RigVMModel/Nodes/RigVMUnitNode.h
--- a/Engine/Source/Developer/RigVMDeveloper/Public/RigVMModel/Nodes/RigVMUnitNode.h
+++ b/Engine/Source/Developer/RigVMDeveloper/Public/RigVMModel/Nodes/RigVMUnitNode.h
@@ -36,10 +36,9 @@ public:
 	virtual URigVMPin* GetSecondAggregatePin() const override;
 	virtual URigVMPin* GetOppositeAggregatePin() const override;
 	virtual bool IsInputAggregate() const override;
-	
-	TArray<URigVMPin*> GetAggregateInputs() const;
-	TArray<URigVMPin*> GetAggregateOutputs() const;
-	FName GetNextAggregateName(const FName& InLastAggregatePinName)const;
+	virtual TArray<URigVMPin*> GetAggregateInputs() const override;
+	virtual TArray<URigVMPin*> GetAggregateOutputs() const override;
+	virtual FName GetNextAggregateName(const FName& InLastAggregatePinName) const override;
 
 	bool IsDeprecated() const;
 	FString GetDeprecatedMetadata() const;
diff --git a/Engine/Source/Developer/RigVMDeveloper/Public/RigVMModel/RigVMController.h b/Engine/Source/Developer/RigVMDeveloper/Public/RigVMModel/RigVMController.h
--- a/Engine/Source/Developer/RigVMDeveloper/Public/RigVMModel/RigVMController.h
+++ b/Engine/Source/Developer/RigVMDeveloper/Public/RigVMModel/RigVMController.h
@@ -979,7 +979,7 @@ private:
 	void AddPinsForStruct(UStruct* InStruct, URigVMNode* InNode, URigVMPin* InParentPin, ERigVMPinDirection InPinDirection, const FString& InDefaultValue, bool bAutoExpandArrays, bool bNotify = false);
 	void AddPinsForArray(FArrayProperty* InArrayProperty, URigVMNode* InNode, URigVMPin* InParentPin, ERigVMPinDirection InPinDirection, const TArray<FString>& InDefaultValues, bool bAutoExpandArrays);
 	void ConfigurePinFromProperty(FProperty* InProperty, URigVMPin* InOutPin, ERigVMPinDirection InPinDirection = ERigVMPinDirection::Invalid);
-	void ConfigurePinFromPin(URigVMPin* InOutPin, URigVMPin* InPin);
+	void ConfigurePinFromPin(URigVMPin* InOutPin, URigVMPin* InPin, bool bCopyDisplayName = false);
 	virtual bool ShouldStructBeUnfolded(const UStruct* InStruct);
 	virtual bool ShouldPinBeUnfolded(URigVMPin* InPin);
 	bool SetPinDefaultValue(URigVMPin* InPin, const FString& InDefaultValue, bool bResizeArrays, bool bSetupUndoRedo, bool bMergeUndoAction, bool bNotify = true);
diff --git a/Engine/Source/Developer/RigVMDeveloper/Public/RigVMModel/RigVMNode.h b/Engine/Source/Developer/RigVMDeveloper/Public/RigVMModel/RigVMNode.h
--- a/Engine/Source/Developer/RigVMDeveloper/Public/RigVMModel/RigVMNode.h
+++ b/Engine/Source/Developer/RigVMDeveloper/Public/RigVMModel/RigVMNode.h
@@ -239,7 +239,16 @@ public:
 
 	UFUNCTION(BlueprintPure, Category = RigVMNode)
 	virtual bool IsInputAggregate() const { return false; }
-	
+
+	UFUNCTION(BlueprintPure, Category = RigVMNode)
+	virtual TArray<URigVMPin*> GetAggregateInputs() const { return {}; }
+
+	UFUNCTION(BlueprintPure, Category = RigVMNode)
+	virtual TArray<URigVMPin*> GetAggregateOutputs() const { return {}; }
+
+	UFUNCTION(BlueprintPure, Category = RigVMNode)
+	virtual FName GetNextAggregateName(const FName& InLastAggregatePinName) const { return NAME_None; }
+
 private:
 
 	static const FString NodeColorName;
diff --git a/Engine/Source/Developer/RigVMDeveloper/Public/RigVMModel/RigVMPin.h b/Engine/Source/Developer/RigVMDeveloper/Public/RigVMModel/RigVMPin.h
--- a/Engine/Source/Developer/RigVMDeveloper/Public/RigVMModel/RigVMPin.h
+++ b/Engine/Source/Developer/RigVMDeveloper/Public/RigVMModel/RigVMPin.h
@@ -426,6 +426,7 @@ private:
 
 	void UpdateTypeInformationIfRequired() const;
 	void SetNameFromIndex();
+	void SetDisplayName(const FName& InDisplayName);
 
 	void GetExposedPinChainImpl(TArray<const URigVMPin*>& OutExposedPins, TArray<const URigVMPin*>& VisitedPins) const;
 
diff --git a/Engine/Source/Runtime/RigVM/Private/RigVMCore/RigVMStruct.cpp b/Engine/Source/Runtime/RigVM/Private/RigVMCore/RigVMStruct.cpp
--- a/Engine/Source/Runtime/RigVM/Private/RigVMCore/RigVMStruct.cpp
+++ b/Engine/Source/Runtime/RigVM/Private/RigVMCore/RigVMStruct.cpp
@@ -171,6 +171,51 @@ FString FRigVMStructUpgradeInfo::RemapPin(const FString& InPinPath, bool bIsInpu
 	return PinPath;
 }
 
+FString FRigVMStructUpgradeInfo::AddAggregatePin(FString InPinName)
+{
+	if(InPinName.IsEmpty())
+	{
+		FString LastPinName;
+		
+		FStructOnScope StructOnScope(GetNewStruct());
+		const FRigVMStruct* StructMemory = (const FRigVMStruct*)StructOnScope.GetStructMemory();
+
+		if(AggregatePins.IsEmpty())
+		{
+#if WITH_EDITOR
+			FString LastInputPinName;
+			FString LastOutputPinName;
+			for (TFieldIterator<FProperty> It(GetNewStruct()); It; ++It)
+			{
+				if(It->HasMetaData(FRigVMStruct::AggregateMetaName))
+				{
+					FString& LastDeterminedPinName = It->HasMetaData(FRigVMStruct::InputMetaName) ?
+						LastInputPinName : LastOutputPinName;
+
+					if(!LastDeterminedPinName.IsEmpty())
+					{
+						LastPinName = It->GetName();
+						break;
+					}
+
+					LastDeterminedPinName = It->GetName();
+				}
+			}
+#else
+			LastPinName = TEXT("B");
+#endif
+		}
+		else
+		{
+			LastPinName = AggregatePins.Last();
+		}
+
+		InPinName = StructMemory->GetNextAggregateName(*LastPinName).ToString();
+	}
+	AggregatePins.Add(InPinName);
+	return AggregatePins.Last();
+}
+
 void FRigVMStructUpgradeInfo::SetDefaultValues(const FRigVMStruct* InNewStructMemory)
 {
 	check(NewStruct);
@@ -581,6 +626,41 @@ FString FRigVMStruct::ExportToFullyQualifiedText(const UScriptStruct* InScriptSt
 	return ExportToFullyQualifiedText(Property, StructMemberMemoryPtr);
 }
 
+FName FRigVMStruct::GetNextAggregateName(const FName& InLastAggregatePinName) const
+{
+	if(InLastAggregatePinName.IsNone())
+	{
+		return InLastAggregatePinName;
+	}
+
+	const FString PinName = InLastAggregatePinName.ToString();
+	if(PinName.IsNumeric())
+	{
+		const int32 Index = FCString::Atoi(*PinName);
+		return *FString::FormatAsNumber(Index+1);
+	}
+
+	if(PinName.Len() == 1)
+	{
+		const TCHAR C = PinName[0];
+		if((C >='a' && C < 'z') || (C >='A' && C <'Z'))
+		{
+			FString Result;
+			Result.AppendChar(C + 1);
+			return *Result;
+		}
+	}
+
+	if(PinName.Contains(TEXT("_")))
+	{
+		const FString Left = PinName.Left(PinName.Find(TEXT("_"), ESearchCase::IgnoreCase, ESearchDir::FromEnd));
+		const FString Right = PinName.Mid(PinName.Find(TEXT("_"), ESearchCase::IgnoreCase, ESearchDir::FromEnd)+1);
+		return *FString::Printf(TEXT("%s_%s"), *Left, *GetNextAggregateName(*Right).ToString());
+	}
+
+	return *FString::Printf(TEXT("%s_1"), *PinName);
+}
+
 TMap<FName, FString> FRigVMStruct::GetDefaultValues(UScriptStruct* InScriptStruct) const
 {
 	check(InScriptStruct);
diff --git a/Engine/Source/Runtime/RigVM/Public/RigVMCore/RigVMStruct.h b/Engine/Source/Runtime/RigVM/Public/RigVMCore/RigVMStruct.h
--- a/Engine/Source/Runtime/RigVM/Public/RigVMCore/RigVMStruct.h
+++ b/Engine/Source/Runtime/RigVM/Public/RigVMCore/RigVMStruct.h
@@ -172,6 +172,12 @@ public:
 	// remaps a pin path based on our internals
 	FString RemapPin(const FString& InPinPath, bool bIsInput, bool bContainsNodeName) const;
 
+	// adds a new aggregate pin
+	FString AddAggregatePin(FString InPinName = FString());
+
+	// returns the aggregate pins to add
+	const TArray<FString>& GetAggregatePins() const { return AggregatePins; }
+
 private:
 	
 	// The complete node path including models / collapse node.
@@ -195,6 +201,9 @@ private:
 	// New sets of default values
 	TMap<FName, FString> DefaultValues;
 
+	// Aggregate pins to add
+	TArray<FString> AggregatePins;
+
 	// sets the default values from the new struct.
 	void SetDefaultValues(const FRigVMStruct* InNewStructMemory);
 
@@ -257,7 +266,7 @@ public:
 
 	FString ExportToFullyQualifiedText(const UScriptStruct* InScriptStruct, const FName& InPropertyName, const uint8* InStructMemoryPointer = nullptr) const;
 	
-	virtual FName GetNextAggregateName(const FName& InLastAggregatePinName) const { return FName(); }
+	virtual FName GetNextAggregateName(const FName& InLastAggregatePinName) const;
 	virtual FRigVMStructUpgradeInfo GetUpgradeInfo() const { return FRigVMStructUpgradeInfo(); }
 
 	static const FName DeprecatedMetaName;
diff --git a/Engine/Source/Runtime/RigVM/Public/RigVMDefines.h b/Engine/Source/Runtime/RigVM/Public/RigVMDefines.h
--- a/Engine/Source/Runtime/RigVM/Public/RigVMDefines.h
+++ b/Engine/Source/Runtime/RigVM/Public/RigVMDefines.h
@@ -19,5 +19,5 @@
 #endif
 
 #ifndef UE_RIGVM_AGGREGATE_NODES_ENABLED
-#define UE_RIGVM_AGGREGATE_NODES_ENABLED 0
+#define UE_RIGVM_AGGREGATE_NODES_ENABLED 1
 #endif
\ No newline at end of file
