Control Rig: Fixed Use Selected Item not working with space in object name

diff --git a/Engine/Source/Developer/RigVMDeveloper/Private/RigVMModel/RigVMPin.cpp b/Engine/Source/Developer/RigVMDeveloper/Private/RigVMModel/RigVMPin.cpp
--- a/Engine/Source/Developer/RigVMDeveloper/Private/RigVMModel/RigVMPin.cpp
+++ b/Engine/Source/Developer/RigVMDeveloper/Private/RigVMModel/RigVMPin.cpp
@@ -157,8 +157,8 @@ TArray<FString> URigVMPin::SplitDefaultValue(const FString& InDefaultValue)
 
 	if (!Content.IsEmpty())
 	{
-		// ignore whitespaces
-		Parts.Add(Content.Mid(LastPartStartIndex).Replace(TEXT(" "), TEXT("")));
+		// ignore whitespaces from the start and end of the string
+		Parts.Add(Content.Mid(LastPartStartIndex).TrimStartAndEnd());
 	}
 	return Parts;
 }
