RigVM: Maintain links property for aggregate nodes.

diff --git a/Engine/Source/Developer/RigVMDeveloper/Private/RigVMModel/RigVMController.cpp b/Engine/Source/Developer/RigVMDeveloper/Private/RigVMModel/RigVMController.cpp
--- a/Engine/Source/Developer/RigVMDeveloper/Private/RigVMModel/RigVMController.cpp
+++ b/Engine/Source/Developer/RigVMDeveloper/Private/RigVMModel/RigVMController.cpp
@@ -1827,32 +1827,11 @@ bool URigVMController::UnresolveTemplateNodes(const TArray<URigVMTemplateNode*>&
 		ResolveTemplateNodeMetaData(Node, bSetupUndoRedo);
 	}
 
-	for(const TPair<FString, FString>& LinkedPath : LinkedPaths)
-	{
-		const FString& SourcePath = LinkedPath.Key;
-		const FString& TargetPath = LinkedPath.Value;
-
-		URigVMPin* SourcePin = GetGraph()->FindPin(SourcePath);
-		URigVMPin* TargetPin = GetGraph()->FindPin(TargetPath);
-
-		if(SourcePin == nullptr || TargetPin == nullptr)
-		{
-			ReportRemovedLink(SourcePath, TargetPath);
-			continue;
-		}
-
-		if(SourcePin->IsWildCard() != TargetPin->IsWildCard())
+	RestoreLinkedPaths(LinkedPaths, {}, {},
+		FRigVMController_CheckPinComatibilityDelegate::CreateLambda([](URigVMPin* A, URigVMPin* B) -> bool
 		{
-			ReportRemovedLink(SourcePath, TargetPath);
-			continue;
-		}
-
-		// it's ok if this fails - we want to maintain the minimum set of links
-		if(!AddLink(SourcePin, TargetPin, bSetupUndoRedo))
-		{
-			ReportRemovedLink(SourcePath, TargetPath);
-		}
-	}
+			return A->IsWildCard() == B->IsWildCard();
+		}), bSetupUndoRedo);
 	
 	if(bSetupUndoRedo)
 	{
@@ -1963,17 +1942,13 @@ TArray<URigVMNode*> URigVMController::UpgradeNodes(const TArray<URigVMNode*>& In
 
 	// find all links affecting the nodes to upgrade
 	TArray<TPair<FString, FString>> LinkedPaths = GetLinkedPinPaths(InNodes);
-	for(const TPair<FString, FString>& LinkedPath : LinkedPaths)
+	if(!BreakLinkedPaths(LinkedPaths, bSetupUndoRedo))
 	{
-		if(!BreakLink(LinkedPath.Key, LinkedPath.Value, bSetupUndoRedo))
+		if(bSetupUndoRedo)
 		{
-			if(bSetupUndoRedo)
-			{
-				ActionStack->CancelAction(Action, this);
-			}
-			ReportErrorf(TEXT("Couldn't remove link '%s' -> '%s'"), *LinkedPath.Key, *LinkedPath.Value);
-			return TArray<URigVMNode*>();
+			ActionStack->CancelAction(Action, this);
 		}
+		return TArray<URigVMNode*>();
 	}
 
 	TArray<URigVMNode*> UpgradedNodes;
@@ -1989,37 +1964,7 @@ TArray<URigVMNode*> URigVMController::UpgradeNodes(const TArray<URigVMNode*>& In
 		}
 	}
 
-	for(const TPair<FString, FString>& LinkedPath : LinkedPaths)
-	{
-		FString SourcePath = LinkedPath.Key;
-		FString TargetPath = LinkedPath.Value;
-
-		FString SourceNodeName, SourceSegmentPath;
-		if(!URigVMPin::SplitPinPathAtStart(SourcePath, SourceNodeName, SourceSegmentPath))
-		{
-			continue;
-		}
-		
-		FString TargetNodeName, TargetSegmentPath;
-		if(!URigVMPin::SplitPinPathAtStart(TargetPath, TargetNodeName, TargetSegmentPath))
-		{
-			continue;
-		}
-
-		if(FRigVMController_PinPathRemapDelegate* SourceRemapDelegate = RemapPinDelegates.Find(SourceNodeName))
-		{
-			SourcePath = SourceRemapDelegate->Execute(SourcePath, false);
-		}
-		if(FRigVMController_PinPathRemapDelegate* TargetRemapDelegate = RemapPinDelegates.Find(TargetNodeName))
-		{
-			TargetPath = TargetRemapDelegate->Execute(TargetPath, true);
-		}
-		
-		if(!AddLink(SourcePath, TargetPath, bSetupUndoRedo))
-		{
-			ReportRemovedLink(SourcePath, TargetPath);
-		}
-	}
+	RestoreLinkedPaths(LinkedPaths, {}, RemapPinDelegates, bSetupUndoRedo);
 
 	if(bRecursive)
 	{
@@ -5389,7 +5334,7 @@ URigVMFunctionReferenceNode* URigVMController::PromoteCollapseNodeToFunctionRefe
 	{
 		if (Pin->IsWildCard())
 		{
-			ReportErrorf(TEXT("Cannot create function %s because it contains a wildcard pin %s"), *InCollapseNode->GetName(), *Pin->GetName());
+			ReportAndNotifyErrorf(TEXT("Cannot create function %s because it contains a wildcard pin %s"), *InCollapseNode->GetName(), *Pin->GetName());
 			return nullptr;
 		}
 	}
@@ -7661,6 +7606,16 @@ FString URigVMController::AddAggregatePin(URigVMNode* InNode, const FString& InP
 		const FString AggregateArg2 = Arg2->GetName();
 		const FString AggregateOppositeArg = ArgOpposite->GetName();
 
+		TArray<TPair<FString, FString>> LinkedPaths = GetLinkedPinPaths(InNode);
+		if(!BreakLinkedPaths(LinkedPaths, bSetupUndoRedo))
+		{
+			if(bSetupUndoRedo)
+			{
+				ActionStack->CancelAction(Action, this);
+			}
+			return FString();
+		}
+
 		const FName PreviousNodeName = InNode->GetFName();
 		URigVMCollapseNode* CollapseNode = CollapseNodes({InNode}, InNode->GetName(), bSetupUndoRedo, true);
 		if (!CollapseNode)
@@ -7729,6 +7684,8 @@ FString URigVMController::AddAggregatePin(URigVMNode* InNode, const FString& InP
 			}
 			return FString();
 		}
+
+		RestoreLinkedPaths(LinkedPaths, {{PreviousNodeName.ToString(), AggregateNode->GetName()}}, {}, bSetupUndoRedo);
 	}
 	
 	if (!AggregateNode)
@@ -8008,8 +7965,66 @@ bool URigVMController::RemoveAggregatePin(URigVMPin* InPin, bool bSetupUndoRedo,
 
 		if (bSuccess && AggregateNode->GetContainedNodes().Num() == 3)
 		{
+			TArray<TPair<FString, FString>> LinkedPaths = GetLinkedPinPaths(AggregateNode);
+			if(!BreakLinkedPaths(LinkedPaths, bSetupUndoRedo))
+			{
+				if(bSetupUndoRedo)
+				{
+					ActionStack->CancelAction(Action, this);
+				}
+				return false;
+			}
+
+			TMap<FString, FString> PinNameMap;
+			for(URigVMPin* Pin : AggregateNode->GetPins())
+			{
+				if(URigVMPin* EntryPin = AggregateNode->GetEntryNode()->FindPin(Pin->GetName()))
+				{
+					TArray<URigVMPin*> TargetPins = EntryPin->GetLinkedTargetPins();
+					if(TargetPins.Num() > 0)
+					{
+						PinNameMap.Add(EntryPin->GetName(), TargetPins[0]->GetName());
+					}
+				}
+				else if(URigVMPin* ReturnPin = AggregateNode->GetReturnNode()->FindPin(Pin->GetName()))
+				{
+					TArray<URigVMPin*> SourcePins = ReturnPin->GetLinkedSourcePins();
+					if(SourcePins.Num() > 0)
+					{
+						PinNameMap.Add(ReturnPin->GetName(), SourcePins[0]->GetName());
+					}
+				}
+			}
+
+			const FString PreviousNodeName = AggregateNode->GetName();
 			TArray<URigVMNode*> NodesEjected = ExpandLibraryNode(AggregateNode, bSetupUndoRedo);
 			bSuccess = NodesEjected.Num() == 1;
+
+			if(bSuccess)
+			{
+				URigVMNode* EjectedNode = NodesEjected[0];
+				RestoreLinkedPaths(LinkedPaths, {}, {{
+					PreviousNodeName, FRigVMController_PinPathRemapDelegate::CreateLambda([
+						PreviousNodeName,
+						EjectedNode,
+						PinNameMap
+					](const FString& InPinPath, bool bIsInput) -> FString
+					{
+						static constexpr TCHAR PinPrefixFormat[] = TEXT("%s.");
+
+						TArray<FString> Segments;
+						URigVMPin::SplitPinPath(InPinPath, Segments);
+						Segments[0] = EjectedNode->GetName();
+
+						if(const FString* RemappedPin = PinNameMap.Find(Segments[1]))
+						{
+							Segments[1] = *RemappedPin;
+						}
+						
+						return URigVMPin::JoinPinPath(Segments);
+					})
+				}}, bSetupUndoRedo);
+			}
 		}		
 	}
 
@@ -15852,3 +15867,84 @@ TArray<TPair<FString, FString>> URigVMController::GetLinkedPinPaths(const TArray
 	}
 	return LinkedPaths;
 }
+
+bool URigVMController::BreakLinkedPaths(const TArray<TPair<FString, FString>>& InLinkedPaths, bool bSetupUndoRedo)
+{
+	for(const TPair<FString, FString>& LinkedPath : InLinkedPaths)
+	{
+		if(!BreakLink(LinkedPath.Key, LinkedPath.Value, bSetupUndoRedo))
+		{
+			ReportErrorf(TEXT("Couldn't remove link '%s' -> '%s'"), *LinkedPath.Key, *LinkedPath.Value);
+			return false;
+		}
+	}
+
+	return true;
+}
+
+bool URigVMController::RestoreLinkedPaths(
+	const TArray<TPair<FString, FString>>& InLinkedPaths,
+	const TMap<FString, FString>& InNodeNameMap,
+	const TMap<FString,FRigVMController_PinPathRemapDelegate>& InRemapDelegates,
+	FRigVMController_CheckPinComatibilityDelegate InCompatibilityDelegate,
+	bool bSetupUndoRedo,
+	ERigVMPinDirection InUserDirection)
+{
+	bool bSuccess = true;
+
+	auto RemapNodeName = [InNodeNameMap, InRemapDelegates](const FString& InPinPath, bool bAsInput) -> FString
+	{
+		FString NodeName, SegmentPath;
+		if(!URigVMPin::SplitPinPathAtStart(InPinPath, NodeName, SegmentPath))
+		{
+			return InPinPath;
+		}
+
+		FString PinPath = InPinPath;
+
+		if(const FRigVMController_PinPathRemapDelegate* RemapDelegate = InRemapDelegates.Find(NodeName))
+		{
+			PinPath = RemapDelegate->Execute(PinPath, bAsInput);
+		}
+		else if(const FString* RemappedNodeName = InNodeNameMap.Find(NodeName))
+		{
+			PinPath = URigVMPin::JoinPinPath(*RemappedNodeName, SegmentPath);
+		}
+
+		return PinPath;
+	};
+	
+	for(const TPair<FString, FString>& LinkedPath : InLinkedPaths)
+	{
+		const FString SourcePath = RemapNodeName(LinkedPath.Key, false);
+		const FString TargetPath = RemapNodeName(LinkedPath.Value, true);
+
+		URigVMPin* SourcePin = GetGraph()->FindPin(SourcePath);
+		URigVMPin* TargetPin = GetGraph()->FindPin(TargetPath);
+
+		if(SourcePin == nullptr || TargetPin == nullptr)
+		{
+			ReportRemovedLink(SourcePath, TargetPath);
+			bSuccess = false;
+			continue;
+		}
+
+		if(InCompatibilityDelegate.IsBound())
+		{
+			if(!InCompatibilityDelegate.Execute(SourcePin, TargetPin))
+			{
+				bSuccess = false;
+				continue;
+			}
+		}
+
+		// it's ok if this fails - we want to maintain the minimum set of links
+		if(!AddLink(SourcePin, TargetPin, bSetupUndoRedo, InUserDirection))
+		{
+			ReportRemovedLink(SourcePath, TargetPath);
+			bSuccess = false;
+		}
+	}
+	
+	return bSuccess;
+}
diff --git a/Engine/Source/Developer/RigVMDeveloper/Public/RigVMModel/RigVMController.h b/Engine/Source/Developer/RigVMDeveloper/Public/RigVMModel/RigVMController.h
--- a/Engine/Source/Developer/RigVMDeveloper/Public/RigVMModel/RigVMController.h
+++ b/Engine/Source/Developer/RigVMDeveloper/Public/RigVMModel/RigVMController.h
@@ -90,6 +90,7 @@ DECLARE_DELEGATE_FiveParams(FRigVMController_OnBulkEditProgressDelegate, TSoftOb
 DECLARE_DELEGATE_RetVal_TwoParams(FString, FRigVMController_PinPathRemapDelegate, const FString& /* InPinPath */, bool /* bIsInput */);
 DECLARE_DELEGATE_OneParam(FRigVMController_RequestJumpToHyperlinkDelegate, const UObject* InSubject);
 DECLARE_DELEGATE_OneParam(FRigVMController_ConfigureWorkflowOptionsDelegate, URigVMUserWorkflowOptions* InOutOptions);
+DECLARE_DELEGATE_RetVal_TwoParams(bool, FRigVMController_CheckPinComatibilityDelegate, URigVMPin*, URigVMPin*);
 
 UINTERFACE()
 class RIGVMDEVELOPER_API URigVMControllerHost : public UInterface
@@ -1090,6 +1091,23 @@ public:
 	static void SanitizeName(FString& InOutName, bool bAllowPeriod, bool bAllowSpace);
 	static TArray<TPair<FString, FString>> GetLinkedPinPaths(URigVMNode* InNode, bool bIncludeInjectionNodes = false);
 	static TArray<TPair<FString, FString>> GetLinkedPinPaths(const TArray<URigVMNode*>& InNodes, bool bIncludeInjectionNodes = false);
+	bool BreakLinkedPaths(const TArray<TPair<FString, FString>>& InLinkedPaths, bool bSetupUndoRedo);
+	bool RestoreLinkedPaths(
+		const TArray<TPair<FString, FString>>& InLinkedPaths,
+		const TMap<FString, FString>& InNodeNameMap,
+		const TMap<FString,FRigVMController_PinPathRemapDelegate>& InRemapDelegates,
+		FRigVMController_CheckPinComatibilityDelegate InCompatibilityDelegate,
+		bool bSetupUndoRedo,
+		ERigVMPinDirection InUserDirection = ERigVMPinDirection::Invalid);
+	bool RestoreLinkedPaths(
+		const TArray<TPair<FString, FString>>& InLinkedPaths,
+		const TMap<FString, FString>& InNodeNameMap,
+		const TMap<FString,FRigVMController_PinPathRemapDelegate>& InRemapDelegates,
+		bool bSetupUndoRedo,
+		ERigVMPinDirection InUserDirection = ERigVMPinDirection::Invalid)
+	{
+		return RestoreLinkedPaths(InLinkedPaths, InNodeNameMap, InRemapDelegates, FRigVMController_CheckPinComatibilityDelegate(), bSetupUndoRedo, InUserDirection);
+	}
 
 private: 
 	UPROPERTY(transient)
