Control Rig: Improve collapse functionality for execute pins

diff --git a/Engine/Source/Developer/RigVMDeveloper/Private/RigVMModel/RigVMController.cpp b/Engine/Source/Developer/RigVMDeveloper/Private/RigVMModel/RigVMController.cpp
--- a/Engine/Source/Developer/RigVMDeveloper/Private/RigVMModel/RigVMController.cpp
+++ b/Engine/Source/Developer/RigVMDeveloper/Private/RigVMModel/RigVMController.cpp
@@ -4046,114 +4046,21 @@ URigVMCollapseNode* URigVMController::CollapseNodes(const TArray<URigVMNode*>& I
 		LinksToRewire.Add(Link);
 	}
 
-	// make sure that for execute pins we are on one branch only
-	TArray<URigVMPin*> InputExecutePins;
-	TArray<URigVMPin*> IntermediateExecutePins;
-	TArray<URigVMPin*> OutputExecutePins;
-
-	// first collect the output execute pins
-	for (URigVMLink* Link : LinksToRewire)
+	// sort the links so that the links on the same node are in the right order
+	Algo::Sort(LinksToRewire, [&AllLinks](URigVMLink* A, URigVMLink* B) -> bool
 	{
-		URigVMPin* ExecutePin = Link->GetSourcePin();
-		if (!ExecutePin->IsExecuteContext())
-		{
-			continue;
-		}
-		if (!NodeToBeCollapsed(ExecutePin->GetNode()))
-		{
-			continue;
-		}
-		if (!OutputExecutePins.IsEmpty())
+		if(A->GetSourcePin()->GetNode() == B->GetSourcePin()->GetNode())
 		{
-			if (bSetupUndoRedo)
-			{
-				ReportAndNotifyErrorf(
-					TEXT("Only one set of execute branches can be collapsed, pin %s and %s are on separate branches"),
-					*OutputExecutePins[0]->GetPinPath(),
-					*ExecutePin->GetPinPath()
-				);
-			}
-			return nullptr;
+			return A->GetSourcePin()->GetAbsolutePinIndex() < B->GetSourcePin()->GetAbsolutePinIndex();
 		}
-		OutputExecutePins.Add(ExecutePin);
-
-		while (ExecutePin)
-		{
-			if (IntermediateExecutePins.Contains(ExecutePin))
-			{
-				if (bSetupUndoRedo)
-				{
-					ReportAndNotifyErrorf(TEXT("Only one set of execute branches can be collapsed."));
-				}
-				return nullptr;
-			}
-			IntermediateExecutePins.Add(ExecutePin);
-
-			// walk backwards and find all "known execute pins"
-			URigVMNode* ExecuteNode = ExecutePin->GetNode();
-			for (URigVMPin* Pin : ExecuteNode->GetPins())
-			{
-				if (Pin->GetDirection() != ERigVMPinDirection::Input &&
-					Pin->GetDirection() != ERigVMPinDirection::IO)
-				{
-					continue;
-				}
-				if (!Pin->IsExecuteContext())
-				{
-					continue;
-				}
-				TArray<URigVMLink*> SourceLinks = Pin->GetSourceLinks();
-				ExecutePin = nullptr;
-				if (SourceLinks.Num() > 0)
-				{
-					URigVMPin* PreviousExecutePin = SourceLinks[0]->GetSourcePin();
-					if (NodeToBeCollapsed(PreviousExecutePin->GetNode()))
-					{
-						if (Pin != IntermediateExecutePins.Last())
-						{
-							IntermediateExecutePins.Add(Pin);
-						}
-						ExecutePin = PreviousExecutePin;
-						break;
-					}
-				}
-			}
-		}
-	}
-	for (URigVMLink* Link : LinksToRewire)
-	{
-		URigVMPin* ExecutePin = Link->GetTargetPin();
-		if (!ExecutePin->IsExecuteContext())
+		
+		if(A->GetTargetPin()->GetNode() == B->GetTargetPin()->GetNode())
 		{
-			continue;
-		}
-		if (!NodeToBeCollapsed(ExecutePin->GetNode()))
-		{
-			continue;
-		}
-		if (!IntermediateExecutePins.Contains(ExecutePin) && !IntermediateExecutePins.IsEmpty())
-		{
-			if (bSetupUndoRedo)
-			{
-				ReportAndNotifyErrorf(TEXT("Only one set of execute branches can be collapsed"));
-			}
-			return nullptr;
+			return A->GetTargetPin()->GetAbsolutePinIndex() < B->GetTargetPin()->GetAbsolutePinIndex();
 		}
 
-		if (!InputExecutePins.IsEmpty())
-		{
-			if (bSetupUndoRedo)
-			{
-				ReportAndNotifyErrorf(
-					TEXT("Only one set of execute branches can be collapsed, pin %s and %s are on separate branches"),
-					*InputExecutePins[0]->GetPinPath(),
-					*ExecutePin->GetPinPath()
-				);
-			}
-			return nullptr;
-		}
-		InputExecutePins.Add(ExecutePin);
-	}
+		return AllLinks.Find(A) < AllLinks.Find(B);
+	});
 
 	FRigVMControllerCompileBracketScope CompileScope(this);
 	FRigVMCollapseNodesAction CollapseAction;
@@ -4197,7 +4104,7 @@ URigVMCollapseNode* URigVMController::CollapseNodes(const TArray<URigVMNode*>& I
 	for (URigVMLink* Link : LinksToRewire)
 	{
 		bool bSourceToBeCollapsed = NodeToBeCollapsed(Link->GetSourcePin()->GetNode());
-		bool bTargetToBeCollapsed = NodeToBeCollapsed(Link->GetTargetPin()->GetNode());
+		bContainsOutputs = bContainsOutputs || bSourceToBeCollapsed;
 
 		URigVMPin* PinToCollapse = bSourceToBeCollapsed ? Link->GetSourcePin() : Link->GetTargetPin();
 		if (CollapsedPins.Contains(PinToCollapse))
@@ -4205,20 +4112,17 @@ URigVMCollapseNode* URigVMController::CollapseNodes(const TArray<URigVMNode*>& I
 			continue;
 		}
 
-		if (PinToCollapse->IsExecuteContext())
+		if(PinToCollapse->IsExecuteContext() && PinToCollapse->GetDirection() == ERigVMPinDirection::IO)
 		{
-			bool bFoundExistingPin = false;
-			for (URigVMPin* ExistingPin : CollapseNode->Pins)
+			for(const TPair<URigVMPin*, URigVMPin*>& Pair : CollapsedPins)
 			{
-				if (ExistingPin->IsExecuteContext())
+				if(Pair.Key->IsExecuteContext() && Pair.Key->GetDirection() == ERigVMPinDirection::IO)
 				{
-					CollapsedPins.Add(PinToCollapse, ExistingPin);
-					bFoundExistingPin = true;
+					CollapsedPins.Add(PinToCollapse, Pair.Value);
 					break;
 				}
 			}
-
-			if (bFoundExistingPin)
+			if (CollapsedPins.Contains(PinToCollapse))
 			{
 				continue;
 			}
@@ -4253,14 +4157,16 @@ URigVMCollapseNode* URigVMController::CollapseNodes(const TArray<URigVMNode*>& I
 		URigVMPin* CollapsedPin = NewObject<URigVMPin>(CollapseNode, PinName);
 		ConfigurePinFromPin(CollapsedPin, PinToCollapse, true);
 
-		if (CollapsedPin->IsExecuteContext())
-		{
-			CollapsedPin->Direction = ERigVMPinDirection::IO;
-			bContainsOutputs = true;
-		}
-		else if (CollapsedPin->GetDirection() == ERigVMPinDirection::IO)
+		if (CollapsedPin->GetDirection() == ERigVMPinDirection::IO)
 		{
-			CollapsedPin->Direction = ERigVMPinDirection::Input;
+			if(CollapsedPin->IsExecuteContext())
+			{
+				bContainsOutputs = true;
+			}
+			else
+			{
+				CollapsedPin->Direction = bSourceToBeCollapsed ? ERigVMPinDirection::Output : ERigVMPinDirection::Input;
+			}
 		}
 
 		if (CollapsedPin->IsStruct())
@@ -4268,8 +4174,6 @@ URigVMCollapseNode* URigVMController::CollapseNodes(const TArray<URigVMNode*>& I
 			AddPinsForStruct(CollapsedPin->GetScriptStruct(), CollapseNode, CollapsedPin, CollapsedPin->GetDirection(), FString(), false);
 		}
 
-		bContainsOutputs = bContainsOutputs || bSourceToBeCollapsed;
-
 		AddNodePin(CollapseNode, CollapsedPin);
 
 		FPinState PinState = GetPinState(PinToCollapse);
@@ -9494,6 +9398,22 @@ FName URigVMController::AddExposedPin(const FName& InPinName, ERigVMPinDirection
 		}
 	}
 
+	/* Eventually we want to disallow this only on function library collapse nodes
+	if (CPPTypeObject)
+	{
+		if(CPPTypeObject == RigVMTypeUtils::GetWildCardCPPTypeObject())
+		{
+			if(const URigVMNode* CollapseNode = Cast<URigVMNode>(GetGraph()->GetOuter()))
+			{
+				if(CollapseNode->GetOuter()->IsA<URigVMFunctionLibrary>())
+				{
+					ReportError(TEXT("Cannot expose pins of wildcard type in functions."));
+					return NAME_None;
+				}
+			}
+		}
+	}
+	*/
 	if (CPPTypeObject)
 	{
 		if(CPPTypeObject == RigVMTypeUtils::GetWildCardCPPTypeObject())
@@ -9502,7 +9422,7 @@ FName URigVMController::AddExposedPin(const FName& InPinName, ERigVMPinDirection
 			return NAME_None;
 		}
 	}
-
+    	
 	// only allow one IO / input exposed pin of type execute context per direction
 	if(InDirection != ERigVMPinDirection::Output)
 	{
@@ -15127,11 +15047,52 @@ bool URigVMController::ResolveWildCardPinImpl(URigVMPin* InPinToResolve, const F
 		return false;
 	}
 
+	if(bTraverseLinks)
+	{
+		if(URigVMNode* Node = InPinToResolve->GetNode())
+		{
+			if(URigVMLibraryNode* LibraryNode = Cast<URigVMLibraryNode>(Node))
+			{
+				const FString& PinName = InPinToResolve->GetRootPin()->GetName();
+
+				if(URigVMNode* EntryNode = LibraryNode->GetEntryNode())
+				{
+					if(URigVMPin* EntryPin = EntryNode->FindPin(PinName))
+					{
+						ResolveWildCardPin(EntryPin, InPinToResolve->GetRootPin(), bSetupUndoRedo, false, bTraverseParentPins, bTraverseLinks);
+					}
+				}
+
+				if(URigVMNode* ReturnNode = LibraryNode->GetReturnNode())
+				{
+					if(URigVMPin* ReturnPin = ReturnNode->FindPin(PinName))
+					{
+						ResolveWildCardPin(ReturnPin, InPinToResolve->GetRootPin(), bSetupUndoRedo, false, bTraverseParentPins, bTraverseLinks);
+					}
+				}
+			}
+			else if(Node->IsA<URigVMFunctionEntryNode>() || Node->IsA<URigVMFunctionReturnNode>())
+			{
+				if(URigVMLibraryNode* OuterNode = Cast<URigVMLibraryNode>(Node->GetGraph()->GetOuter()))
+				{
+					const FString& PinName = InPinToResolve->GetRootPin()->GetName();
+					if(URigVMPin* OuterPin = OuterNode->FindPin(PinName))
+					{
+						ResolveWildCardPin(OuterPin, InPinToResolve->GetRootPin(), bSetupUndoRedo, false, bTraverseParentPins, bTraverseLinks);
+					}
+				}
+			}
+		}
+	}
+
 	if(bTraverseNode)
 	{
 		if(const URigVMNode* Node = InPinToResolve->GetNode())
 		{
-			if(!Node->IsA<URigVMTemplateNode>())
+			if(!Node->IsA<URigVMTemplateNode>() &&
+				!Node->IsA<URigVMLibraryNode>() &&
+				!Node->IsA<URigVMFunctionEntryNode>() &&
+				!Node->IsA<URigVMFunctionReturnNode>())
 			{
 				for(URigVMPin* Pin : Node->GetPins())
 				{
diff --git a/Engine/Source/Developer/RigVMDeveloper/Private/RigVMModel/RigVMPin.cpp b/Engine/Source/Developer/RigVMDeveloper/Private/RigVMModel/RigVMPin.cpp
--- a/Engine/Source/Developer/RigVMDeveloper/Private/RigVMModel/RigVMPin.cpp
+++ b/Engine/Source/Developer/RigVMDeveloper/Private/RigVMModel/RigVMPin.cpp
@@ -566,6 +566,11 @@ int32 URigVMPin::GetPinIndex() const
 	return Index;
 }
 
+int32 URigVMPin::GetAbsolutePinIndex() const
+{
+	return GetNode()->GetAllPinsRecursively().Find((URigVMPin*)this);
+}
+
 void URigVMPin::SetNameFromIndex()
 {
 	LowLevelRename(*FString::FormatAsNumber(GetPinIndex()));
diff --git a/Engine/Source/Developer/RigVMDeveloper/Public/RigVMModel/RigVMPin.h b/Engine/Source/Developer/RigVMDeveloper/Public/RigVMModel/RigVMPin.h
--- a/Engine/Source/Developer/RigVMDeveloper/Public/RigVMModel/RigVMPin.h
+++ b/Engine/Source/Developer/RigVMDeveloper/Public/RigVMModel/RigVMPin.h
@@ -217,6 +217,10 @@ public:
 	UFUNCTION(BlueprintCallable, Category = RigVMPin)
 	int32 GetPinIndex() const;
 
+	// Returns the absolute index of the Pin within the node / parent Pin
+	UFUNCTION(BlueprintCallable, Category = RigVMPin)
+	int32 GetAbsolutePinIndex() const;
+
 	// Returns the number of elements within an array Pin
 	UFUNCTION(BlueprintCallable, Category = RigVMPin)
 	int32 GetArraySize() const;
