Control Rig: Filtered permutations for template nodes

diff --git a/Engine/Plugins/Animation/ControlRig/Source/ControlRigDeveloper/Private/ControlRigBlueprint.cpp b/Engine/Plugins/Animation/ControlRig/Source/ControlRigDeveloper/Private/ControlRigBlueprint.cpp
--- a/Engine/Plugins/Animation/ControlRig/Source/ControlRigDeveloper/Private/ControlRigBlueprint.cpp
+++ b/Engine/Plugins/Animation/ControlRig/Source/ControlRigDeveloper/Private/ControlRigBlueprint.cpp
@@ -912,6 +912,8 @@ void UControlRigBlueprint::RefreshAllModels()
 		{
 			Controller->RemoveUnusedOrphanedPins(ModelNode, false);
 		}
+
+		Controller->RecomputeAllTemplateFilteredTypes(false);
 	}
 }
 
@@ -1584,6 +1586,19 @@ URigVMController* UControlRigBlueprint::GetOrCreateController(URigVMGraph* InGra
 		return FRigVMController_BulkEditResult();
 	});
 
+	Controller->RequestBreakLinksDialogDelegate.BindLambda([WeakThis, WeakController](TArray<URigVMLink*> InLinks) -> bool 
+	{
+		if(WeakThis.IsValid() && WeakController.IsValid())
+		{
+			UControlRigBlueprint* StrongThis = WeakThis.Get();
+			if(StrongThis->OnRequestBreakLinksDialog().IsBound())
+			{
+				return StrongThis->OnRequestBreakLinksDialog().Execute(InLinks);
+			}
+		}
+		return false;
+	});
+
 	Controller->RequestNewExternalVariableDelegate.BindLambda([WeakThis](FRigVMGraphVariableDescription InVariable, bool bInIsPublic, bool bInIsReadOnly) -> FName
 	{
 		if (WeakThis.IsValid())
diff --git a/Engine/Plugins/Animation/ControlRig/Source/ControlRigDeveloper/Public/ControlRigBlueprint.h b/Engine/Plugins/Animation/ControlRig/Source/ControlRigDeveloper/Public/ControlRigBlueprint.h
--- a/Engine/Plugins/Animation/ControlRig/Source/ControlRigDeveloper/Public/ControlRigBlueprint.h
+++ b/Engine/Plugins/Animation/ControlRig/Source/ControlRigDeveloper/Public/ControlRigBlueprint.h
@@ -42,6 +42,7 @@ DECLARE_EVENT_OneParam(UControlRigBlueprint, FOnPostEditChangeChainProperty, FPr
 DECLARE_EVENT_ThreeParams(UControlRigBlueprint, FOnLocalizeFunctionDialogRequested, URigVMLibraryNode*, UControlRigBlueprint*, bool);
 DECLARE_EVENT_ThreeParams(UControlRigBlueprint, FOnReportCompilerMessage, EMessageSeverity::Type, UObject*, const FString&);
 DECLARE_DELEGATE_RetVal_FourParams(FRigVMController_BulkEditResult, FControlRigOnBulkEditDialogRequestedDelegate, UControlRigBlueprint*, URigVMController*, URigVMLibraryNode*, ERigVMControllerBulkEditType);
+DECLARE_DELEGATE_RetVal_OneParam(bool, FControlRigOnBreakLinksDialogRequestedDelegate, TArray<URigVMLink*>);
 DECLARE_EVENT(UControlRigBlueprint, FOnBreakpointAdded);
 DECLARE_EVENT_OneParam(UControlRigBlueprint, FOnRequestInspectObject, const TArray<UObject*>& );
 
@@ -652,6 +653,8 @@ public:
 
 	FControlRigOnBulkEditDialogRequestedDelegate& OnRequestBulkEditDialog() { return RequestBulkEditDialog; }
 
+	FControlRigOnBreakLinksDialogRequestedDelegate& OnRequestBreakLinksDialog() { return RequestBreakLinksDialog; }
+
 	FRigVMController_RequestJumpToHyperlinkDelegate& OnRequestJumpToHyperlink() { return RequestJumpToHyperlink; };
 
 	FOnReportCompilerMessage& OnReportCompilerMessage() { return ReportCompilerMessageEvent; }
@@ -669,6 +672,7 @@ private:
 	FOnLocalizeFunctionDialogRequested RequestLocalizeFunctionDialog;
 	FOnReportCompilerMessage ReportCompilerMessageEvent;
 	FControlRigOnBulkEditDialogRequestedDelegate RequestBulkEditDialog;
+	FControlRigOnBreakLinksDialogRequestedDelegate RequestBreakLinksDialog;
 	FRigVMController_RequestJumpToHyperlinkDelegate RequestJumpToHyperlink;
 
 #endif
diff --git a/Engine/Plugins/Animation/ControlRig/Source/ControlRigEditor/Private/ControlRigEditorModule.cpp b/Engine/Plugins/Animation/ControlRig/Source/ControlRigEditor/Private/ControlRigEditorModule.cpp
--- a/Engine/Plugins/Animation/ControlRig/Source/ControlRigEditor/Private/ControlRigEditorModule.cpp
+++ b/Engine/Plugins/Animation/ControlRig/Source/ControlRigEditor/Private/ControlRigEditorModule.cpp
@@ -1384,28 +1384,31 @@ void FControlRigEditorModule::GetContextMenuActions(const UControlRigGraphSchema
 
 					if (URigVMTemplateNode* TemplateNode = Cast<URigVMTemplateNode>(ModelPin->GetNode()))
 					{
-						FToolMenuSection& TemplatesSection = Menu->AddSection("EdGraphSchemaTemplates", LOCTEXT("TemplatesHeader", "Templates"));
-
-						if(ModelPin->IsRootPin())
+						if (!TemplateNode->IsSingleton())
 						{
-							TSharedRef<SControlRigChangePinType> ChangePinTypeWidget =
-							SNew(SControlRigChangePinType)
-							.Blueprint(RigBlueprint)
-							.ModelPins({ModelPin});
+							FToolMenuSection& TemplatesSection = Menu->AddSection("EdGraphSchemaTemplates", LOCTEXT("TemplatesHeader", "Templates"));
 
-							TemplatesSection.AddEntry(FToolMenuEntry::InitWidget("ChangePinTypeWidget", ChangePinTypeWidget, FText(), true));
-						}
+							if(ModelPin->IsRootPin())
+							{
+								TSharedRef<SControlRigChangePinType> ChangePinTypeWidget =
+								SNew(SControlRigChangePinType)
+								.Blueprint(RigBlueprint)
+								.ModelPins({ModelPin});
+
+								TemplatesSection.AddEntry(FToolMenuEntry::InitWidget("ChangePinTypeWidget", ChangePinTypeWidget, FText(), true));
+							}
 						
-						TemplatesSection.AddMenuEntry(
-							"Unresolve Template Node",
-							LOCTEXT("UnresolveTemplateNode", "Unresolve Template Node"),
-							LOCTEXT("UnresolveTemplateNode_Tooltip", "Removes any type information from the template node"),
-							FSlateIcon(),
-							FUIAction(FExecuteAction::CreateLambda([Controller, ModelPin]() {
-								const TArray<FName> Nodes = ModelPin->GetGraph()->GetSelectNodes();
-								Controller->UnresolveTemplateNodes(Nodes);
-							})
-						));
+							TemplatesSection.AddMenuEntry(
+								"Unresolve Template Node",
+								LOCTEXT("UnresolveTemplateNode", "Unresolve Template Node"),
+								LOCTEXT("UnresolveTemplateNode_Tooltip", "Removes any type information from the template node"),
+								FSlateIcon(),
+								FUIAction(FExecuteAction::CreateLambda([Controller, ModelPin]() {
+									const TArray<FName> Nodes = ModelPin->GetGraph()->GetSelectNodes();
+									Controller->UnresolveTemplateNodes(Nodes);
+								})
+							));
+						}
 					}
 
 
@@ -2328,17 +2331,20 @@ void FControlRigEditorModule::GetContextMenuActions(const UControlRigGraphSchema
 
 					if (URigVMTemplateNode* TemplateNode = Cast<URigVMTemplateNode>(RigNode->GetModelNode()))
 					{
-						FToolMenuSection& TemplatesSection = Menu->AddSection("EdGraphSchemaTemplates", LOCTEXT("TemplatesHeader", "Templates"));
-						TemplatesSection.AddMenuEntry(
-							"Unresolve Template Node",
-							LOCTEXT("UnresolveTemplateNode", "Unresolve Template Node"),
-							LOCTEXT("UnresolveTemplateNode_Tooltip", "Removes any type information from the template node"),
-							FSlateIcon(),
-							FUIAction(FExecuteAction::CreateLambda([Controller, Model]() {
-								const TArray<FName> Nodes = Model->GetSelectNodes();
-								Controller->UnresolveTemplateNodes(Nodes);
-							})
-						));
+						if (!TemplateNode->IsSingleton())
+						{
+							FToolMenuSection& TemplatesSection = Menu->AddSection("EdGraphSchemaTemplates", LOCTEXT("TemplatesHeader", "Templates"));
+							TemplatesSection.AddMenuEntry(
+								"Unresolve Template Node",
+								LOCTEXT("UnresolveTemplateNode", "Unresolve Template Node"),
+								LOCTEXT("UnresolveTemplateNode_Tooltip", "Removes any type information from the template node"),
+								FSlateIcon(),
+								FUIAction(FExecuteAction::CreateLambda([Controller, Model]() {
+									const TArray<FName> Nodes = Model->GetSelectNodes();
+									Controller->UnresolveTemplateNodes(Nodes);
+								})
+							));
+						}
 					}
 
 					if (URigVMLibraryNode* LibraryNode = Cast<URigVMLibraryNode>(RigNode->GetModelNode()))
diff --git a/Engine/Plugins/Animation/ControlRig/Source/ControlRigEditor/Private/Editor/ControlRigEditor.cpp b/Engine/Plugins/Animation/ControlRig/Source/ControlRigEditor/Private/Editor/ControlRigEditor.cpp
--- a/Engine/Plugins/Animation/ControlRig/Source/ControlRigEditor/Private/Editor/ControlRigEditor.cpp
+++ b/Engine/Plugins/Animation/ControlRig/Source/ControlRigEditor/Private/Editor/ControlRigEditor.cpp
@@ -84,6 +84,7 @@
 #include "IMessageLogListing.h"
 #include "SControlRigFunctionLocalizationWidget.h"
 #include "SControlRigFunctionBulkEditWidget.h"
+#include "SControlRigBreakLinksWidget.h"
 #include "SGraphPanel.h"
 #include "Engine/StaticMesh.h"
 #include "Engine/StaticMeshActor.h"
@@ -247,6 +248,7 @@ FControlRigEditor::~FControlRigEditor()
 		RigBlueprint->OnGraphImported().RemoveAll(this);
 		RigBlueprint->OnRequestLocalizeFunctionDialog().RemoveAll(this);
 		RigBlueprint->OnRequestBulkEditDialog().Unbind();
+		RigBlueprint->OnRequestBreakLinksDialog().Unbind();
 		RigBlueprint->OnRequestJumpToHyperlink().Unbind();
 		RigBlueprint->OnReportCompilerMessage().RemoveAll(this);
 
@@ -507,6 +509,7 @@ void FControlRigEditor::InitControlRigEditor(const EToolkitMode::Type Mode, cons
 		InControlRigBlueprint->OnGraphImported().AddSP(this, &FControlRigEditor::OnGraphImported);
 		InControlRigBlueprint->OnRequestLocalizeFunctionDialog().AddSP(this, &FControlRigEditor::OnRequestLocalizeFunctionDialog);
 		InControlRigBlueprint->OnRequestBulkEditDialog().BindSP(this, &FControlRigEditor::OnRequestBulkEditDialog);
+		InControlRigBlueprint->OnRequestBreakLinksDialog().BindSP(this, &FControlRigEditor::OnRequestBreakLinksDialog);
 		InControlRigBlueprint->OnRequestJumpToHyperlink().BindSP(this, &FControlRigEditor::HandleJumpToHyperlink);
 	}
 
@@ -4702,6 +4705,23 @@ FRigVMController_BulkEditResult FControlRigEditor::OnRequestBulkEditDialog(UCont
 	return Result;
 }
 
+bool FControlRigEditor::OnRequestBreakLinksDialog(TArray<URigVMLink*> InLinks)
+{
+	if(InLinks.Num() == 0)
+	{
+		return true;
+	}
+
+	TSharedRef<SControlRigBreakLinksDialog> BreakLinksDialog = SNew(SControlRigBreakLinksDialog)
+	.Links(InLinks)
+	.OnFocusOnLink(FControlRigOnFocusOnLinkRequestedDelegate::CreateLambda([&](URigVMLink* InLink)
+	{
+		HandleJumpToHyperlink(InLink);
+	}));
+
+	return BreakLinksDialog->ShowModal() == EAppReturnType::Ok; 
+}
+
 void FControlRigEditor::HandleJumpToHyperlink(const UObject* InSubject)
 {
 	UControlRigBlueprint* RigBlueprint = GetControlRigBlueprint();
diff --git a/Engine/Plugins/Animation/ControlRig/Source/ControlRigEditor/Private/Editor/ControlRigEditor.h b/Engine/Plugins/Animation/ControlRig/Source/ControlRigEditor/Private/Editor/ControlRigEditor.h
--- a/Engine/Plugins/Animation/ControlRig/Source/ControlRigEditor/Private/Editor/ControlRigEditor.h
+++ b/Engine/Plugins/Animation/ControlRig/Source/ControlRigEditor/Private/Editor/ControlRigEditor.h
@@ -439,6 +439,7 @@ protected:
 	void OnWrappedPropertyChangedChainEvent(UDetailsViewWrapperObject* InWrapperObject, const FString& InPropertyPath, FPropertyChangedChainEvent& InPropertyChangedChainEvent);
 	void OnRequestLocalizeFunctionDialog(URigVMLibraryNode* InFunction, UControlRigBlueprint* InTargetBlueprint, bool bForce);
 	FRigVMController_BulkEditResult OnRequestBulkEditDialog(UControlRigBlueprint* InBlueprint, URigVMController* InController, URigVMLibraryNode* InFunction, ERigVMControllerBulkEditType InEditType);
+	bool OnRequestBreakLinksDialog(TArray<URigVMLink*> InLinks);
 	void HandleJumpToHyperlink(const UObject* InSubject);
 	void UpdateDefaultValueForVariable(FBPVariableDescription& InVariable, bool bUseCDO);
 
diff --git a/Engine/Plugins/Animation/ControlRig/Source/ControlRigEditor/Private/Editor/SControlRigBreakLinksWidget.cpp b/Engine/Plugins/Animation/ControlRig/Source/ControlRigEditor/Private/Editor/SControlRigBreakLinksWidget.cpp
new file mode 100644
--- /dev/null
+++ b/Engine/Plugins/Animation/ControlRig/Source/ControlRigEditor/Private/Editor/SControlRigBreakLinksWidget.cpp
@@ -0,0 +1,208 @@
+﻿// Copyright Epic Games, Inc. All Rights Reserved.
+
+#include "SControlRigBreakLinksWidget.h"
+#include "Widgets/Layout/SBorder.h"
+#include "Widgets/SBoxPanel.h"
+#include "Widgets/Input/SButton.h"
+#include "Widgets/Input/SCheckBox.h"
+#include "Widgets/Layout/SUniformGridPanel.h"
+#include "DetailLayoutBuilder.h"
+#include "Editor.h"
+#include "EditorStyleSet.h"
+#include "HAL/ConsoleManager.h"
+#include "RigVMModel/RigVMController.h"
+#include "ControlRigEditor.h"
+#include "ContentBrowserModule.h"
+#include "IContentBrowserSingleton.h"
+#include "Misc/ScopedSlowTask.h"
+
+#define LOCTEXT_NAMESPACE "SControlRigBreakLinksWidget"
+
+class SControlRigBreakLinksWidgetRow : public SMultiColumnTableRow< URigVMLink* >
+{
+	SLATE_BEGIN_ARGS(SControlRigBreakLinksWidgetRow) { }
+	SLATE_ARGUMENT(URigVMLink*, Item)
+	SLATE_END_ARGS()
+	
+	void Construct( const FArguments& InArgs, const TSharedRef<STableViewBase>& InOwnerTableView)
+	{
+		Link = InArgs._Item;
+
+		SMultiColumnTableRow<URigVMLink*>::Construct(FSuperRowType::FArguments(), InOwnerTableView);
+	}
+
+	virtual TSharedRef<SWidget> GenerateWidgetForColumn(const FName& InColumnName) override
+	{
+		if(InColumnName == TEXT("Source"))
+		{
+			return SNew(SBox)
+				.Padding(FMargin(4.0f, 0.0f))
+				.VAlign(VAlign_Center)
+				[
+					SNew(STextBlock)
+					.Text(FText::FromString(Link->GetSourcePin()->GetPinPath()))
+				];
+		}
+		else if (InColumnName == TEXT("Target") )
+		{
+			return SNew(SBox)
+				.Padding(FMargin(4.0f, 0.0f))
+				.VAlign(VAlign_Center)
+				[
+					SNew(STextBlock)
+					.Text(FText::FromString(Link->GetTargetPin()->GetPinPath()))
+				];
+		}
+		else
+		{
+			ensure(false);
+			return SNew(SBorder);
+		}
+	}
+
+private:
+	URigVMLink* Link;
+};
+
+//////////////////////////////////////////////////////////////
+/// SControlRigBreakLinksWidget
+///////////////////////////////////////////////////////////
+
+void SControlRigBreakLinksWidget::Construct(const FArguments& InArgs, TArray<URigVMLink*> InLinks)
+{
+	Links = InLinks;
+	OnFocusOnLink = InArgs._OnFocusOnLink;
+
+	ChildSlot
+	[
+		SNew(SVerticalBox)
+
+		+ SVerticalBox::Slot()
+		.Padding(4.0f, 4.0f, 0.f, 0.f)
+		.AutoHeight()
+		.MaxHeight(600.f)
+		[
+			SNew(SBorder)
+			.Visibility(EVisibility::Visible)
+			.BorderImage(FEditorStyle::GetBrush("Menu.Background"))
+			[			
+			   SNew( SListView< URigVMLink* > )
+			     .ItemHeight(24)
+			     .ListItemsSource( &Links )
+			     .ListViewStyle( FAppStyle::Get(), "SimpleListView" )
+			     .OnMouseButtonDoubleClick(this, &SControlRigBreakLinksWidget::HandleItemMouseDoubleClick)
+			     .HeaderRow(SNew(SHeaderRow)
+			     	+SHeaderRow::Column(FName(TEXT("Source")))
+			     	.DefaultLabel(LOCTEXT("SourceColumnHeader", "Source"))
+			     	.FillWidth(0.5f)
+			     	+SHeaderRow::Column(FName(TEXT("Target")))
+			     	.DefaultLabel(LOCTEXT("TargetColumnHeader", "Target"))
+			     	.FillWidth(0.5f)
+			     )
+			     .OnGenerateRow(this, &SControlRigBreakLinksWidget::GenerateItemRow)
+			]
+		]
+	];
+}
+
+TSharedRef<ITableRow> SControlRigBreakLinksWidget::GenerateItemRow(URigVMLink* Item, const TSharedRef<STableViewBase>& OwnerTable)
+{
+	return SNew(SControlRigBreakLinksWidgetRow, OwnerTable)
+			.Item(Item);
+}
+
+void SControlRigBreakLinksWidget::HandleItemMouseDoubleClick(URigVMLink* InItem)
+{
+	OnFocusOnLink.ExecuteIfBound(InItem);		
+}
+
+void SControlRigBreakLinksDialog::Construct(const FArguments& InArgs)
+{
+	UserResponse = EAppReturnType::Cancel;
+
+	static const FString DialogText = TEXT(
+        "This action will break links." 
+    );
+
+	const FString DialogTitle = TEXT("Break Links"); 
+	
+	SWindow::Construct(SWindow::FArguments()
+        .Title(FText::FromString(DialogTitle))
+        .SizingRule( ESizingRule::Autosized )
+        //.ClientSize(FVector2D(450, 450))
+        [
+	        SNew(SHorizontalBox)
+
+	        +SHorizontalBox::Slot()
+	        .FillWidth(1.f)
+	        .HAlign(HAlign_Left)
+	        .VAlign(VAlign_Top)
+	        //.Padding(2)
+	        [
+				SNew(SVerticalBox)
+
+                +SVerticalBox::Slot()
+                .AutoHeight()
+                .HAlign(HAlign_Left)
+                //.Padding(8)
+                [
+
+                    SNew(STextBlock)
+                    .Text(FText::FromString(DialogText))
+                    .Font(IDetailLayoutBuilder::GetDetailFont())
+                ]
+				
+                +SVerticalBox::Slot()
+                .AutoHeight()
+                .HAlign(HAlign_Left)
+                .Padding(2)
+                [
+                    SAssignNew(BreakLinksWidget, SControlRigBreakLinksWidget, InArgs._Links)
+                    .OnFocusOnLink(InArgs._OnFocusOnLink)
+                ]
+
+                +SVerticalBox::Slot()
+                .AutoHeight()
+                .HAlign(HAlign_Right)
+                .Padding(5)
+                [
+                    SNew(SUniformGridPanel)
+                    .SlotPadding(FEditorStyle::GetMargin("StandardDialog.SlotPadding"))
+                    .MinDesiredSlotWidth(FEditorStyle::GetFloat("StandardDialog.MinDesiredSlotWidth"))
+                    .MinDesiredSlotHeight(FEditorStyle::GetFloat("StandardDialog.MinDesiredSlotHeight"))
+                    +SUniformGridPanel::Slot(0, 0)
+                    [
+                        SNew(SButton)
+                        .HAlign(HAlign_Center)
+                        .ContentPadding(FEditorStyle::GetMargin("StandardDialog.ContentPadding"))
+                        .Text(LOCTEXT("OK", "OK"))
+                        .OnClicked(this, &SControlRigBreakLinksDialog::OnButtonClick, EAppReturnType::Ok)
+                    ]
+                    +SUniformGridPanel::Slot(1, 0)
+                    [
+                        SNew(SButton)
+                        .HAlign(HAlign_Center)
+                        .ContentPadding(FEditorStyle::GetMargin("StandardDialog.ContentPadding"))
+                        .Text(LOCTEXT("Cancel", "Cancel"))
+                        .OnClicked(this, &SControlRigBreakLinksDialog::OnButtonClick, EAppReturnType::Cancel)
+                    ]
+                ]
+			]
+		]);
+}
+
+FReply SControlRigBreakLinksDialog::OnButtonClick(EAppReturnType::Type ButtonID)
+{
+	UserResponse = ButtonID;
+
+	RequestDestroyWindow();
+	return FReply::Handled();
+}
+
+EAppReturnType::Type SControlRigBreakLinksDialog::ShowModal()
+{
+	GEditor->EditorAddModalWindow(SharedThis(this));
+	return UserResponse;
+}
+
+#undef LOCTEXT_NAMESPACE
diff --git a/Engine/Plugins/Animation/ControlRig/Source/ControlRigEditor/Private/Editor/SControlRigBreakLinksWidget.h b/Engine/Plugins/Animation/ControlRig/Source/ControlRigEditor/Private/Editor/SControlRigBreakLinksWidget.h
new file mode 100644
--- /dev/null
+++ b/Engine/Plugins/Animation/ControlRig/Source/ControlRigEditor/Private/Editor/SControlRigBreakLinksWidget.h
@@ -0,0 +1,59 @@
+﻿// Copyright Epic Games, Inc. All Rights Reserved.
+
+#pragma once
+
+#include "CoreMinimal.h"
+#include "ControlRigBlueprint.h"
+#include "Dialogs/Dialogs.h"
+#include "IAssetTypeActions.h"
+
+DECLARE_DELEGATE_OneParam(FControlRigOnFocusOnLinkRequestedDelegate, URigVMLink*);
+
+class SControlRigBreakLinksWidget : public SCompoundWidget
+{
+public:
+	SLATE_BEGIN_ARGS(SControlRigBreakLinksWidget) {}
+	SLATE_ARGUMENT(FControlRigOnFocusOnLinkRequestedDelegate, OnFocusOnLink)
+	SLATE_END_ARGS()
+
+	void Construct(const FArguments& InArgs, TArray<URigVMLink*> InLinks);
+
+private:
+
+	/** Handler for when the user double clicks, presses enter, or presses space on an asset */
+	void OnLinkDoubleClicked(const URigVMLink* InLink, EAssetTypeActivationMethod::Type ActivationMethod);
+
+	TSharedRef<ITableRow> GenerateItemRow(URigVMLink* Item, const TSharedRef<STableViewBase>& OwnerTable);
+ 
+	TArray<URigVMLink*> Links;
+	FControlRigOnFocusOnLinkRequestedDelegate OnFocusOnLink;
+
+	void HandleItemMouseDoubleClick(URigVMLink* InItem);
+
+	friend class SControlRigBreakLinksDialog;
+};
+
+class SControlRigBreakLinksDialog : public SWindow
+{
+public:
+	
+	SLATE_BEGIN_ARGS(SControlRigBreakLinksDialog)
+	{
+	}
+
+	SLATE_ARGUMENT(TArray<URigVMLink*>, Links)
+	SLATE_ARGUMENT(FControlRigOnFocusOnLinkRequestedDelegate, OnFocusOnLink)
+	SLATE_END_ARGS()
+
+	void Construct(const FArguments& InArgs);
+
+	EAppReturnType::Type ShowModal();
+
+protected:
+
+	TSharedPtr<SControlRigBreakLinksWidget> BreakLinksWidget;
+	
+	FReply OnButtonClick(EAppReturnType::Type ButtonID);
+	bool IsOkButtonEnabled() const;
+	EAppReturnType::Type UserResponse;
+};
\ No newline at end of file
diff --git a/Engine/Plugins/Animation/ControlRig/Source/ControlRigEditor/Private/Graph/SControlRigGraphChangePinType.cpp b/Engine/Plugins/Animation/ControlRig/Source/ControlRigEditor/Private/Graph/SControlRigGraphChangePinType.cpp
--- a/Engine/Plugins/Animation/ControlRig/Source/ControlRigEditor/Private/Graph/SControlRigGraphChangePinType.cpp
+++ b/Engine/Plugins/Animation/ControlRig/Source/ControlRigEditor/Private/Graph/SControlRigGraphChangePinType.cpp
@@ -169,22 +169,25 @@ void SControlRigChangePinType::FillPinTypeMenu(FMenuBuilder& MenuBuilder)
 
 	struct FArgumentInfo
 	{
-		static FArgumentInfo Make(const FProperty* InProperty)
+		static FArgumentInfo Make(const FProperty* InProperty, const bool bInIsFilteredOut)
 		{
 			FArgumentInfo Info;
 			Info.Property = InProperty;
+			Info.bIsFilteredOut = bInIsFilteredOut;
 			return Info;
 		}
 
-		static FArgumentInfo Make(const FRigVMTemplateArgument::FType& InType)
+		static FArgumentInfo Make(const FRigVMTemplateArgument::FType& InType, const bool bInIsFilteredOut)
 		{
 			FArgumentInfo Info;
 			Info.Property = nullptr;
 			Info.Type = InType;
+			Info.bIsFilteredOut = bInIsFilteredOut;
 			return Info;
 		}
 
 		const FProperty* Property;
+		bool bIsFilteredOut;
 		FRigVMTemplateArgument::FType Type;
 	};
 
@@ -205,14 +208,16 @@ void SControlRigChangePinType::FillPinTypeMenu(FMenuBuilder& MenuBuilder)
 				if(const FRigVMTemplateArgument* Argument = Template->FindArgument(ArgumentName))
 				{
 					const TArray<FRigVMTemplateArgument::FType>& AllArgumentTypes = Argument->GetTypes();
-					const TArray<FRigVMTemplateArgument::FType>& FilteredArgumentTypes = Argument->GetSupportedTypes(TemplateNode->GetResolvedPermutationIndices());
+					const TArray<FRigVMTemplateArgument::FType>& FilteredArgumentTypes = Argument->GetSupportedTypes(TemplateNode->GetFilteredPermutationsIndices());
 					
 					for(int32 PermutationIndex = 0; PermutationIndex < Template->NumPermutations(); PermutationIndex++)
 					{
 						FRigVMTemplateArgument::FType ArgumentType = AllArgumentTypes[PermutationIndex];
+					
+						bool bIsFilteredOut = false;
 						if(!FilteredArgumentTypes.Contains(ArgumentType))
 						{
-							continue;
+							bIsFilteredOut = true;
 						}
 
 						if(ArgumentType.CPPType == RigVMTypeUtils::FloatType)
@@ -235,13 +240,13 @@ void SControlRigChangePinType::FillPinTypeMenu(FMenuBuilder& MenuBuilder)
 							{
 								if(const FProperty* Property = FunctionStruct->FindPropertyByName(ArgumentName))
 								{
-									Types.Add(ArgumentType, FArgumentInfo::Make(Property));
+									Types.Add(ArgumentType, FArgumentInfo::Make(Property, bIsFilteredOut));
 								}
 							}
 						}
 						else
 						{
-							Types.Add(ArgumentType, FArgumentInfo::Make(ArgumentType));
+							Types.Add(ArgumentType, FArgumentInfo::Make(ArgumentType, bIsFilteredOut));
 						}
 					}
 				}
@@ -300,6 +305,7 @@ void SControlRigChangePinType::FillPinTypeMenu(FMenuBuilder& MenuBuilder)
 		for(int32 TypeIndex=0; TypeIndex < SortedTypes.Num(); TypeIndex++)
 		{
 			const FRigVMTemplateArgument::FType& Type = SortedTypes[TypeIndex].Key;
+			const bool bIsFilteredOut = SortedTypes[TypeIndex].Value.bIsFilteredOut;
 			if (Type.CPPTypeObject != nullptr && !IsValid(Type.CPPTypeObject))
 			{
 				continue;
@@ -348,7 +354,7 @@ void SControlRigChangePinType::FillPinTypeMenu(FMenuBuilder& MenuBuilder)
 					[
 						SNew(SImage)
 						.Image(FBlueprintEditorUtils::GetIconFromPin(PinType, true))
-						.ColorAndOpacity(Schema->GetPinTypeColor(PinType))
+						.ColorAndOpacity(Schema->GetPinTypeColor(PinType) * (bIsFilteredOut ? 0.5f : 1.f))
 					]
 					+SHorizontalBox::Slot()
 					.AutoWidth()
@@ -357,6 +363,9 @@ void SControlRigChangePinType::FillPinTypeMenu(FMenuBuilder& MenuBuilder)
 					[
 						SNew(STextBlock)
 						.Text(GetBindingText(Type))
+						.ColorAndOpacity(FLinearColor::White * (bIsFilteredOut ? 0.5f : 1.f))
+						.ToolTipText(bIsFilteredOut ? LOCTEXT("WildcardAvailableTypeTooltip","Will break connections if resolved to this type.") :
+							LOCTEXT("WildcardAvailableTypeTooltip", "Available filtered type"))
 					]);
 		}
 	}
diff --git a/Engine/Source/Developer/RigVMDeveloper/Private/RigVMModel/Nodes/RigVMTemplateNode.cpp b/Engine/Source/Developer/RigVMDeveloper/Private/RigVMModel/Nodes/RigVMTemplateNode.cpp
--- a/Engine/Source/Developer/RigVMDeveloper/Private/RigVMModel/Nodes/RigVMTemplateNode.cpp
+++ b/Engine/Source/Developer/RigVMDeveloper/Private/RigVMModel/Nodes/RigVMTemplateNode.cpp
@@ -1,7 +1,7 @@
 // Copyright Epic Games, Inc. All Rights Reserved.
 
 #include "RigVMModel/Nodes/RigVMTemplateNode.h"
-
+#include "RigVMCore/RigVMTemplate.h"
 #include "RigVMModel/RigVMController.h"
 
 URigVMTemplateNode::URigVMTemplateNode()
@@ -71,7 +71,7 @@ FText URigVMTemplateNode::GetToolTipText() const
 {
 	if(const FRigVMTemplate* Template = GetTemplate())
 	{
-		const TArray<int32> PermutationIndices = GetResolvedPermutationIndices();
+		const TArray<int32> PermutationIndices = GetFilteredPermutationsIndices();
 		return Template->GetTooltipText(PermutationIndices);
 	}
 	return Super::GetToolTipText();
@@ -89,7 +89,7 @@ FText URigVMTemplateNode::GetToolTipTextForPin(const URigVMPin* InPin) const
 			{
 				FString Tooltip;
 
-				const TArray<int32> PermutationIndices = GetResolvedPermutationIndices();
+				const TArray<int32> PermutationIndices = GetFilteredPermutationsIndices();
 				if(PermutationIndices.Num() == GetTemplate()->NumPermutations())
 				{
 					if(Arg->GetTypes().Num() > 100)
@@ -100,7 +100,16 @@ FText URigVMTemplateNode::GetToolTipTextForPin(const URigVMPin* InPin) const
 				
 				if(Tooltip.IsEmpty())
 				{
-					const FString SupportedTypesJoined = FString::Join(Arg->GetSupportedTypeStrings(PermutationIndices), TEXT("\n"));
+					FString SupportedTypesJoined;
+					for (int32 Index : PermutationIndices)
+					{
+						FString Type = Arg->GetTypes()[Index].CPPType;
+						if (!FilteredPermutations.Contains(Index))
+						{
+							Type += TEXT(" : Breaks Connections");
+						}
+						SupportedTypesJoined += Type + TEXT("\n");
+					}
 					Tooltip = TEXT("Supported Types:\n\n") + SupportedTypesJoined;
 				}
 
@@ -121,6 +130,11 @@ FName URigVMTemplateNode::GetNotation() const
 	return TemplateNotation;
 }
 
+bool URigVMTemplateNode::IsSingleton() const
+{
+	return GetTemplate() == nullptr;
+}
+
 bool URigVMTemplateNode::SupportsType(const URigVMPin* InPin, const FString& InCPPType, FString* OutCPPType)
 {
 	static const FString WildCardCPPType = RigVMTypeUtils::GetWildCardCPPType();
@@ -197,6 +211,10 @@ bool URigVMTemplateNode::SupportsType(const URigVMPin* InPin, const FString& InC
 			{
 				continue;
 			}
+			if (Pin->GetDirection() == ERigVMPinDirection::Hidden)
+			{
+				continue;
+			}
 			Types.Add(Pin->GetFName(), Pin->GetTemplateArgumentType());
 		}
 
@@ -215,11 +233,15 @@ bool URigVMTemplateNode::SupportsType(const URigVMPin* InPin, const FString& InC
 				ResolvedType = ResolvedTypes.FindChecked(RootPin->GetFName());
 			}
 		}
-	
-		SupportedTypesCache.Add(CacheKey, TPair<bool, FRigVMTemplateArgument::FType>(bSupportsType, ResolvedType));
-		if(OutCPPType)
+
+		if (bSupportsType)
 		{
-			*OutCPPType = ResolvedType.CPPType;
+			ensure(ResolvedType.CPPTypeObject != RigVMTypeUtils::GetWildCardCPPTypeObject());
+			SupportedTypesCache.Add(CacheKey, TPair<bool, FRigVMTemplateArgument::FType>(bSupportsType, ResolvedType));
+			if(OutCPPType)
+			{
+				*OutCPPType = ResolvedType.CPPType;
+			}
 		}
 		return bSupportsType;
 	}
@@ -235,46 +257,47 @@ bool URigVMTemplateNode::SupportsType(const URigVMPin* InPin, const FString& InC
 	return false;
 }
 
-bool URigVMTemplateNode::GetTypeMapForNewPinType(const URigVMPin* InPin, const FString& InCPPType, UObject* InCPPTypeObject,
-	FRigVMTemplate::FTypeMap& OutTypes) const
+bool URigVMTemplateNode::FilteredSupportsType(const URigVMPin* InPin, const FString& InCPPType, FString* OutCPPType, bool bAllowFloatingPointCasts)
 {
-	check(InPin);
-	check(InPin->GetNode() == this);
-	check(InPin->IsRootPin());
-	
-	if(const FRigVMTemplate* Template = GetTemplate())
+	if (OutCPPType)
 	{
-		OutTypes = GetResolvedTypes();
-		const FRigVMTemplateArgument::FType ExpectedType(InCPPType, InCPPTypeObject);
-		return GetTemplate()->ResolveArgument(InPin->GetFName(), ExpectedType, OutTypes);
+		*OutCPPType = FString(); 
 	}
 
-	OutTypes.Reset();
-	return false;
-}
-
-const TArray<int32>& URigVMTemplateNode::GetResolvedPermutationIndices(FRigVMTemplate::FTypeMap* OutTypes) const
-{
-	if(!ResolvedPermutations.IsEmpty())
+	const URigVMPin* RootPin = InPin;
+	bool bIsArrayElement = false;
+	if (URigVMPin* ParentPin = InPin->GetParentPin())
 	{
-		return ResolvedPermutations;
+		RootPin = ParentPin;
+		bIsArrayElement = true;
 	}
-	if (const FRigVMTemplate* Template = GetTemplate())
+	
+	const FRigVMTemplateArgument* Argument = GetTemplate()->FindArgument(RootPin->GetFName());
+	if (Argument == nullptr)
 	{
-		FRigVMTemplate::FTypeMap Types = GetResolvedTypes();
-		Template->Resolve(Types, ResolvedPermutations, false);
+		return false;
+	}
 
-		if (OutTypes)
+	const TArray<FRigVMTemplateArgument::FType>& Types = Argument->GetTypes();
+	for (int32 PermutationIndex : FilteredPermutations)
+	{
+		FRigVMTemplateArgument::FType FilteredType = Types[PermutationIndex];
+		if (bIsArrayElement)
+		{
+			FilteredType.ConvertToBaseElement();
+		}
+		if (FilteredType.Matches(InCPPType, bAllowFloatingPointCasts))
 		{
-			*OutTypes = Types;
+			return true;
 		}
 	}
-	return ResolvedPermutations;
+
+	return false;
 }
 
-TArray<const FRigVMFunction*> URigVMTemplateNode::GetResolvedPermutations(FRigVMTemplate::FTypeMap* OutTypes) const
+TArray<const FRigVMFunction*> URigVMTemplateNode::GetResolvedPermutations() const
 {
-	TArray<int32> Indices = GetResolvedPermutationIndices(OutTypes);
+	TArray<int32> Indices = GetFilteredPermutationsIndices();
 	TArray<const FRigVMFunction*> Functions;
 	for(const int32 Index : Indices)
 	{
@@ -292,20 +315,6 @@ const FRigVMTemplate* URigVMTemplateNode::GetTemplate() const
 	return CachedTemplate;
 }
 
-FRigVMTemplate::FTypeMap URigVMTemplateNode::GetResolvedTypes() const
-{
-	FRigVMTemplate::FTypeMap Types;
-	for (URigVMPin* Pin : GetPins())
-	{
-		if (Pin->IsWildCard())
-		{
-			continue;
-		}
-		Types.Add(Pin->GetFName(), Pin->GetTemplateArgumentType());
-	}
-	return Types;
-}
-
 const FRigVMFunction* URigVMTemplateNode::GetResolvedFunction() const
 {
 	if(CachedFunction == nullptr)
@@ -317,7 +326,7 @@ const FRigVMFunction* URigVMTemplateNode::GetResolvedFunction() const
 
 		if(CachedFunction == nullptr)
 		{
-			TArray<int32> PermutationIndices = GetResolvedPermutationIndices();
+			TArray<int32> PermutationIndices = GetFilteredPermutationsIndices();
 			if(PermutationIndices.Num() == 1)
 			{
 				CachedFunction = GetTemplate()->GetPermutation(PermutationIndices[0]);
@@ -337,7 +346,7 @@ bool URigVMTemplateNode::IsFullyUnresolved() const
 	check(GetTemplate());
 
 	// all permutations are available means we haven't resolved any wildcard pin
-	return GetResolvedPermutations().Num() == GetTemplate()->NumPermutations();
+	return GetFilteredPermutationsIndices().Num() == GetTemplate()->NumPermutations();
 }
 
 FString URigVMTemplateNode::GetInitialDefaultValueForPin(const FName& InRootPinName, const TArray<int32>& InPermutationIndices) const
@@ -350,7 +359,7 @@ FString URigVMTemplateNode::GetInitialDefaultValueForPin(const FName& InRootPinN
 	TArray<int32> PermutationIndices = InPermutationIndices;
 	if(PermutationIndices.IsEmpty())
 	{
-		PermutationIndices = GetResolvedPermutationIndices();
+		PermutationIndices = GetFilteredPermutationsIndices();
 	}
 
 	FString DefaultValue;
@@ -451,7 +460,7 @@ FName URigVMTemplateNode::GetDisplayNameForPin(const FName& InRootPinName,
 		const TArray<int32>* PermutationIndicesPtr = &InPermutationIndices;
 		if(PermutationIndicesPtr->IsEmpty())
 		{
-			PermutationIndicesPtr = &GetResolvedPermutationIndices();
+			PermutationIndicesPtr = &GetFilteredPermutationsIndices();
 		}
 
 		const FText DisplayNameText = Template->GetDisplayNameForArgument(InRootPinName, *PermutationIndicesPtr);
@@ -471,6 +480,262 @@ FName URigVMTemplateNode::GetDisplayNameForPin(const FName& InRootPinName,
 	return NAME_None;
 }
 
+const TArray<int32>& URigVMTemplateNode::GetFilteredPermutationsIndices() const
+{
+	return FilteredPermutations;
+}
+
+TArray<FRigVMTemplateArgument::FType> URigVMTemplateNode::GetFilteredTypesForPin(URigVMPin* InPin) const
+{
+	ensureMsgf(InPin->GetNode() == this, TEXT("GetFilteredTypesForPin of %s with pin from another node %s"), *GetNodePath(), *InPin->GetPinPath(true));
+	TArray<FRigVMTemplateArgument::FType> FilteredTypes;
+
+	if (FilteredPermutations.IsEmpty())
+	{
+		return FilteredTypes;	
+	}
+	
+	URigVMPin* RootPin = InPin;
+	bool bIsArrayElement = false;
+	if (URigVMPin* ParentPin = InPin->GetParentPin())
+	{
+		RootPin = ParentPin;
+		bIsArrayElement = true;
+	}
+	if (const FRigVMTemplateArgument* Argument = GetTemplate()->FindArgument(RootPin->GetFName()))
+	{
+		FilteredTypes = Argument->GetSupportedTypes(FilteredPermutations);
+		if (bIsArrayElement)
+		{
+			for (FRigVMTemplateArgument::FType& ArrayType : FilteredTypes)
+			{
+				ArrayType.ConvertToBaseElement();
+			}
+		}
+	}
+	return FilteredTypes;
+}
+
+TArray<int32> URigVMTemplateNode::GetNewFilteredPermutations(URigVMPin* InPin, URigVMPin* LinkedPin)
+{	
+	TArray<int32> NewFilteredPermutations;
+	if (InPin->GetNode() != this)
+	{
+		return NewFilteredPermutations;
+	}
+
+	NewFilteredPermutations.Reserve(FilteredPermutations.Num());
+	
+	bool bIsArrayElement = false;
+	URigVMPin* RootPin = InPin;
+	if (URigVMPin* ParentPin = InPin->GetParentPin())
+	{
+		RootPin = ParentPin;
+		bIsArrayElement = true;
+	}
+
+	TArray<int32> PermutationsToTry = PreferredPermutationTypes.IsEmpty()? FilteredPermutations : TArray<int32>({FindPermuationForTypes(PreferredPermutationTypes)});
+	bool bLinkedIsTemplate = false;
+	if (URigVMTemplateNode* LinkedTemplate = Cast<URigVMTemplateNode>(LinkedPin->GetNode()))
+	{
+		if (!LinkedTemplate->IsSingleton())
+		{
+			bLinkedIsTemplate = true;
+			if (const FRigVMTemplateArgument* Argument = GetTemplate()->FindArgument(RootPin->GetFName()))
+			{
+				for(int32 PermutationIndex : PermutationsToTry)
+				{
+					FRigVMTemplateArgument::FType Type = Argument->GetTypes()[PermutationIndex];
+					if (bIsArrayElement)
+					{
+						Type.ConvertToBaseElement();
+					}
+					if (LinkedTemplate->FilteredSupportsType(LinkedPin, Type.CPPType))
+					{
+						NewFilteredPermutations.Add(PermutationIndex);
+					}
+				}
+			}
+		}
+	}
+	
+	if (!bLinkedIsTemplate)
+	{
+		if (const FRigVMTemplateArgument* Argument = GetTemplate()->FindArgument(RootPin->GetFName()))
+		{
+			for(int32 PermIndex : PermutationsToTry)
+			{
+				FRigVMTemplateArgument::FType Type = Argument->GetTypes()[PermIndex];
+				if (bIsArrayElement)
+				{
+					Type.ConvertToBaseElement();
+				}
+				if (Type.Matches(LinkedPin->GetCPPType()))
+				{
+					NewFilteredPermutations.Add(PermIndex);
+				}
+			}
+		}
+	}
+	return NewFilteredPermutations;
+}
+
+TArray<int32> URigVMTemplateNode::GetNewFilteredPermutations(URigVMPin* InPin, const TArray<FRigVMTemplateArgument::FType>& InTypes)
+{
+	TArray<int32> NewFilteredPermutations;
+	NewFilteredPermutations.Reserve(FilteredPermutations.Num());
+
+	URigVMPin* RootPin = InPin;
+	TArray<FRigVMTemplateArgument::FType> RootTypes = InTypes;
+	if (URigVMPin* ParentPin = InPin->GetParentPin())
+	{
+		RootPin = ParentPin;
+		for (FRigVMTemplateArgument::FType& Type : RootTypes)
+		{
+			Type.ConvertToArray();
+		}
+	}
+
+	TArray<int32> PermutationsToTry = PreferredPermutationTypes.IsEmpty()? FilteredPermutations : TArray<int32>({FindPermuationForTypes(PreferredPermutationTypes)});
+	
+	if (const FRigVMTemplateArgument* Argument = GetTemplate()->FindArgument(RootPin->GetFName()))
+	{
+		const TArray<FRigVMTemplateArgument::FType>& Types = Argument->GetTypes();
+		for(int32 PermIndex : PermutationsToTry)
+		{
+			for (FRigVMTemplateArgument::FType& RootType : RootTypes)
+			{				
+				if (Types[PermIndex].Matches(RootType.CPPType))
+				{
+					NewFilteredPermutations.Add(PermIndex);
+					break;
+				}
+			}			
+		}
+	}
+	return NewFilteredPermutations;
+}
+
+int32 URigVMTemplateNode::FindPermuationForTypes(const TArray<FString>& ArgumentTypes)
+{
+	if (const FRigVMTemplate* Template = GetTemplate())
+	{
+		for (int32 i=0; i<Template->NumPermutations(); ++i)
+		{
+			bool bAllArgsMatched = true;
+			for (const FString& TypePairString : ArgumentTypes)
+			{
+				FString ArgName, Type;
+				if (!TypePairString.Split(TEXT(":"), &ArgName, &Type))
+				{
+					return INDEX_NONE;
+				}
+				
+				if (const FRigVMTemplateArgument* Argument = Template->FindArgument(*ArgName))
+				{
+					if (!Argument->GetTypes()[i].Matches(Type))
+					{
+						bAllArgsMatched = false;
+						break;
+					}
+				}
+				else
+				{
+					return INDEX_NONE;
+				}
+			}
+
+			if (bAllArgsMatched)
+			{
+				return i;
+			}
+		}
+	}
+	return INDEX_NONE;
+}
+
+TArray<FString> URigVMTemplateNode::GetArgumentTypesForPermutation(const int32 InPermutationIndex)
+{
+	TArray<FString> ArgTypes;
+	if (const FRigVMTemplate* Template = GetTemplate())
+	{
+		for (int32 ArgIndex = 0; ArgIndex < Template->NumArguments(); ++ArgIndex)
+		{
+			const FRigVMTemplateArgument* Argument = Template->GetArgument(ArgIndex);
+			if (Argument->GetTypes().Num() > InPermutationIndex)
+			{
+				ArgTypes.Add(Argument->GetName().ToString() + TEXT(":") + Argument->GetTypes()[InPermutationIndex].CPPType);
+			}
+			else
+			{
+				ArgTypes.Reset();
+				return ArgTypes;
+			}
+		}
+	}
+
+	return ArgTypes;
+}
+
+bool URigVMTemplateNode::PinNeedsFilteredTypesUpdate(URigVMPin* InPin, const TArray<FRigVMTemplateArgument::FType>& InTypes)
+{
+	TArray<int32> NewFilteredPermutations = GetNewFilteredPermutations(InPin, InTypes);
+	if (NewFilteredPermutations.Num() == FilteredPermutations.Num())
+	{
+		return false;
+	}
+	return true;
+}
+
+bool URigVMTemplateNode::PinNeedsFilteredTypesUpdate(URigVMPin* InPin, URigVMPin* LinkedPin)
+{
+	TArray<int32> NewFilteredPermutations = GetNewFilteredPermutations(InPin, LinkedPin);
+
+	if (NewFilteredPermutations.Num() == FilteredPermutations.Num())
+	{
+		return false;
+	}
+
+	return true;
+}
+
+bool URigVMTemplateNode::UpdateFilteredPermutations(URigVMPin* InPin, URigVMPin* LinkedPin)
+{
+	ensureMsgf(InPin->GetNode() == this, TEXT("Updating filtered permutations of %s with pin from another node %s"), *GetNodePath(), *InPin->GetPinPath(true));
+	ensureMsgf(LinkedPin->GetNode() != this, TEXT("Updating filtered permutations of %s with linked pin from same node %s"), *GetNodePath(), *LinkedPin->GetPinPath(true));
+
+	TArray<int32> NewFilteredPermutations = GetNewFilteredPermutations(InPin, LinkedPin);	
+	if (NewFilteredPermutations.Num() == FilteredPermutations.Num())
+	{
+		return false;
+	}
+
+	if (NewFilteredPermutations.IsEmpty())
+	{
+		return false;
+	}
+	
+	FilteredPermutations = NewFilteredPermutations;
+	return true;
+}
+
+bool URigVMTemplateNode::UpdateFilteredPermutations(URigVMPin* InPin, const TArray<FRigVMTemplateArgument::FType>& InTypes)
+{
+	TArray<int32> NewFilteredPermutations = GetNewFilteredPermutations(InPin, InTypes);
+	if (NewFilteredPermutations.Num() == FilteredPermutations.Num())
+	{
+		return false;
+	}
+
+	if (NewFilteredPermutations.IsEmpty())
+	{
+		return false;
+	}
+	
+	FilteredPermutations = NewFilteredPermutations;
+	return true;
+}
+
 void URigVMTemplateNode::InvalidateCache()
 {
 	SupportedTypesCache.Reset();
@@ -486,3 +751,22 @@ void URigVMTemplateNode::InvalidateCache()
 		}
 	}
 }
+
+void URigVMTemplateNode::InitializeFilteredPermutations()
+{
+	if (const FRigVMTemplate* Template = GetTemplate())
+	{
+		if (!PreferredPermutationTypes.IsEmpty())
+		{
+			FilteredPermutations = {FindPermuationForTypes(PreferredPermutationTypes)};
+		}
+		else
+		{
+			FilteredPermutations.SetNumUninitialized(Template->NumPermutations());
+			for (int32 i=0; i<FilteredPermutations.Num(); ++i)
+			{
+				FilteredPermutations[i] = i;
+			}
+		}
+	}	
+}
diff --git a/Engine/Source/Developer/RigVMDeveloper/Private/RigVMModel/RigVMController.cpp b/Engine/Source/Developer/RigVMDeveloper/Private/RigVMModel/RigVMController.cpp
--- a/Engine/Source/Developer/RigVMDeveloper/Private/RigVMModel/RigVMController.cpp
+++ b/Engine/Source/Developer/RigVMDeveloper/Private/RigVMModel/RigVMController.cpp
@@ -69,6 +69,7 @@ URigVMController::URigVMController()
 	, bIsTransacting(false)
 	, bIsRunningUnitTest(false)
 	, bIsFullyResolvingTemplateNode(false)
+	, bSuspendRecomputingTemplateFilters(false)
 {
 }
 
@@ -82,6 +83,7 @@ URigVMController::URigVMController(const FObjectInitializer& ObjectInitializer)
 	, bIsTransacting(false)
 	, bIsRunningUnitTest(false)
 	, bIsFullyResolvingTemplateNode(false)
+	, bSuspendRecomputingTemplateFilters(false)
 {
 	ActionStack = CreateDefaultSubobject<URigVMActionStack>(TEXT("ActionStack"));
 
@@ -746,7 +748,28 @@ URigVMUnitNode* URigVMController::AddUnitNode(UScriptStruct* InScriptStruct, con
 			return nullptr;
 		}
 
-		FullyResolveTemplateNode(TemplateNode, Template->FindPermutation(Function), bSetupUndoRedo);
+		TArray<int32> OldPermutations = TemplateNode->FilteredPermutations;
+		int32 PermutationIndex = Template->FindPermutation(Function);
+		TemplateNode->FilteredPermutations = {PermutationIndex};
+		const TArray<FString> NewPreferredPermutationTypes = TemplateNode->GetArgumentTypesForPermutation(PermutationIndex);
+		if (bSetupUndoRedo)
+		{
+			FRigVMSetTemplateFilteredPermutationsAction Action(TemplateNode, nullptr, OldPermutations);
+			ActionStack->AddAction(Action);
+			ActionStack->AddAction(FRigVMSetPreferredTemplatePermutationsAction(TemplateNode, NewPreferredPermutationTypes));
+		}		
+		TemplateNode->PreferredPermutationTypes = NewPreferredPermutationTypes;
+		UpdateTemplateNodePinTypes(TemplateNode, bSetupUndoRedo);
+
+		if (UnitNodeCreatedContext.IsValid())
+		{
+			if (TSharedPtr<FStructOnScope> StructScope = TemplateNode->ConstructStructInstance())
+			{
+				TGuardValue<FName> NodeNameScope(UnitNodeCreatedContext.NodeName, TemplateNode->GetFName());
+				FRigVMStruct* StructInstance = (FRigVMStruct*)StructScope->GetStructMemory();
+				StructInstance->OnUnitNodeCreated(UnitNodeCreatedContext);
+			}
+		}
 		
 		if(bSetupUndoRedo)
 		{
@@ -1681,7 +1704,10 @@ bool URigVMController::UnresolveTemplateNodes(const TArray<FName>& InNodeNames,
 	{
 		if (URigVMTemplateNode* Node = Cast<URigVMTemplateNode>(GetGraph()->FindNodeByName(NodeName)))
 		{
-			Nodes.Add(Node);
+			if (!Node->IsSingleton())
+			{
+				Nodes.Add(Node);
+			}
 		}
 	}
 
@@ -1745,10 +1771,93 @@ bool URigVMController::UnresolveTemplateNodes(const TArray<URigVMTemplateNode*>&
 	Algo::Transform(InNodes, Nodes, [](URigVMTemplateNode* InNode) -> URigVMNode* { return InNode; });
 	TArray<TPair<FString, FString>> LinkedPaths = GetLinkedPinPaths(Nodes);
 
+	// Find pins outside our set of nodes which were resolved to a type
+	TMap<URigVMPin*, FRigVMTemplateArgument::FType> PinsToResolve;
+	for (TPair<FString, FString>& LinkPath : LinkedPaths)
+	{
+		const FString& SourcePath = LinkPath.Key;
+		const FString& TargetPath = LinkPath.Value;
+
+		URigVMPin* SourcePin = GetGraph()->FindPin(SourcePath);
+		URigVMPin* TargetPin = GetGraph()->FindPin(TargetPath);
+
+		bool bSourceOutside = false;
+		URigVMNode* SourceNode = SourcePin->GetNode();
+		if (!Nodes.Contains(SourceNode))
+		{
+			bSourceOutside = true;
+			if (SourceNode->IsInjected())
+			{
+				if (URigVMNode* OutNode = SourceNode->GetTypedOuter<URigVMNode>())
+				{
+					if (Nodes.Contains(OutNode))
+					{
+						bSourceOutside = false;
+					}
+				}
+			}
+		}
+		bool bTargetOutside = false;
+		URigVMNode* TargetNode = TargetPin->GetNode();
+		if (!Nodes.Contains(TargetNode))
+		{
+			bTargetOutside = true;
+			if (TargetNode->IsInjected())
+			{
+				if (URigVMNode* OutNode = TargetNode->GetTypedOuter<URigVMNode>())
+				{
+					if (Nodes.Contains(OutNode))
+					{
+						bTargetOutside = false;
+					}
+				}
+			}
+		}
+		if (bSourceOutside && !bTargetOutside)
+		{
+			if (!SourcePin->IsWildCard())
+			{
+				PinsToResolve.Add(SourcePin, FRigVMTemplateArgument::FType(SourcePin->GetCPPType(), SourcePin->GetCPPTypeObject()));
+			}
+		}
+		else if(bTargetOutside && !bSourceOutside)
+		{
+			if (!TargetPin->IsWildCard())
+			{
+				PinsToResolve.Add(TargetPin, FRigVMTemplateArgument::FType(TargetPin->GetCPPType(), TargetPin->GetCPPTypeObject()));
+			}
+		}			
+	}
+
 	bool bChangedAnyPin = false;
+	bool bChangedFilteredPermutations = false;
 	for(URigVMTemplateNode* Node : InNodes)
 	{
-		if(Node->IsFullyUnresolved())
+		if (Node->IsSingleton())
+		{
+			continue;
+		}
+
+		if (!Node->PreferredPermutationTypes.IsEmpty())
+		{
+			if (bSetupUndoRedo)
+			{
+				ActionStack->AddAction(FRigVMSetPreferredTemplatePermutationsAction(Node, {}));
+			}
+			Node->PreferredPermutationTypes = {};
+		}
+		
+		TArray<int32> OldPermutations = Node->FilteredPermutations;
+		Node->InitializeFilteredPermutations();
+		if (bSetupUndoRedo)
+		{
+			FRigVMSetTemplateFilteredPermutationsAction FilteringAction(Node, nullptr, OldPermutations);
+			ActionStack->AddAction(FilteringAction);
+		}
+
+		bChangedFilteredPermutations = OldPermutations != Node->FilteredPermutations;			
+		
+		if(!bChangedFilteredPermutations)
 		{
 			continue;
 		}
@@ -1833,9 +1942,34 @@ bool URigVMController::UnresolveTemplateNodes(const TArray<URigVMTemplateNode*>&
 			return A->IsWildCard() == B->IsWildCard();
 		}), bSetupUndoRedo);
 	
+	for (TPair<URigVMPin*, FRigVMTemplateArgument::FType>& Pin : PinsToResolve)
+	{
+		TArray<URigVMLink*> LinksToBreak;
+		for (URigVMLink* Link : Pin.Key->GetLinks())
+		{
+			if (URigVMPin* OppositePin = Link->GetOppositePin(Pin.Key))
+			{
+				if (Nodes.Contains(OppositePin->GetNode()))
+				{
+					if (OppositePin->GetCPPType() != Pin.Value.CPPType)
+					{
+						LinksToBreak.Add(Link);
+					}
+				}
+			}
+		}
+
+		for (URigVMLink* Link : LinksToBreak)
+		{
+			BreakLink(Link->GetSourcePin(), Link->GetTargetPin(), bSetupUndoRedo);
+		}
+		
+		ResolveWildCardPin(Pin.Key, Pin.Value, bSetupUndoRedo);
+	}
+
 	if(bSetupUndoRedo)
 	{
-		if(bChangedAnyPin)
+		if(bChangedAnyPin || bChangedFilteredPermutations)
 		{
 			ActionStack->EndAction(Action);
 		}
@@ -2366,6 +2500,7 @@ URigVMRerouteNode* URigVMController::AddRerouteNodeOnPin(const FString& InPinPat
 
 	Notify(ERigVMGraphNotifType::NodeAdded, Node);
 
+	Node->InitializeFilteredPermutations();
 	if (bAsInput)
 	{
 		AddLink(ValuePin, PinForLink, bSetupUndoRedo);
@@ -3396,6 +3531,22 @@ TArray<FName> URigVMController::ImportNodesFromText(const FString& InText, bool
 				}
 			}
 
+			if (URigVMTemplateNode* TemplateNode = Cast<URigVMTemplateNode>(SubNode))
+			{
+				if (!TemplateNode->IsSingleton())
+				{
+					{
+						TGuardValue<TArray<FString>> GuardPreferredPermutation(TemplateNode->PreferredPermutationTypes, {});
+						TemplateNode->InitializeFilteredPermutations();
+					}
+					if (!TemplateNode->PreferredPermutationTypes.IsEmpty())
+					{
+						TemplateNode->FilteredPermutations = {TemplateNode->FindPermuationForTypes(TemplateNode->PreferredPermutationTypes)};
+					}
+					UpdateTemplateNodePinTypes(TemplateNode, bSetupUndoRedo);
+				}
+			}
+
 			if (URigVMFunctionReferenceNode* FunctionRefNode = Cast<URigVMFunctionReferenceNode>(SubNode))
 			{
 				if (URigVMLibraryNode* FunctionDefinition = FunctionRefNode->GetReferencedNode())
@@ -8949,7 +9100,8 @@ bool URigVMController::AddLink(URigVMPin* OutputPin, URigVMPin* InputPin, bool b
 
 	{
 		FString FailureReason;
-		if (!Graph->CanLink(OutputPin, InputPin, &FailureReason, GetCurrentByteCode(), UserLinkDirection))
+		const bool bAllowWildcard = bIsTransacting; 
+		if (!Graph->CanLink(OutputPin, InputPin, &FailureReason, GetCurrentByteCode(), UserLinkDirection, bAllowWildcard))
 		{
 			if(OutputPin->IsExecuteContext() && InputPin->IsExecuteContext())
 			{
@@ -9003,7 +9155,6 @@ bool URigVMController::AddLink(URigVMPin* OutputPin, URigVMPin* InputPin, bool b
 		if(bOutputPinCanChangeType)
 		{
 			Notify(ERigVMGraphNotifType::InteractionBracketOpened, nullptr);
-			ResolveWildCardPin(OutputPin, InputPin, bSetupUndoRedo, true, true, true, true);
 			if(OutputPin->GetNode()->IsA<URigVMRerouteNode>())
 			{
 				SetPinDefaultValue(OutputPin, InputPin->GetDefaultValue(), true, bSetupUndoRedo, false, true);
@@ -9014,13 +9165,6 @@ bool URigVMController::AddLink(URigVMPin* OutputPin, URigVMPin* InputPin, bool b
 			}
 			Notify(ERigVMGraphNotifType::InteractionBracketClosed, nullptr);
 		}
-		else if(bInputPinCanChangeType)
-		{
-			Notify(ERigVMGraphNotifType::InteractionBracketOpened, nullptr);
-			ResolveWildCardPin(InputPin, OutputPin, bSetupUndoRedo, true, true, true, true);
-			Notify(ERigVMGraphNotifType::InteractionBracketClosed, nullptr);
-		}
-
 	}
 
 	if (bSetupUndoRedo)
@@ -9029,6 +9173,45 @@ bool URigVMController::AddLink(URigVMPin* OutputPin, URigVMPin* InputPin, bool b
 		ExpandPinRecursively(InputPin->GetParentPin(), bSetupUndoRedo);
 	}
 
+	// Before adding the link, let's resolve input and ouput pin types
+	// If templates, we will filter the permutations that support this link
+	// If any links need to be broken before perfoming this connection, try to find them and break them
+	if (!bIsTransacting)
+	{
+		URigVMPin* FirstToResolve = (InUserDirection == ERigVMPinDirection::Input) ? OutputPin : InputPin;
+		URigVMPin* SecondToResolve = (FirstToResolve == OutputPin) ? InputPin : OutputPin;
+		if (URigVMTemplateNode* FirstTemplateNode = Cast<URigVMTemplateNode>(FirstToResolve->GetNode()))
+		{
+			if (!FirstTemplateNode->IsSingleton())
+			{
+				TArray<FRigVMTemplateArgument::FType> InputTypes = GetWildcardFilteredTypes(SecondToResolve);
+				if (InputTypes.Num() > 0)
+				{
+					if (!PrepareTemplatePinForType(FirstToResolve, InputTypes, bSetupUndoRedo))
+					{
+						ActionStack->CancelAction(Action, this);
+						return false;
+					}
+				}
+			}
+		}
+		if (URigVMTemplateNode* SecondTemplateNode = Cast<URigVMTemplateNode>(SecondToResolve->GetNode()))
+		{
+			if (!SecondTemplateNode->IsSingleton())
+			{
+				TArray<FRigVMTemplateArgument::FType> OutTypes = GetWildcardFilteredTypes(FirstToResolve);
+				if (OutTypes.Num() > 0)
+				{
+					if (!PrepareTemplatePinForType(SecondToResolve, OutTypes, bSetupUndoRedo))
+					{
+						ActionStack->CancelAction(Action, this);
+						return false;
+					}
+				}
+			}
+		}	
+	}
+
 	if (bSetupUndoRedo)
 	{
 		ActionStack->AddAction(FRigVMAddLinkAction(OutputPin, InputPin));
@@ -9209,6 +9392,11 @@ bool URigVMController::BreakLink(URigVMPin* OutputPin, URigVMPin* InputPin, bool
 			OutputPin->Links.Remove(Link);
 			InputPin->Links.Remove(Link);
 			Graph->Links.Remove(Link);
+
+			if (!bIsTransacting && !bSuspendRecomputingTemplateFilters)
+			{
+				RecomputeAllTemplateFilteredTypes(bSetupUndoRedo);
+			}
 			
 			if (!bSuspendNotifications)
 			{
@@ -9299,26 +9487,34 @@ bool URigVMController::BreakAllLinks(URigVMPin* Pin, bool bAsInput, bool bSetupU
 	}
 
 	int32 LinksBroken = 0;
-	if (Pin->IsBoundToVariable() && bAsInput && bSetupUndoRedo)
-	{
-		UnbindPinFromVariable(Pin, bSetupUndoRedo);
-		LinksBroken++;
-	}
-
-	TArray<URigVMLink*> Links = Pin->GetLinks();
-	for (int32 LinkIndex = Links.Num() - 1; LinkIndex >= 0; LinkIndex--)
 	{
-		URigVMLink* Link = Links[LinkIndex];
-		if (bAsInput && Link->GetTargetPin() == Pin)
+		TGuardValue<bool> GuardSuspendTemplateRecompute(bSuspendRecomputingTemplateFilters, true);
+		if (Pin->IsBoundToVariable() && bAsInput && bSetupUndoRedo)
 		{
-			LinksBroken += BreakLink(Link->GetSourcePin(), Pin, bSetupUndoRedo) ? 1 : 0;
+			UnbindPinFromVariable(Pin, bSetupUndoRedo);
+			LinksBroken++;
 		}
-		else if (!bAsInput && Link->GetSourcePin() == Pin)
+
+		TArray<URigVMLink*> Links = Pin->GetLinks();
+		for (int32 LinkIndex = Links.Num() - 1; LinkIndex >= 0; LinkIndex--)
 		{
-			LinksBroken += BreakLink(Pin, Link->GetTargetPin(), bSetupUndoRedo) ? 1 : 0;
+			URigVMLink* Link = Links[LinkIndex];
+			if (bAsInput && Link->GetTargetPin() == Pin)
+			{
+				LinksBroken += BreakLink(Link->GetSourcePin(), Pin, bSetupUndoRedo) ? 1 : 0;
+			}
+			else if (!bAsInput && Link->GetSourcePin() == Pin)
+			{
+				LinksBroken += BreakLink(Pin, Link->GetTargetPin(), bSetupUndoRedo) ? 1 : 0;
+			}
 		}
 	}
 
+	if (LinksBroken > 0 && !bSuspendRecomputingTemplateFilters)
+	{
+		RecomputeAllTemplateFilteredTypes(bSetupUndoRedo);
+	}
+
 	if (bSetupUndoRedo)
 	{
 		if (LinksBroken > 0)
@@ -9334,25 +9530,36 @@ bool URigVMController::BreakAllLinks(URigVMPin* Pin, bool bAsInput, bool bSetupU
 	return LinksBroken > 0;
 }
 
-void URigVMController::BreakAllLinksRecursive(URigVMPin* Pin, bool bAsInput, bool bTowardsParent, bool bSetupUndoRedo)
+bool URigVMController::BreakAllLinksRecursive(URigVMPin* Pin, bool bAsInput, bool bTowardsParent, bool bSetupUndoRedo)
 {
-	if (bTowardsParent)
+	bool bBrokenLinks = false;
 	{
-		URigVMPin* ParentPin = Pin->GetParentPin();
-		if (ParentPin)
+		TGuardValue<bool> GuardSuspendRecomputeTemplates(bSuspendRecomputingTemplateFilters, true);
+		if (bTowardsParent)
 		{
-			BreakAllLinks(ParentPin, bAsInput, bSetupUndoRedo);
-			BreakAllLinksRecursive(ParentPin, bAsInput, bTowardsParent, bSetupUndoRedo);
+			URigVMPin* ParentPin = Pin->GetParentPin();
+			if (ParentPin)
+			{
+				bBrokenLinks |= BreakAllLinks(ParentPin, bAsInput, bSetupUndoRedo);
+				bBrokenLinks |= BreakAllLinksRecursive(ParentPin, bAsInput, bTowardsParent, bSetupUndoRedo);
+			}
 		}
-	}
-	else
-	{
-		for (URigVMPin* SubPin : Pin->SubPins)
+		else
 		{
-			BreakAllLinks(SubPin, bAsInput, bSetupUndoRedo);
-			BreakAllLinksRecursive(SubPin, bAsInput, bTowardsParent, bSetupUndoRedo);
+			for (URigVMPin* SubPin : Pin->SubPins)
+			{
+				bBrokenLinks |= BreakAllLinks(SubPin, bAsInput, bSetupUndoRedo);
+				bBrokenLinks |= BreakAllLinksRecursive(SubPin, bAsInput, bTowardsParent, bSetupUndoRedo);
+			}
 		}
 	}
+
+	if (bBrokenLinks && !bSuspendRecomputingTemplateFilters)
+	{
+		RecomputeAllTemplateFilteredTypes(bSetupUndoRedo);
+	}
+	
+	return bBrokenLinks;
 }
 
 FName URigVMController::AddExposedPin(const FName& InPinName, ERigVMPinDirection InDirection, const FString& InCPPType, const FName& InCPPTypeObjectPath, const FString& InDefaultValue, bool bSetupUndoRedo, bool bPrintPythonCommand)
@@ -11372,6 +11579,22 @@ URigVMRerouteNode* URigVMController::AddFreeRerouteNode(bool bShowAsFullNode, co
 
 	Notify(ERigVMGraphNotifType::NodeAdded, Node);
 
+	Node->InitializeFilteredPermutations();
+	if (InCPPType != RigVMTypeUtils::GetWildCardCPPType() && InCPPType != RigVMTypeUtils::GetWildCardArrayCPPType())
+	{
+		PrepareTemplatePinForType(ValuePin, {FRigVMTemplateArgument::FType(ValuePin->GetCPPType(), ValuePin->GetCPPTypeObject())}, bSetupUndoRedo);
+		TArray<int32> FilterPermutations = Node->GetFilteredPermutationsIndices();
+		if (FilterPermutations.Num() == 1)
+		{
+			const TArray<FString> NewPreferredPermutationTypes = Node->GetArgumentTypesForPermutation(FilterPermutations[0]);
+			if (bSetupUndoRedo)
+			{
+				ActionStack->AddAction(FRigVMSetPreferredTemplatePermutationsAction(Node, NewPreferredPermutationTypes));
+			}
+			Node->PreferredPermutationTypes = NewPreferredPermutationTypes;
+		}
+	}
+
 	if (bSetupUndoRedo)
 	{
 		ActionStack->EndAction(Action);
@@ -11547,6 +11770,22 @@ URigVMIfNode* URigVMController::AddIfNode(const FString& InCPPType, const FName&
 	{
 		ActionStack->AddAction(FRigVMAddIfNodeAction(Node));
 	}
+
+	Node->InitializeFilteredPermutations();
+	if (InCPPType != RigVMTypeUtils::GetWildCardCPPType() && InCPPType != RigVMTypeUtils::GetWildCardArrayCPPType())
+	{
+		PrepareTemplatePinForType(TruePin, {FRigVMTemplateArgument::FType(TruePin->GetCPPType(), TruePin->GetCPPTypeObject())}, bSetupUndoRedo);
+		TArray<int32> FilterPermutations = Node->GetFilteredPermutationsIndices();
+		if (FilterPermutations.Num() == 1)
+		{
+			const TArray<FString> NewPreferredPermutationTypes = Node->GetArgumentTypesForPermutation(FilterPermutations[0]);
+			if (bSetupUndoRedo)
+			{
+				ActionStack->AddAction(FRigVMSetPreferredTemplatePermutationsAction(Node, NewPreferredPermutationTypes));
+			}
+			Node->PreferredPermutationTypes = NewPreferredPermutationTypes;
+		}
+	}
 	
 	if (bPrintPythonCommand)
 	{
@@ -11653,6 +11892,22 @@ URigVMSelectNode* URigVMController::AddSelectNode(const FString& InCPPType, cons
 		ActionStack->AddAction(FRigVMAddSelectNodeAction(Node));
 	}
 
+	Node->InitializeFilteredPermutations();
+	if (InCPPType != RigVMTypeUtils::GetWildCardCPPType() && InCPPType != RigVMTypeUtils::GetWildCardArrayCPPType())
+	{
+		PrepareTemplatePinForType(ResultPin, {FRigVMTemplateArgument::FType(ResultPin->GetCPPType(), ResultPin->GetCPPTypeObject())}, bSetupUndoRedo);
+		TArray<int32> FilterPermutations = Node->GetFilteredPermutationsIndices();
+		if (FilterPermutations.Num() == 1)
+		{
+			const TArray<FString> NewPreferredPermutationTypes = Node->GetArgumentTypesForPermutation(FilterPermutations[0]);
+			if (bSetupUndoRedo)
+			{
+				ActionStack->AddAction(FRigVMSetPreferredTemplatePermutationsAction(Node, NewPreferredPermutationTypes));
+			}
+			Node->PreferredPermutationTypes = NewPreferredPermutationTypes;
+		}
+	}
+	
 	if (bPrintPythonCommand)
 	{
 		TArray<FString> Commands = GetAddNodePythonCommands(Node);
@@ -11735,7 +11990,8 @@ URigVMTemplateNode* URigVMController::AddTemplateNode(const FName& InNotation, c
 	int32 PermutationIndex = INDEX_NONE;
 	FRigVMTemplate::FTypeMap Types;
 	Template->FullyResolve(Types, PermutationIndex);
-
+	Node->InitializeFilteredPermutations();
+	
 	for (int32 ArgIndex = 0; ArgIndex < Template->NumArguments(); ArgIndex++)
 	{
 		const FRigVMTemplateArgument* Arg = Template->GetArgument(ArgIndex);
@@ -12218,6 +12474,23 @@ URigVMArrayNode* URigVMController::AddArrayNode(ERigVMOpCode InOpCode, const FSt
 		ActionStack->EndAction(Action);
 	}
 
+	Node->InitializeFilteredPermutations();
+	if (InCPPType != RigVMTypeUtils::GetWildCardCPPType() && InCPPType != RigVMTypeUtils::GetWildCardArrayCPPType())
+	{
+		URigVMPin* ArrayPin = Node->FindPin(URigVMArrayNode::ArrayName);
+		PrepareTemplatePinForType(ArrayPin, {FRigVMTemplateArgument::FType(ArrayPin->GetCPPType(), ArrayPin->GetCPPTypeObject())}, bSetupUndoRedo);
+		TArray<int32> FilterPermutations = Node->GetFilteredPermutationsIndices();
+		if (FilterPermutations.Num() == 1)
+		{
+			const TArray<FString> NewPreferredPermutationTypes = Node->GetArgumentTypesForPermutation(FilterPermutations[0]);
+			if (bSetupUndoRedo)
+			{
+				ActionStack->AddAction(FRigVMSetPreferredTemplatePermutationsAction(Node, NewPreferredPermutationTypes));
+			}
+			Node->PreferredPermutationTypes = NewPreferredPermutationTypes;
+		}
+	}
+
 	if (bPrintPythonCommand)
 	{
 		TArray<FString> Commands = GetAddNodePythonCommands(Node);
@@ -14467,86 +14740,6 @@ void URigVMController::PostProcessDefaultValue(URigVMPin* Pin, FString& OutDefau
 	}
 }
 
-void URigVMController::ResolveTemplateNodePins(URigVMTemplateNode* InNode, bool bSetupUndoRedo)
-{
-	if (InNode == nullptr)
-	{
-		return;
-	}
-
-	InNode->ResolvedFunctionName.Reset();
-
-	// check if the node is resolved
-	TArray<int32> ResolvedPermutationIndices = InNode->GetResolvedPermutationIndices();
-	if(ResolvedPermutationIndices.Num() > 0)
-	{
-		TArray<URigVMPin*> ResolvedPins;
-		
-		for (URigVMPin* Pin : InNode->GetPins())
-		{
-			if(!Pin->IsWildCard())
-			{
-				continue;
-			}
-
-			const FRigVMTemplateArgument* Arg = InNode->GetTemplate()->FindArgument(Pin->GetFName());
-			if(Arg == nullptr)
-			{
-				continue;
-			}
-
-			if(Arg->IsSingleton(ResolvedPermutationIndices))
-			{
-				const FRigVMTemplateArgument::FType ResolvedType = Arg->GetTypes()[ResolvedPermutationIndices[0]];
-				if(ChangePinType(Pin, ResolvedType.CPPType, ResolvedType.CPPTypeObject, bSetupUndoRedo, false, false, false))
-				{
-					ResolvedPins.Add(Pin);
-					
-					if(Pin->IsArray())
-					{
-						const FString BaseCPPType = RigVMTypeUtils::BaseTypeFromArrayType(ResolvedType.CPPType);
-						for(int32 SubPinIndex = 0; SubPinIndex < Pin->GetSubPins().Num(); SubPinIndex++)
-						{
-							URigVMPin* SubPin = Pin->GetSubPins()[SubPinIndex];
-							if(SubPin->IsWildCard())
-							{
-								ResolveWildCardPinImpl(SubPin, BaseCPPType, ResolvedType.CPPTypeObject, bSetupUndoRedo, false, false, true);
-							}
-						}
-					}
-
-					const FString NewDefaultValue = InNode->GetInitialDefaultValueForPin(Pin->GetFName(), ResolvedPermutationIndices);
-					SetPinDefaultValue(Pin, NewDefaultValue, true, bSetupUndoRedo, false, true);
-				}
-			}
-		}
-
-		if(ResolvedPermutationIndices.Num() == 1)
-		{
-			FullyResolveTemplateNode(InNode, ResolvedPermutationIndices[0], bSetupUndoRedo);
-		}
-
-		// resolve all wires on resolved pins
-		for(URigVMPin* ResolvedPin : ResolvedPins)
-		{
-			// we don't need to recurse since unknown type pins don't have sub pins
-			TArray<URigVMPin*> LinkedPins = ResolvedPin->GetLinkedSourcePins(false); 
-			LinkedPins.Append(ResolvedPin->GetLinkedTargetPins(false));
-
-			for(URigVMPin* LinkedPin : LinkedPins)
-			{
-				if(LinkedPin->IsWildCard())
-				{
-					ResolveWildCardPin(LinkedPin, ResolvedPin, bSetupUndoRedo);
-				}
-			}
-
-		}
-	}
-
-	ResolveTemplateNodeMetaData(InNode, bSetupUndoRedo);
-}
-
 void URigVMController::ResolveTemplateNodeMetaData(URigVMTemplateNode* InNode, bool bSetupUndoRedo)
 {
 #if WITH_EDITOR
@@ -14554,8 +14747,8 @@ void URigVMController::ResolveTemplateNodeMetaData(URigVMTemplateNode* InNode, b
 
 	const FLinearColor PreviousColor = InNode->NodeColor;
 	
-	const TArray<int32> ResolvedPermutationIndices = InNode->GetResolvedPermutationIndices();
-	InNode->NodeColor = InNode->GetTemplate()->GetColor(ResolvedPermutationIndices);
+	const TArray<int32> FilteredPermutationIndices = InNode->GetFilteredPermutationsIndices();
+	InNode->NodeColor = InNode->GetTemplate()->GetColor(FilteredPermutationIndices);
 
 	if(!InNode->NodeColor.Equals(PreviousColor, 0.01f))
 	{
@@ -14583,7 +14776,7 @@ void URigVMController::ResolveTemplateNodeMetaData(URigVMTemplateNode* InNode, b
 			}
 			if(!Pin->IsValidDefaultValue(Pin->GetDefaultValue()))
 			{
-				const FString NewDefaultValue = InNode->GetInitialDefaultValueForPin(Pin->GetFName(), ResolvedPermutationIndices);
+				const FString NewDefaultValue = InNode->GetInitialDefaultValueForPin(Pin->GetFName(), FilteredPermutationIndices);
 				SetPinDefaultValue(Pin, NewDefaultValue, true, bSetupUndoRedo, false, true);
 			}
 		}
@@ -14604,6 +14797,7 @@ bool URigVMController::FullyResolveTemplateNode(URigVMTemplateNode* InNode, int3
 	const FRigVMTemplate* Template = InNode->GetTemplate();
 	const FRigVMFunction* ResolvedFunction = Template->GetPermutation(InPermutationIndex);
 	const TArray<int32> PermutationIndices = {InPermutationIndex};
+	InNode->FilteredPermutations = PermutationIndices;
 
 	// find all existing pins that we may need to change
 	TArray<FRigVMTemplateArgument> MissingPins;
@@ -14680,6 +14874,7 @@ bool URigVMController::FullyResolveTemplateNode(URigVMTemplateNode* InNode, int3
 	}
 
 	// update the incorrectly typed pins
+	bool bNeedsTemplateUpdate = false;
 	for(const TPair<URigVMPin*, FRigVMTemplateArgument::FType>& Pair : PinTypesToChange)
 	{
 		URigVMPin* Pin = Pair.Key;
@@ -14690,6 +14885,7 @@ bool URigVMController::FullyResolveTemplateNode(URigVMTemplateNode* InNode, int3
 			if((Pin->GetCPPType() != ExpectedType.CPPType) ||
 				(Pin->GetCPPTypeObject() != ExpectedType.CPPTypeObject))
 			{
+				bNeedsTemplateUpdate = true;
 				const FString CPPType = Pin->IsArray() ? RigVMTypeUtils::GetWildCardArrayCPPType() : RigVMTypeUtils::GetWildCardCPPType();
 				if(!ChangePinType(Pin, CPPType, RigVMTypeUtils::GetWildCardCPPTypeObject(), bSetupUndoRedo, false, true, true))
 				{
@@ -14704,7 +14900,9 @@ bool URigVMController::FullyResolveTemplateNode(URigVMTemplateNode* InNode, int3
 		
 		if(Pin->IsWildCard())
 		{
-			if(!ResolveWildCardPinImpl(Pin, ExpectedType.CPPType, ExpectedType.CPPTypeObject, bSetupUndoRedo, false, true, true))
+			bNeedsTemplateUpdate = true;
+			
+			if (!UpdateFilteredPermutations(Pin, {ExpectedType}, bSetupUndoRedo))
 			{
 				if(bSetupUndoRedo)
 				{
@@ -14741,17 +14939,9 @@ bool URigVMController::FullyResolveTemplateNode(URigVMTemplateNode* InNode, int3
 		}
 	}
 
-	if (UnitNodeCreatedContext.IsValid())
+	if (bNeedsTemplateUpdate)
 	{
-		if(URigVMUnitNode* UnitNode = Cast<URigVMUnitNode>(InNode))
-		{
-			if (TSharedPtr<FStructOnScope> StructScope = UnitNode->ConstructStructInstance())
-			{
-				TGuardValue<FName> NodeNameScope(UnitNodeCreatedContext.NodeName, UnitNode->GetFName());
-				FRigVMStruct* StructInstance = (FRigVMStruct*)StructScope->GetStructMemory();
-				StructInstance->OnUnitNodeCreated(UnitNodeCreatedContext);
-			}
-		}
+		UpdateTemplateNodePinTypes(InNode, bSetupUndoRedo);
 	}
 	
 	if(bSetupUndoRedo)
@@ -14762,7 +14952,7 @@ bool URigVMController::FullyResolveTemplateNode(URigVMTemplateNode* InNode, int3
 	return true;
 }
 
-bool URigVMController::ResolveWildCardPin(const FString& InPinPath, const FString& InCPPType, const FName& InCPPTypeObjectPath, bool bSetupUndoRedo, bool bPrintPythonCommand)
+bool URigVMController::PrepareTemplatePinForType(URigVMPin* InPin, const TArray<FRigVMTemplateArgument::FType>& InTypes, bool bSetupUndoRedo)
 {
 	if (!IsValidGraph())
 	{
@@ -14774,356 +14964,619 @@ bool URigVMController::ResolveWildCardPin(const FString& InPinPath, const FStrin
 		return false;
 	}
 
-	URigVMGraph* Graph = GetGraph();
-	check(Graph);
+	URigVMTemplateNode* TemplateNode = Cast<URigVMTemplateNode>(InPin->GetNode());
+	if (!TemplateNode)
+	{
+		return false;
+	}
 
-	UObject* CPPTypeObject = nullptr;
-	if (!InCPPTypeObjectPath.IsNone())
+	if (TemplateNode->IsSingleton())
 	{
-		CPPTypeObject = URigVMPin::FindObjectFromCPPTypeObjectPath<UObject>(InCPPTypeObjectPath.ToString());
-		if (CPPTypeObject == nullptr)
+		return false;
+	}
+
+	// If the pin is an element of an array, get the array pin and convert the types
+	URigVMPin* RootPin = InPin;
+	TArray<FRigVMTemplateArgument::FType> Types = InTypes;
+	if (InPin->IsArrayElement())
+	{
+		RootPin = InPin->GetParentPin();
+		for (FRigVMTemplateArgument::FType& Type : Types)
 		{
-			ReportErrorf(TEXT("Cannot find cpp type object for path '%s'."), *InCPPTypeObjectPath.ToString());
-			return false;
+			Type.ConvertToArray();
 		}
 	}
 
-	const FString CPPType = RigVMTypeUtils::PostProcessCPPType(InCPPType, CPPTypeObject);
+	// It might be a subpin of a struct, in that case we just want to make sure
+	// the type required matches the one in the pin
+	if (!TemplateNode->GetTemplate()->FindArgument(RootPin->GetFName()))
+	{
+		if (Types.Num() == 1 && Types[0].Matches(RootPin->GetCPPType()))
+		{
+			return true;
+		}
+		return false;
+	}
 
-	if (URigVMPin* Pin = Graph->FindPin(InPinPath))
+	bool bFilteredSupportsType = false;
+	for (const FRigVMTemplateArgument::FType& Type : Types)
 	{
-		if(Pin->IsWildCard() || Pin->GetNode()->IsA<URigVMTemplateNode>())
+		if (TemplateNode->FilteredSupportsType(RootPin, Type.CPPType))
 		{
-			FRigVMBaseAction Action;
-			if (bSetupUndoRedo)
-			{
-				Action.Title = TEXT("Resolve Wildcard Pin");
-				ActionStack->BeginAction(Action);
-			}
-			
-			if(ResolveWildCardPinImpl(Pin, CPPType, CPPTypeObject, bSetupUndoRedo))
+			bFilteredSupportsType = true;
+			break;
+		}
+	}
+
+	bool bSupportsType = false;
+	if (!bFilteredSupportsType)
+	{
+		if (!TemplateNode->PreferredPermutationTypes.IsEmpty())
+		{
+			if (bSetupUndoRedo)
 			{
-				if(bPrintPythonCommand)
-				{
-					const FString GraphName = GetSanitizedGraphName(GetGraph()->GetGraphName());
-					const FString SanitizedPinPath = GetSanitizedPinPath(Pin->GetPinPath());
-					const FString CPPObjectPath = CPPTypeObject ? FString::Printf(TEXT("'%s'"), *CPPTypeObject->GetPathName()) : TEXT("None");
+				ActionStack->AddAction(FRigVMSetPreferredTemplatePermutationsAction(TemplateNode, {}));
+			}
+			TemplateNode->PreferredPermutationTypes = {};
+		}
+		
+		for (const FRigVMTemplateArgument::FType& Type : Types)
+		{
+			if (TemplateNode->SupportsType(RootPin, Type.CPPType))
+			{
+				bSupportsType = true;
+				break;
+			}
+		}	
+	}
 
-					RigVMPythonUtils::Print(GetGraphOuterName(), 
-						FString::Printf(TEXT("blueprint.get_controller_by_name('%s').resolve_wild_card_pin('%s', '%s', %s)"),
-									*GraphName,
-									*SanitizedPinPath,
-									*CPPType,
-									*CPPObjectPath));
-				}
+	if (!bFilteredSupportsType && !bSupportsType)
+	{
+		return false;
+	}
 
-				if (bSetupUndoRedo)
+	if (bFilteredSupportsType)
+	{
+		if (TemplateNode->PinNeedsFilteredTypesUpdate(RootPin, Types))
+		{
+			if (UpdateFilteredPermutations(RootPin, Types, bSetupUndoRedo))
+			{
+				UpdateTemplateNodePinTypes(TemplateNode, bSetupUndoRedo);
+				PropagateTemplateFilteredTypes(TemplateNode, bSetupUndoRedo);
+			}
+		}
+		return true;
+	}
+	else
+	{
+		// Figure out what links we would need to break to allow this
+		TArray<URigVMLink*> LinksToBreak;
+		{
+			TGuardValue<FRigVMController_RequestBreakLinksDialogDelegate> GuardDelegate(RequestBreakLinksDialogDelegate, nullptr);
+			TGuardValue<bool> GuardNotifications(bSuspendNotifications, true);
+			TGuardValue<bool> SuspendRecomputeTemplates(bSuspendRecomputingTemplateFilters, true);
+			OpenUndoBracket(FString::Printf(TEXT("Resolve wildcard pin %s"), *InPin->GetPinPath()));
+
+			bool bBrokenLinks = false;
+			do
+			{
+				bBrokenLinks = false;
+				
+				// Initialize all template nodes in the graph
+				InitializeAllTemplateFiltersInGraph(true);
+
+				// Update our node with the requested type
+				UpdateFilteredPermutations(InPin, InTypes, true);
+
+				// Propagate the filtered types and collect break link actions 
+				bBrokenLinks = !PropagateTemplateFilteredTypes(TemplateNode, true);
+			}
+			while (bBrokenLinks);
+
+			// Find break link actions
+			TArray<TPair<FString, FString>> InconsistentLinks;
+			TArray<FRigVMActionKey> Actions = ActionStack->BracketActions.Last()->SubActions;
+			for (int32 i=0; i<Actions.Num(); ++i)
+			{
+				const FRigVMActionKey& ActionKey = Actions[i];
+				FRigVMActionWrapper Wrapper(ActionKey);
+				if (Wrapper.GetAction()->GetScriptStruct() == FRigVMBreakLinkAction::StaticStruct())
 				{
-					ActionStack->EndAction(Action);
+					const FRigVMBreakLinkAction* BreakLinkAction = (const FRigVMBreakLinkAction*)Wrapper.GetAction();
+					InconsistentLinks.AddUnique(TPair<FString,FString>(BreakLinkAction->OutputPinPath, BreakLinkAction->InputPinPath));
 				}
-				return true;
+
+				Actions.Append(Wrapper.GetAction()->SubActions);
 			}
+			CancelUndoBracket();
 
-			if (bSetupUndoRedo)
+			for (TPair<FString, FString>& Pair : InconsistentLinks)
 			{
-				ActionStack->CancelAction(Action, this);
+				LinksToBreak.AddUnique(GetGraph()->FindLink(FString::Printf(TEXT("%s -> %s"), *Pair.Key, *Pair.Value)));
+			}
+		}
+
+		// We might be in this situation just because we are resolving to a different type than previosly resolved
+		// so no links to break
+		if (LinksToBreak.IsEmpty())
+		{
+			if (!TemplateNode->PreferredPermutationTypes.IsEmpty())
+			{
+				return false;
+			}
+			
+			UnresolveTemplateNodes({TemplateNode}, bSetupUndoRedo);
+			PrepareTemplatePinForType(InPin, InTypes, bSetupUndoRedo);
+			return true;
+		}
+
+		// Warn the user the links that would be broken
+		bool bBreakLinks = true;
+		bool bConsultedBreak = false;
+		if(!bIsTransacting && RequestBreakLinksDialogDelegate.IsBound())
+		{
+			bBreakLinks = RequestBreakLinksDialogDelegate.Execute(LinksToBreak);
+			bConsultedBreak = true;
+		}
+
+		if (bBreakLinks)
+		{
+			// Break links
+			for (URigVMLink* Link : LinksToBreak)
+			{
+				if (!bConsultedBreak && !bIsTransacting && !bSuspendNotifications)
+				{
+					ReportWarningf(TEXT("The link between %s was broken due to inconsistent types"), *Link->GetPinPathRepresentation());
+				}
+				BreakLink(Link->GetSourcePin(), Link->GetTargetPin(), bSetupUndoRedo);
 			}
+
+			ensure(LinksToBreak.Num() > 0);
+			PrepareTemplatePinForType(InPin, InTypes, bSetupUndoRedo);
+			return true;
 		}
 	}
 
 	return false;
 }
 
-bool URigVMController::ResolveWildCardPin(URigVMPin* InPinToResolve, URigVMPin* InTemplatePin,
-										  bool bSetupUndoRedo, bool bTraverseNode, bool bTraverseParentPins,
-										  bool bTraverseLinks, bool bAllowFloatingPointCasts)
+TArray<FRigVMTemplateArgument::FType> URigVMController::GetWildcardFilteredTypes(URigVMPin* InPin)
 {
-	check(InTemplatePin);
-	check(!InTemplatePin->IsWildCard());
-	check(InPinToResolve != InTemplatePin);
-	check(InPinToResolve->GetRootPin() != InTemplatePin->GetRootPin());
+	TArray<FRigVMTemplateArgument::FType> Types;
+	if (!InPin->IsWildCard())
+	{
+		Types.Add(FRigVMTemplateArgument::FType(InPin->GetCPPType(), InPin->GetCPPTypeObject()));
+		return Types;
+	}
 
-	FRigVMBaseAction Action;
-	if (bSetupUndoRedo)
+	if (URigVMTemplateNode* TemplateNode = Cast<URigVMTemplateNode>(InPin->GetNode()))
 	{
-		Action.Title = TEXT("Resolve Wildcard Pin");
-		ActionStack->BeginAction(Action);
+		if (!TemplateNode->IsSingleton())
+		{
+			Types = TemplateNode->GetFilteredTypesForPin(InPin);
+		}
+	}
+	
+	return Types;
+}
+
+bool URigVMController::ResolveWildCardPin(const FString& InPinPath, const FString& InCPPType, const FName& InCPPTypeObjectPath, bool bSetupUndoRedo, bool bPrintPythonCommand)
+{
+	if (!IsValidGraph())
+	{
+		return false;
+	}
+
+	if (!bIsTransacting && !IsGraphEditable())
+	{
+		return false;
 	}
 
-	const bool bResult = ResolveWildCardPinImpl(InPinToResolve, InTemplatePin->GetCPPType(), InTemplatePin->GetCPPTypeObject(), bSetupUndoRedo, bTraverseNode, bTraverseParentPins, bTraverseLinks, bAllowFloatingPointCasts);
-	if(!bResult)
+	URigVMGraph* Graph = GetGraph();
+	check(Graph);
+
+	UObject* CPPTypeObject = nullptr;
+	if (!InCPPTypeObjectPath.IsNone())
 	{
-		if(InPinToResolve->IsLinkedTo(InTemplatePin))
+		CPPTypeObject = URigVMPin::FindObjectFromCPPTypeObjectPath<UObject>(InCPPTypeObjectPath.ToString());
+		if (CPPTypeObject == nullptr)
 		{
-			BreakLink(InPinToResolve, InTemplatePin, bSetupUndoRedo);
-			BreakLink(InTemplatePin, InPinToResolve, bSetupUndoRedo);
+			ReportErrorf(TEXT("Cannot find cpp type object for path '%s'."), *InCPPTypeObjectPath.ToString());
+			return false;
 		}
 	}
 
-	if (bSetupUndoRedo)
+	const FString CPPType = RigVMTypeUtils::PostProcessCPPType(InCPPType, CPPTypeObject);
+
+	if (URigVMPin* Pin = Graph->FindPin(InPinPath))
 	{
-		ActionStack->EndAction(Action);
+		return ResolveWildCardPin(Pin, FRigVMTemplateArgument::FType(CPPType, CPPTypeObject), bSetupUndoRedo, bPrintPythonCommand);
 	}
-	
-	return bResult;
+
+	return false;
 }
 
-bool URigVMController::ResolveWildCardPinImpl(URigVMPin* InPinToResolve, const FString& InCPPType, UObject* InCPPTypeObject,
-                                             bool bSetupUndoRedo, bool bTraverseNode, bool bTraverseParentPins,
-                                             bool bTraverseLinks, bool bAllowFloatingPointCasts)
+bool URigVMController::ResolveWildCardPin(URigVMPin* InPin, const FRigVMTemplateArgument::FType& InType, bool bSetupUndoRedo, bool bPrintPythonCommand)
 {
-	check(InPinToResolve);
-	check(InPinToResolve->IsWildCard() || InPinToResolve->GetNode()->IsA<URigVMTemplateNode>());
+	if (!IsValidGraph())
+	{
+		return false;
+	}
 
-	// exit early if we already have the right type
-	const FRigVMTemplateArgument::FType ArgumentType(InCPPType, InCPPTypeObject);
-	if(ArgumentType.Matches(InPinToResolve->GetCPPType(), bAllowFloatingPointCasts))
+	if (!bIsTransacting && !IsGraphEditable())
 	{
 		return false;
 	}
 
-	// if we are trying to resolve a pin on the template node...
-	FString CPPType = InCPPType;
-	UObject* CPPTypeObject = InCPPTypeObject;
-	
-	if(URigVMTemplateNode* TemplateNode = Cast<URigVMTemplateNode>(InPinToResolve->GetNode()))
+	if(InPin->IsWildCard())
 	{
-		URigVMPin* RootPin = InPinToResolve->GetRootPin();
-		if (RootPin != InPinToResolve)
+		ensure(InPin->GetNode()->IsA<URigVMTemplateNode>());
+		
+		FRigVMBaseAction Action;
+		if (bSetupUndoRedo)
 		{
-			const FString BaseCPPType = RigVMTypeUtils::IsArrayType(CPPType) ? RigVMTypeUtils::BaseTypeFromArrayType(CPPType) : CPPType;
-			const FString TypeToResolveAgainst = InPinToResolve->GetRootPin()->IsArray() ? RigVMTypeUtils::ArrayTypeFromBaseType(BaseCPPType) : BaseCPPType; 
-			if(!TemplateNode->SupportsType(InPinToResolve->GetRootPin(), TypeToResolveAgainst, &CPPType))
-			{
-				return false;
-			}			
+			Action.Title = TEXT("Resolve Wildcard Pin");
+			ActionStack->BeginAction(Action);
 		}
-		else
+		
+		URigVMTemplateNode* TemplateNode = CastChecked<URigVMTemplateNode>(InPin->GetNode());
+		if (!PrepareTemplatePinForType(InPin, {InType}, bSetupUndoRedo))
 		{
-			if(!TemplateNode->SupportsType(InPinToResolve->GetRootPin(), CPPType, &CPPType))
+			if (bSetupUndoRedo)
 			{
-				return false;
+				ActionStack->CancelAction(Action, this);
 			}
+			return false;
 		}
-	}
-
-	// if the root pin is not a wildcard - we will need to unresolve first but maintain all links.
-	if(!InPinToResolve->IsWildCard())
-	{
-		URigVMTemplateNode* TemplateNode = CastChecked<URigVMTemplateNode>(InPinToResolve->GetNode());
-		TArray<TPair<FString, FString>> LinkedPaths = GetLinkedPinPaths(TemplateNode);
-		const FRigVMTemplate::FTypeMap OriginalTypes = TemplateNode->GetResolvedTypes();
-		const TMap<FString, FPinState> PinStates = GetPinStates(TemplateNode);
 
-		bool bUnresolveResult = false;
+		const TArray<int32>& FilteredPermutations = TemplateNode->GetFilteredPermutationsIndices();
+		if (FilteredPermutations.Num() == 1)
 		{
-			TGuardValue<bool> NotifGuard(bSuspendNotifications, true);
-			bUnresolveResult = UnresolveTemplateNodes({TemplateNode}, bSetupUndoRedo);
+			const TArray<FString> NewPreferredPermutationTypes = TemplateNode->GetArgumentTypesForPermutation(FilteredPermutations[0]);
+			if (bSetupUndoRedo)
+			{
+				ActionStack->AddAction(FRigVMSetPreferredTemplatePermutationsAction(TemplateNode, NewPreferredPermutationTypes));
+			}
+			TemplateNode->PreferredPermutationTypes = NewPreferredPermutationTypes;
 		}
-
-		if(bUnresolveResult)
+		if (bSetupUndoRedo)
 		{
-			URigVMPin* RootPin = InPinToResolve->GetRootPin();
-			const FName ArgumentName = RootPin->GetFName();
-			const FString BaseCPPType = RigVMTypeUtils::IsArrayType(CPPType) ? RigVMTypeUtils::BaseTypeFromArrayType(CPPType) : CPPType;
-			const FString TypeToResolveAgainst = RootPin->IsArray() ? RigVMTypeUtils::ArrayTypeFromBaseType(BaseCPPType) : BaseCPPType;
-			const FRigVMTemplateArgument::FType ExpectedType(TypeToResolveAgainst, InCPPTypeObject);
-
-			// get all of the new types that we want to resolve towards
-			FRigVMTemplate::FTypeMap ResolvedTypes = OriginalTypes;
-			ResolvedTypes.Remove(ArgumentName);
-			if(TemplateNode->GetTemplate()->ResolveArgument(ArgumentName, ExpectedType, ResolvedTypes))
-			{
-				for(const FRigVMTemplate::FTypePair& TypePair : ResolvedTypes)
-				{
-					const FName& RootPinName = TypePair.Key;
-					const FRigVMTemplateArgument::FType& ResolvedType = TypePair.Value;
+			ActionStack->EndAction(Action);
+		}
 
-					URigVMPin* Pin = TemplateNode->FindPin(RootPinName.ToString());
-					check(Pin);
+		// Print python command
+		
+		return true;		
+	}
 
-					if(Pin->GetTemplateArgumentType().Matches(ResolvedType.CPPType, false))
-					{
-						continue;
-					}
+	return false;
+}
 
-					ResolveWildCardPinImpl(Pin, ResolvedType.CPPType, ResolvedType.CPPTypeObject, bSetupUndoRedo, false, false, false);
-				}
+bool URigVMController::UpdateFilteredPermutations(URigVMPin* InPin, URigVMPin* InLinkedPin, bool bSetupUndoRedo)
+{
+	URigVMTemplateNode* Node = Cast<URigVMTemplateNode>(InPin->GetNode());
+	if (!Node)
+	{
+		return false;
+	}
+	
+	TArray<int32> OldPermutations;
+	if (bSetupUndoRedo)
+	{
+		OldPermutations = Node->GetFilteredPermutationsIndices();
+	}
 
-				// re-attach the links. this is going to recursively change types
-				// on consecutive template nodes
-				for(const TPair<FString, FString>& LinkedPath : LinkedPaths)
-				{
-					const FString& SourcePath = LinkedPath.Key;
-					const FString& TargetPath = LinkedPath.Value;
-					URigVMPin* SourcePin = GetGraph()->FindPin(SourcePath);
-					URigVMPin* TargetPin = GetGraph()->FindPin(TargetPath);
+	if (Node->UpdateFilteredPermutations(InPin, InLinkedPin))
+	{
+		if (bSetupUndoRedo)
+		{
+			FRigVMSetTemplateFilteredPermutationsAction Action(Node, InPin->FindLinkForPin(InLinkedPin), OldPermutations);
+			ActionStack->BeginAction(Action);
+			ActionStack->AddAction(Action);
+			ActionStack->EndAction(Action);
+		}
 
-					if(SourcePin == nullptr || TargetPin == nullptr)
-					{
-						ReportRemovedLink(SourcePath, TargetPath);
-						continue;
-					}
+		return true;
+	}
 
-					// we are ok with this failing
-					if(!AddLink(SourcePin, TargetPin, bSetupUndoRedo,
-						SourcePin->GetNode() == TemplateNode ? ERigVMPinDirection::Output : ERigVMPinDirection::Input))
-					{
-						ReportRemovedLink(SourcePath, TargetPath);
-					}
-				}
+	return false;
+}
 
-				// reapply default values
-				ApplyPinStates(TemplateNode, PinStates);
-				
-				return true;
-			}
-			else
-			{
-				return false;
-			}
-		}
+bool URigVMController::UpdateFilteredPermutations(URigVMPin* InPin, const TArray<FRigVMTemplateArgument::FType>& InTypes, bool bSetupUndoRedo)
+{
+	URigVMTemplateNode* Node = Cast<URigVMTemplateNode>(InPin->GetNode());
+	if (!Node)
+	{
+		return false;
+	}
+	
+	TArray<int32> OldPermutations;
+	if (bSetupUndoRedo)
+	{
+		OldPermutations = Node->GetFilteredPermutationsIndices();
 	}
 
-	FRigVMControllerCompileBracketScope CompileScope(this);
-	if(bTraverseParentPins)
+	if (Node->UpdateFilteredPermutations(InPin, InTypes))
 	{
-		URigVMPin* RootPin = InPinToResolve->GetRootPin();
-		if(RootPin != InPinToResolve && RootPin->IsWildCard())
+		if (bSetupUndoRedo)
 		{
-			const FString BaseCPPType = RigVMTypeUtils::IsArrayType(CPPType) ? RigVMTypeUtils::BaseTypeFromArrayType(CPPType) : CPPType;
-			const FString TypeToResolveAgainst = RootPin->IsArray() ? RigVMTypeUtils::ArrayTypeFromBaseType(BaseCPPType) : BaseCPPType; 
-			return ResolveWildCardPinImpl(RootPin, TypeToResolveAgainst, CPPTypeObject, bSetupUndoRedo, bTraverseNode, false, bTraverseLinks);
+			FRigVMSetTemplateFilteredPermutationsAction Action(Node, nullptr, OldPermutations);
+			ActionStack->BeginAction(Action);
+			ActionStack->AddAction(Action);
+			ActionStack->EndAction(Action);
 		}
+
+		return true;
 	}
 
-	if(InPinToResolve->IsArray())
+	return false;
+}
+
+bool URigVMController::UpdateTemplateNodePinTypes(URigVMTemplateNode* InNode, bool bSetupUndoRedo)
+{
+	TArray<URigVMPin*> RemainingWildcardPins;
+	for (URigVMPin* Pin : InNode->GetPins())
 	{
-		for(URigVMPin* SubPin : InPinToResolve->GetSubPins())
+		TArray<FRigVMTemplateArgument::FType> Types = InNode->GetFilteredTypesForPin(Pin);
+		if (Types.IsEmpty())
 		{
-			const FString BaseCPPType = RigVMTypeUtils::IsArrayType(CPPType) ? RigVMTypeUtils::BaseTypeFromArrayType(CPPType) : CPPType;
-			const FString TypeToResolveAgainst = SubPin->IsArray() ? RigVMTypeUtils::ArrayTypeFromBaseType(BaseCPPType) : BaseCPPType;
+			continue;
+		}
+		
+		if (Types.Num() > 1)
+		{
+			// Unresolve
+			const FRigVMTemplateArgument* Argument = InNode->GetTemplate()->FindArgument(*Pin->GetName());
+			const FRigVMTemplateArgument::EArrayType ArrayType = Argument->GetArrayType();
 			
-			if(SubPin->IsWildCard())
+			FString CPPType = RigVMTypeUtils::GetWildCardCPPType();
+			UObject* CPPObjectType = RigVMTypeUtils::GetWildCardCPPTypeObject();
+
+			if (ArrayType == FRigVMTemplateArgument::EArrayType_ArrayValue)
 			{
-				ResolveWildCardPinImpl(SubPin, TypeToResolveAgainst, CPPTypeObject, bSetupUndoRedo, false, false, bTraverseLinks);
+				CPPType = RigVMTypeUtils::GetWildCardArrayCPPType();
 			}
-			else
+			else if(ArrayType == FRigVMTemplateArgument::EArrayType_Mixed)
+			{				
+				CPPType = Pin->IsArray() ? RigVMTypeUtils::GetWildCardArrayCPPType() : RigVMTypeUtils::GetWildCardCPPType();
+			}
+			
+			if (Pin->GetCPPType() != CPPType || Pin->GetCPPTypeObject() != CPPObjectType)
 			{
-				ChangePinType(SubPin, TypeToResolveAgainst, InCPPTypeObject, bSetupUndoRedo, false, true, true);
+				ChangePinType(Pin, CPPType, CPPObjectType, bSetupUndoRedo, false, false, false);
 			}
+			RemainingWildcardPins.Add(Pin);
+		}
+		else if (Types.Num() == 1)
+		{
+			// Resolve
+			FString CPPType = Types[0].CPPType;
+			UObject* CPPObjectType = Types[0].CPPTypeObject;
+			if (Pin->GetCPPType() != CPPType || Pin->GetCPPTypeObject() != CPPObjectType)
+			{
+				ChangePinType(Pin, CPPType, CPPObjectType, bSetupUndoRedo, false, false, false);
+			}
+		}
+		else
+		{
+			ensure(false);
+			return false;
 		}
 	}
 
-	FString ResolvedCPPType = RigVMTypeUtils::IsArrayType(CPPType) ? RigVMTypeUtils::BaseTypeFromArrayType(CPPType) : CPPType;
-
-	// change the pin's type to be of an array as well
-	if(InPinToResolve->IsArray())
-	{
-		ResolvedCPPType = RigVMTypeUtils::ArrayTypeFromBaseType(ResolvedCPPType);
-	}
-	else if(InPinToResolve->IsRootPin() && RigVMTypeUtils::IsArrayType(InCPPType))
+	if (RemainingWildcardPins.Num() > 0)
 	{
-		if(URigVMTemplateNode* TemplateNode = Cast<URigVMTemplateNode>(InPinToResolve->GetNode()))
+		TArray<int32> ResolvePermutations;
+		bool bCanAllReduceToSingleType = true;
+		for (URigVMPin* WildcardPin : RemainingWildcardPins)
 		{
-			const FString ArrayCPPType = RigVMTypeUtils::ArrayTypeFromBaseType(ResolvedCPPType);
-			FString ResolvedArrayCPPType;
-			if(TemplateNode->SupportsType(InPinToResolve, ArrayCPPType, &ResolvedArrayCPPType))
+			TArray<FRigVMTemplateArgument::FType> Types = InNode->GetFilteredTypesForPin(WildcardPin);
+			if (Types.IsEmpty())
 			{
-				ResolvedCPPType = ResolvedArrayCPPType;
+				continue;
+			}
+
+			// Find if all types are compatible (for example, float and double) 
+			bool bCanReduceToSingleType=true;
+			if (Types.Num() > 1)
+			{
+				for (int32 i=1; i<Types.Num(); ++i)
+				{
+					if (!Types[0].Matches(Types[i].CPPType))
+					{
+						bCanReduceToSingleType=false;
+						break;
+					}
+				}
 			}
+	
+			if (!bCanReduceToSingleType)
+			{
+				bCanAllReduceToSingleType = false;
+				break;
+			}			
 		}
-	}
 
-	if(ChangePinType(InPinToResolve, ResolvedCPPType, CPPTypeObject, bSetupUndoRedo, false, false, false))
-	{
-		if(URigVMTemplateNode* TemplateNode = Cast<URigVMTemplateNode>(InPinToResolve->GetNode()))
+		// Choose the first permutation
+		if (bCanAllReduceToSingleType)
 		{
-			ResolveTemplateNodePins(TemplateNode, bSetupUndoRedo);
+			TArray<int32> OldPermutations;
+			if (bSetupUndoRedo)
+			{
+				OldPermutations = InNode->GetFilteredPermutationsIndices();
+			}
+			InNode->FilteredPermutations = {InNode->FilteredPermutations[0]};
+			if (bSetupUndoRedo)
+			{
+				FRigVMSetTemplateFilteredPermutationsAction Action(InNode, nullptr, OldPermutations);
+				ActionStack->BeginAction(Action);
+				ActionStack->AddAction(Action);
+				ActionStack->EndAction(Action);
+			}
+			
+			return UpdateTemplateNodePinTypes(InNode, bSetupUndoRedo);
 		}
 	}
 	else
 	{
-		return false;
+		FullyResolveTemplateNode(InNode, InNode->FilteredPermutations[0], bSetupUndoRedo);
 	}
 
-	if(bTraverseLinks)
+	return true;
+}
+
+bool URigVMController::PropagateTemplateFilteredTypes(URigVMTemplateNode* InNode, bool bSetupUndoRedo) 
+{
+	auto UpdateAndPropagete = [&](URigVMPin* Pin, TArray<FRigVMTemplateArgument::FType>& Types)
 	{
-		if(URigVMNode* Node = InPinToResolve->GetNode())
+		TArray<URigVMPin*> OtherPins = Pin->GetLinkedSourcePins();
+		OtherPins.Append(Pin->GetLinkedTargetPins());
+		for (URigVMPin* OtherPin : OtherPins)
 		{
-			if(URigVMLibraryNode* LibraryNode = Cast<URigVMLibraryNode>(Node))
+			bool bPropagate = false;
+			if (URigVMTemplateNode* OtherTemplate = Cast<URigVMTemplateNode>(OtherPin->GetNode()))
 			{
-				const FString& PinName = InPinToResolve->GetRootPin()->GetName();
-
-				if(URigVMNode* EntryNode = LibraryNode->GetEntryNode())
+				if (!OtherTemplate->IsSingleton())
 				{
-					if(URigVMPin* EntryPin = EntryNode->FindPin(PinName))
+					if (OtherTemplate->PinNeedsFilteredTypesUpdate(OtherPin, Pin))
 					{
-						ResolveWildCardPin(EntryPin, InPinToResolve->GetRootPin(), bSetupUndoRedo, false, bTraverseParentPins, bTraverseLinks);
+						if (UpdateFilteredPermutations(OtherPin, Pin, bSetupUndoRedo))
+						{
+							UpdateTemplateNodePinTypes(OtherTemplate, bSetupUndoRedo);
+							if (!PropagateTemplateFilteredTypes(OtherTemplate, bSetupUndoRedo))
+							{
+								return false;
+							}
+						}
+						else
+						{
+							URigVMLink* Link = Pin->FindLinkForPin(OtherPin);
+							BreakLink(Link->GetSourcePin(), Link->GetTargetPin(), bSetupUndoRedo);
+							return false;
+						}
 					}
-				}
+				}								
+			}			
+		}
+		return true;
+	};
+	
+	for(URigVMPin* Pin : InNode->GetPins())
+	{
+		TArray<FRigVMTemplateArgument::FType> Types = InNode->GetFilteredTypesForPin(Pin);
+		if (!UpdateAndPropagete(Pin, Types))
+		{
+			return false;
+		}
 
-				if(URigVMNode* ReturnNode = LibraryNode->GetReturnNode())
-				{
-					if(URigVMPin* ReturnPin = ReturnNode->FindPin(PinName))
-					{
-						ResolveWildCardPin(ReturnPin, InPinToResolve->GetRootPin(), bSetupUndoRedo, false, bTraverseParentPins, bTraverseLinks);
-					}
-				}
+		if (Pin->GetSubPins().Num() > 0)
+		{
+			for (FRigVMTemplateArgument::FType& Type : Types)
+			{
+				Type.ConvertToBaseElement();
 			}
-			else if(Node->IsA<URigVMFunctionEntryNode>() || Node->IsA<URigVMFunctionReturnNode>())
+			for (URigVMPin* SubPin : Pin->GetSubPins())
 			{
-				if(URigVMLibraryNode* OuterNode = Cast<URigVMLibraryNode>(Node->GetGraph()->GetOuter()))
+				if (!UpdateAndPropagete(SubPin, Types))
 				{
-					const FString& PinName = InPinToResolve->GetRootPin()->GetName();
-					if(URigVMPin* OuterPin = OuterNode->FindPin(PinName))
-					{
-						ResolveWildCardPin(OuterPin, InPinToResolve->GetRootPin(), bSetupUndoRedo, false, bTraverseParentPins, bTraverseLinks);
-					}
+					return false;
 				}
 			}
 		}
 	}
+	return true;
+}
 
-	if(bTraverseNode)
+void URigVMController::RecomputeAllTemplateFilteredTypes(bool bSetupUndoRedo)
+{
+	if (!IsValidGraph())
 	{
-		if(const URigVMNode* Node = InPinToResolve->GetNode())
+		return;
+	}
+
+	if (!bIsTransacting && !IsGraphEditable())
+	{
+		return;
+	}
+
+	URigVMGraph* Graph = GetGraph();
+	check(Graph);
+	
+	InitializeAllTemplateFiltersInGraph(bSetupUndoRedo);
+	
+	for (URigVMLink* Link : Graph->GetLinks())
+	{
+		URigVMPin* OutputPin = Link->GetSourcePin();
+		URigVMPin* InputPin = Link->GetTargetPin();
+		if (URigVMTemplateNode* OutputNode = Cast<URigVMTemplateNode>(OutputPin->GetNode()))
 		{
-			if(!Node->IsA<URigVMTemplateNode>() &&
-				!Node->IsA<URigVMLibraryNode>() &&
-				!Node->IsA<URigVMFunctionEntryNode>() &&
-				!Node->IsA<URigVMFunctionReturnNode>())
+			if (!OutputNode->IsSingleton())
 			{
-				for(URigVMPin* Pin : Node->GetPins())
+				if (OutputNode->PinNeedsFilteredTypesUpdate(OutputPin, InputPin))
 				{
-					if(Pin->IsWildCard())
+					if (UpdateFilteredPermutations(OutputPin, InputPin, bSetupUndoRedo))
 					{
-						ResolveWildCardPin(Pin, InPinToResolve, bSetupUndoRedo, false, bTraverseParentPins, bTraverseLinks);
+						UpdateTemplateNodePinTypes(OutputNode, bSetupUndoRedo);
+						PropagateTemplateFilteredTypes(OutputNode, bSetupUndoRedo);
 					}
 				}
 			}
-		}
+		}				
+		if (URigVMTemplateNode* InputNode = Cast<URigVMTemplateNode>(InputPin->GetNode()))
+		{
+			if (!InputNode->IsSingleton())
+			{
+				if (InputNode->PinNeedsFilteredTypesUpdate(InputPin, OutputPin))
+				{
+					if (UpdateFilteredPermutations(InputPin, OutputPin, bSetupUndoRedo))
+					{
+						UpdateTemplateNodePinTypes(InputNode, bSetupUndoRedo);
+						PropagateTemplateFilteredTypes(InputNode, bSetupUndoRedo);
+					}
+				}
+			}
+		}				
+	}
+}
+
+void URigVMController::InitializeAllTemplateFiltersInGraph(bool bSetupUndoRedo)
+{
+	if (!IsValidGraph())
+	{
+		return;
 	}
 
-	if(bTraverseLinks)
+	if (!bIsTransacting && !IsGraphEditable())
 	{
-		// we don't need to recurse since unknown type pins don't have sub pins
-		TArray<URigVMPin*> LinkedPins = InPinToResolve->GetLinkedSourcePins(false); 
-		LinkedPins.Append(InPinToResolve->GetLinkedTargetPins(false));
+		return;
+	}
 
-		for(URigVMPin* LinkedPin : LinkedPins)
+	URigVMGraph* Graph = GetGraph();
+	check(Graph);
+	
+	
+	for (URigVMNode* Node : Graph->GetNodes())
+	{
+		if (URigVMTemplateNode* TemplateNode = Cast<URigVMTemplateNode>(Node))
 		{
-			if(LinkedPin->IsWildCard())
+			if (!TemplateNode->IsSingleton())
 			{
-				if(!ResolveWildCardPin(LinkedPin, InPinToResolve, bSetupUndoRedo))
+				if (!TemplateNode->PreferredPermutationTypes.IsEmpty())
 				{
-					return false;
+					continue;
+				}
+				
+				TArray<int32> OldPermutations = TemplateNode->FilteredPermutations;
+				TemplateNode->InitializeFilteredPermutations();
+				UpdateTemplateNodePinTypes(TemplateNode, bSetupUndoRedo);
+				if (bSetupUndoRedo)
+				{
+					FRigVMSetTemplateFilteredPermutationsAction FilteringAction(TemplateNode, nullptr, OldPermutations);
+					ActionStack->AddAction(FilteringAction);
 				}
 			}
 		}
 	}
-
-	return true;
 }
 
 bool URigVMController::ChangePinType(const FString& InPinPath, const FString& InCPPType, const FName& InCPPTypeObjectPath, bool bSetupUndoRedo, bool bSetupOrphanPins, bool bBreakLinks, bool bRemoveSubPins)
@@ -15314,6 +15767,15 @@ bool URigVMController::ChangePinType(URigVMPin* InPin, const FString& InCPPType,
 		AddPinsForStruct(InPin->GetScriptStruct(), InPin->GetNode(), InPin, InPin->Direction, DefaultValue, false, true);
 	}
 
+	if (InPin->IsArray())
+	{
+		const FString BaseCPPType = RigVMTypeUtils::BaseTypeFromArrayType(CPPType);
+		for (URigVMPin* SubPin : InPin->GetSubPins())
+		{
+			ChangePinType(SubPin, BaseCPPType, InCPPTypeObject, bSetupUndoRedo, bSetupOrphanPins, bBreakLinks, bRemoveSubPins);
+		}
+	}
+
 	// if this is a template clear its caches
 	if(URigVMTemplateNode* TemplateNode = Cast<URigVMTemplateNode>(InPin->GetNode()))
 	{
diff --git a/Engine/Source/Developer/RigVMDeveloper/Private/RigVMModel/RigVMControllerActions.cpp b/Engine/Source/Developer/RigVMDeveloper/Private/RigVMModel/RigVMControllerActions.cpp
--- a/Engine/Source/Developer/RigVMDeveloper/Private/RigVMModel/RigVMControllerActions.cpp
+++ b/Engine/Source/Developer/RigVMDeveloper/Private/RigVMModel/RigVMControllerActions.cpp
@@ -1328,6 +1328,109 @@ bool FRigVMSetPinDefaultValueAction::Redo(URigVMController* InController)
 	return FRigVMBaseAction::Redo(InController);
 }
 
+FRigVMSetTemplateFilteredPermutationsAction::FRigVMSetTemplateFilteredPermutationsAction(URigVMTemplateNode* InNode, URigVMLink* InLink, const TArray<int32>& InOldFilteredPermutations)
+: NodePath(InNode->GetNodePath())
+, LinkPath(InLink ? InLink->GetPinPathRepresentation() : FString())
+, OldFilteredPermutations(InOldFilteredPermutations)
+, NewFilteredPermutations(InNode->GetFilteredPermutationsIndices())
+{
+	
+}
+
+bool FRigVMSetTemplateFilteredPermutationsAction::Merge(const FRigVMBaseAction* Other)
+{
+	if (!FRigVMBaseAction::Merge(Other))
+	{
+		return false;
+	}
+
+	const FRigVMSetTemplateFilteredPermutationsAction* Action = (const FRigVMSetTemplateFilteredPermutationsAction*)Other;
+	if (NodePath != Action->NodePath)
+	{
+		return false;
+	}
+
+	NewFilteredPermutations = Action->NewFilteredPermutations;
+	return true;
+}
+
+bool FRigVMSetTemplateFilteredPermutationsAction::Undo(URigVMController* InController)
+{
+	if(!FRigVMBaseAction::Undo(InController))
+	{
+		return false;
+	}
+
+	if (URigVMTemplateNode* TemplateNode = Cast<URigVMTemplateNode>(InController->GetGraph()->FindNode(NodePath)))
+	{
+		TemplateNode->FilteredPermutations = OldFilteredPermutations;
+		return true;
+	}
+	return false;
+}
+
+bool FRigVMSetTemplateFilteredPermutationsAction::Redo(URigVMController* InController)
+{
+	URigVMTemplateNode* TemplateNode = Cast<URigVMTemplateNode>(InController->GetGraph()->FindNode(NodePath));
+	if (!TemplateNode)
+	{
+		return false;
+	}
+	TemplateNode->FilteredPermutations = NewFilteredPermutations;
+	return FRigVMBaseAction::Redo(InController);
+}
+
+FRigVMSetPreferredTemplatePermutationsAction::FRigVMSetPreferredTemplatePermutationsAction(URigVMTemplateNode* InNode, const TArray<FString>& InPreferredPermutationTypes)
+: NodePath(InNode->GetNodePath())
+, OldPreferredPermutationTypes(InNode->PreferredPermutationTypes)
+, NewPreferredPermutationTypes(InPreferredPermutationTypes)
+{
+	
+}
+
+bool FRigVMSetPreferredTemplatePermutationsAction::Merge(const FRigVMBaseAction* Other)
+{
+	if (!FRigVMBaseAction::Merge(Other))
+	{
+		return false;
+	}
+
+	const FRigVMSetPreferredTemplatePermutationsAction* Action = (const FRigVMSetPreferredTemplatePermutationsAction*)Other;
+	if (NodePath != Action->NodePath)
+	{
+		return false;
+	}
+
+	NewPreferredPermutationTypes = Action->NewPreferredPermutationTypes;
+	return true;
+}
+
+bool FRigVMSetPreferredTemplatePermutationsAction::Undo(URigVMController* InController)
+{
+	if(!FRigVMBaseAction::Undo(InController))
+	{
+		return false;
+	}
+
+	if (URigVMTemplateNode* TemplateNode = Cast<URigVMTemplateNode>(InController->GetGraph()->FindNode(NodePath)))
+	{
+		TemplateNode->PreferredPermutationTypes = OldPreferredPermutationTypes;
+		return true;
+	}
+	return false;
+}
+
+bool FRigVMSetPreferredTemplatePermutationsAction::Redo(URigVMController* InController)
+{
+	URigVMTemplateNode* TemplateNode = Cast<URigVMTemplateNode>(InController->GetGraph()->FindNode(NodePath));
+	if (!TemplateNode)
+	{
+		return false;
+	}
+	TemplateNode->PreferredPermutationTypes = NewPreferredPermutationTypes;
+	return FRigVMBaseAction::Redo(InController);
+}
+
 FRigVMInsertArrayPinAction::FRigVMInsertArrayPinAction(URigVMPin* InArrayPin, int32 InIndex, const FString& InNewDefaultValue)
 : ArrayPinPath(InArrayPin->GetPinPath())
 , Index(InIndex)
@@ -1460,31 +1563,11 @@ bool FRigVMChangePinTypeAction::Undo(URigVMController* InController)
 		return false;
 	}
 	
-	if(const URigVMPin* TemplateNodePin = InController->GetGraph()->FindPin(PinPath))
-	{
-		if(TemplateNodePin->GetNode()->IsA<URigVMTemplateNode>() && TemplateNodePin->IsWildCard())
-		{
-			return InController->ResolveWildCardPin(PinPath, OldCPPType, OldCPPTypeObjectPath, false, false);
-		}
-	}
-	
 	return InController->ChangePinType(PinPath, OldCPPType, OldCPPTypeObjectPath, false, bSetupOrphanPins, bBreakLinks, bRemoveSubPins);
 }
 
 bool FRigVMChangePinTypeAction::Redo(URigVMController* InController)
 {
-	if(const URigVMPin* TemplateNodePin = InController->GetGraph()->FindPin(PinPath))
-	{
-		if(TemplateNodePin->GetNode()->IsA<URigVMTemplateNode>() && TemplateNodePin->IsWildCard())
-		{
-			if(!InController->ResolveWildCardPin(PinPath, NewCPPType, NewCPPTypeObjectPath, false, false))
-			{
-				return false;
-			}
-			return FRigVMBaseAction::Redo(InController);
-		}
-	}
-
 	if(!InController->ChangePinType(PinPath, NewCPPType, NewCPPTypeObjectPath, false, bSetupOrphanPins, bBreakLinks, bRemoveSubPins))
 	{
 		return false;
diff --git a/Engine/Source/Developer/RigVMDeveloper/Private/RigVMModel/RigVMGraph.cpp b/Engine/Source/Developer/RigVMDeveloper/Private/RigVMModel/RigVMGraph.cpp
--- a/Engine/Source/Developer/RigVMDeveloper/Private/RigVMModel/RigVMGraph.cpp
+++ b/Engine/Source/Developer/RigVMDeveloper/Private/RigVMModel/RigVMGraph.cpp
@@ -421,9 +421,9 @@ void URigVMGraph::PrepareCycleChecking(URigVMPin* InPin, bool bAsInput)
 	GetDiagnosticsAST(false, LinksToSkip)->PrepareCycleChecking(InPin);
 }
 
-bool URigVMGraph::CanLink(URigVMPin* InSourcePin, URigVMPin* InTargetPin, FString* OutFailureReason, const FRigVMByteCode* InByteCode, ERigVMPinDirection InUserLinkDirection)
+bool URigVMGraph::CanLink(URigVMPin* InSourcePin, URigVMPin* InTargetPin, FString* OutFailureReason, const FRigVMByteCode* InByteCode, ERigVMPinDirection InUserLinkDirection, bool bInAllowWildcard)
 {
-	if (!URigVMPin::CanLink(InSourcePin, InTargetPin, OutFailureReason, InByteCode, InUserLinkDirection))
+	if (!URigVMPin::CanLink(InSourcePin, InTargetPin, OutFailureReason, InByteCode, InUserLinkDirection, bInAllowWildcard))
 	{
 		return false;
 	}
diff --git a/Engine/Source/Developer/RigVMDeveloper/Private/RigVMModel/RigVMLink.cpp b/Engine/Source/Developer/RigVMDeveloper/Private/RigVMModel/RigVMLink.cpp
--- a/Engine/Source/Developer/RigVMDeveloper/Private/RigVMModel/RigVMLink.cpp
+++ b/Engine/Source/Developer/RigVMDeveloper/Private/RigVMModel/RigVMLink.cpp
@@ -50,7 +50,7 @@ int32 URigVMLink::GetLinkIndex() const
 
 }
 
-URigVMPin* URigVMLink::GetSourcePin()
+URigVMPin* URigVMLink::GetSourcePin() const
 {
 	if (SourcePin == nullptr)
 	{
@@ -59,7 +59,7 @@ URigVMPin* URigVMLink::GetSourcePin()
 	return SourcePin;
 }
 
-URigVMPin* URigVMLink::GetTargetPin()
+URigVMPin* URigVMLink::GetTargetPin() const
 {
 	if (TargetPin == nullptr)
 	{
@@ -68,6 +68,19 @@ URigVMPin* URigVMLink::GetTargetPin()
 	return TargetPin;
 }
 
+URigVMPin* URigVMLink::GetOppositePin(const URigVMPin* InPin) const
+{
+	if (InPin == GetSourcePin())
+	{
+		return GetTargetPin();
+	}
+	else if (InPin == GetTargetPin())
+	{
+		return GetSourcePin();
+	}
+	return nullptr;
+}
+
 FString URigVMLink::GetPinPathRepresentation()
 {
 	return FString::Printf(TEXT("%s -> %s"), *GetSourcePin()->GetPinPath(), *GetTargetPin()->GetPinPath());
diff --git a/Engine/Source/Developer/RigVMDeveloper/Private/RigVMModel/RigVMPin.cpp b/Engine/Source/Developer/RigVMDeveloper/Private/RigVMModel/RigVMPin.cpp
--- a/Engine/Source/Developer/RigVMDeveloper/Private/RigVMModel/RigVMPin.cpp
+++ b/Engine/Source/Developer/RigVMDeveloper/Private/RigVMModel/RigVMPin.cpp
@@ -1338,6 +1338,19 @@ URigVMPin* URigVMPin::GetPinForLink() const
 	return PinForLink;
 }
 
+URigVMLink* URigVMPin::FindLinkForPin(const URigVMPin* InOtherPin) const
+{
+	for (URigVMLink* Link : Links)
+	{
+		if ((Link->GetSourcePin() == this && Link->GetTargetPin() == InOtherPin) ||
+			(Link->GetSourcePin() == InOtherPin && Link->GetTargetPin() == this))
+		{
+			return Link;
+		}
+	}
+	return nullptr;
+}
+
 URigVMPin* URigVMPin::GetOriginalPinFromInjectedNode() const
 {
 	if(GetNode() == nullptr)
@@ -1560,7 +1573,7 @@ URigVMGraph* URigVMPin::GetGraph() const
 	return nullptr;
 }
 
-bool URigVMPin::CanLink(URigVMPin* InSourcePin, URigVMPin* InTargetPin, FString* OutFailureReason, const FRigVMByteCode* InByteCode, ERigVMPinDirection InUserLinkDirection)
+bool URigVMPin::CanLink(URigVMPin* InSourcePin, URigVMPin* InTargetPin, FString* OutFailureReason, const FRigVMByteCode* InByteCode, ERigVMPinDirection InUserLinkDirection, bool bInAllowWildcard)
 {
 	if (InSourcePin == nullptr || InTargetPin == nullptr)
 	{
@@ -1640,12 +1653,16 @@ bool URigVMPin::CanLink(URigVMPin* InSourcePin, URigVMPin* InTargetPin, FString*
 			bCPPTypesDiffer = false;
 		}
 
+		if (bInAllowWildcard && (InSourcePin->CPPType == RigVMTypeUtils::GetWildCardCPPType() || InTargetPin->CPPType == RigVMTypeUtils::GetWildCardCPPType()))
+		{
+			bCPPTypesDiffer = false;
+		}
+
 		if (bCPPTypesDiffer)
 		{
 			{
 				auto TemplateNodeSupportsType = [](URigVMPin* InPin, const FString& InCPPType, FString* OutFailureReason) -> bool
 				{
-
 					if (URigVMTemplateNode* TemplateNode = Cast<URigVMTemplateNode>(InPin->GetNode()))
 					{
 						if (TemplateNode->SupportsType(InPin, InCPPType))
diff --git a/Engine/Source/Developer/RigVMDeveloper/Public/RigVMModel/Nodes/RigVMArrayNode.h b/Engine/Source/Developer/RigVMDeveloper/Public/RigVMModel/Nodes/RigVMArrayNode.h
--- a/Engine/Source/Developer/RigVMDeveloper/Public/RigVMModel/Nodes/RigVMArrayNode.h
+++ b/Engine/Source/Developer/RigVMDeveloper/Public/RigVMModel/Nodes/RigVMArrayNode.h
@@ -38,6 +38,9 @@ public:
 	// Override of the template
 	virtual const FRigVMTemplate* GetTemplate() const override;
 
+	// Override of singleton
+	virtual bool IsSingleton() const override { return false; }
+
 	// Returns the op code of this node
 	UFUNCTION(BlueprintCallable, Category = RigVMArrayNode)
 	ERigVMOpCode GetOpCode() const;
diff --git a/Engine/Source/Developer/RigVMDeveloper/Public/RigVMModel/Nodes/RigVMIfNode.h b/Engine/Source/Developer/RigVMDeveloper/Public/RigVMModel/Nodes/RigVMIfNode.h
--- a/Engine/Source/Developer/RigVMDeveloper/Public/RigVMModel/Nodes/RigVMIfNode.h
+++ b/Engine/Source/Developer/RigVMDeveloper/Public/RigVMModel/Nodes/RigVMIfNode.h
@@ -18,6 +18,7 @@ public:
 	// Override from URigVMTemplateNode
 	virtual FName GetNotation() const override;
 	virtual const FRigVMTemplate* GetTemplate() const override;
+	virtual bool IsSingleton() const override { return false; }
 
 	// Override from URigVMNode
 	virtual FString GetNodeTitle() const override { return IfName; }
diff --git a/Engine/Source/Developer/RigVMDeveloper/Public/RigVMModel/Nodes/RigVMRerouteNode.h b/Engine/Source/Developer/RigVMDeveloper/Public/RigVMModel/Nodes/RigVMRerouteNode.h
--- a/Engine/Source/Developer/RigVMDeveloper/Public/RigVMModel/Nodes/RigVMRerouteNode.h
+++ b/Engine/Source/Developer/RigVMDeveloper/Public/RigVMModel/Nodes/RigVMRerouteNode.h
@@ -34,6 +34,7 @@ public:
 
 	virtual FName GetNotation() const override;
 	virtual const FRigVMTemplate* GetTemplate() const override;
+	virtual bool IsSingleton() const override { return false; }
 
 private:
 
diff --git a/Engine/Source/Developer/RigVMDeveloper/Public/RigVMModel/Nodes/RigVMSelectNode.h b/Engine/Source/Developer/RigVMDeveloper/Public/RigVMModel/Nodes/RigVMSelectNode.h
--- a/Engine/Source/Developer/RigVMDeveloper/Public/RigVMModel/Nodes/RigVMSelectNode.h
+++ b/Engine/Source/Developer/RigVMDeveloper/Public/RigVMModel/Nodes/RigVMSelectNode.h
@@ -18,6 +18,7 @@ public:
 	// Override from URigVMTemplateNode
 	virtual FName GetNotation() const override;
 	virtual const FRigVMTemplate* GetTemplate() const override;
+	virtual bool IsSingleton() const override { return false; }
 
 	// Override from URigVMNode
 	virtual FString GetNodeTitle() const override { return SelectName; }
diff --git a/Engine/Source/Developer/RigVMDeveloper/Public/RigVMModel/Nodes/RigVMTemplateNode.h b/Engine/Source/Developer/RigVMDeveloper/Public/RigVMModel/Nodes/RigVMTemplateNode.h
--- a/Engine/Source/Developer/RigVMDeveloper/Public/RigVMModel/Nodes/RigVMTemplateNode.h
+++ b/Engine/Source/Developer/RigVMDeveloper/Public/RigVMModel/Nodes/RigVMTemplateNode.h
@@ -39,24 +39,21 @@ public:
 	UFUNCTION(BlueprintPure, Category = Template)
 	virtual FName GetNotation() const;
 
+	UFUNCTION(BlueprintCallable, Category = Template)
+	virtual bool IsSingleton() const;
+
 	// returns true if a pin supports a given type
 	bool SupportsType(const URigVMPin* InPin, const FString& InCPPType, FString* OutCPPType = nullptr);
 
-	// resolves a pin to a new type outputs the new type map
-	bool GetTypeMapForNewPinType(const URigVMPin* InPin, const FString& InCPPType, UObject* InCPPTypeObject, FRigVMTemplate::FTypeMap& OutTypes) const;
-
-	// returns the index of the resolved RigVM function (or INDEX_NONE)
-	const TArray<int32>& GetResolvedPermutationIndices(FRigVMTemplate::FTypeMap* OutTypes = nullptr) const;
+	// returns true if a pin supports a given type after filtering
+	bool FilteredSupportsType(const URigVMPin* InPIn, const FString& InCPPType, FString* OutCPPType = nullptr, bool bAllowFloatingPointCasts = true);
 
 	// returns the resolved functions for the template
-	TArray<const FRigVMFunction*> GetResolvedPermutations(FRigVMTemplate::FTypeMap* OutTypes = nullptr) const;
+	TArray<const FRigVMFunction*> GetResolvedPermutations() const;
 
 	// returns the template used for this node
 	virtual const FRigVMTemplate* GetTemplate() const;
 
-	// returns the type map for the currently resolved pins
-	FRigVMTemplate::FTypeMap GetResolvedTypes() const;
-
 	// returns the resolved function or nullptr if there are still unresolved pins left
 	const FRigVMFunction* GetResolvedFunction() const;
 
@@ -74,9 +71,32 @@ public:
 	// returns the display name for a pin
 	FName GetDisplayNameForPin(const FName& InRootPinName, const TArray<int32>& InPermutationIndices = TArray<int32>()) const;
 
+	// returns the indeces of the filtered permutations
+	const TArray<int32>& GetFilteredPermutationsIndices() const;
+
+	// returns the filtered types of this pin
+	TArray<FRigVMTemplateArgument::FType> GetFilteredTypesForPin(URigVMPin* InPin) const;
+
+	// returns true if updating pin filters with InTypes would result in different filters 
+	bool PinNeedsFilteredTypesUpdate(URigVMPin* InPin, const TArray<FRigVMTemplateArgument::FType>& InTypes);
+	bool PinNeedsFilteredTypesUpdate(URigVMPin* InPin, URigVMPin* LinkedPin);
+
+	// updates the filtered permutations given a link or the types for a pin
+	bool UpdateFilteredPermutations(URigVMPin* InPin, const TArray<FRigVMTemplateArgument::FType>& InTypes);
+	bool UpdateFilteredPermutations(URigVMPin* InPin, URigVMPin* LinkedPin);
+
+	// initializes the filtered permutations to all possible permutations
+	void InitializeFilteredPermutations();
+	
 protected:
 
 	virtual void InvalidateCache() override;
+	
+	TArray<int32> GetNewFilteredPermutations(URigVMPin* InPin, URigVMPin* LinkedPin);
+	TArray<int32> GetNewFilteredPermutations(URigVMPin* InPin, const TArray<FRigVMTemplateArgument::FType>& InTypes);
+
+	int32 FindPermuationForTypes(const TArray<FString>& ArgumentTypes);
+	TArray<FString> GetArgumentTypesForPermutation(const int32 InPermutationIndex);
 
 	UPROPERTY()
 	FName TemplateNotation;
@@ -84,6 +104,12 @@ protected:
 	UPROPERTY()
 	FString ResolvedFunctionName;
 
+	// Indicates a preferred permuation using the types of the arguments
+	// Each element is in the format "ArgumentName:CPPType"
+	UPROPERTY()
+	TArray<FString> PreferredPermutationTypes;
+	
+	TArray<int32> FilteredPermutations;
 	TMap<FString, TPair<bool, FRigVMTemplateArgument::FType>> SupportedTypesCache;
 
 	mutable const FRigVMTemplate* CachedTemplate;
@@ -91,5 +117,7 @@ protected:
 	mutable TArray<int32> ResolvedPermutations;
 
 	friend class URigVMController;
+	friend struct FRigVMSetTemplateFilteredPermutationsAction;
+	friend struct FRigVMSetPreferredTemplatePermutationsAction;
 };
 
diff --git a/Engine/Source/Developer/RigVMDeveloper/Public/RigVMModel/RigVMController.h b/Engine/Source/Developer/RigVMDeveloper/Public/RigVMModel/RigVMController.h
--- a/Engine/Source/Developer/RigVMDeveloper/Public/RigVMModel/RigVMController.h
+++ b/Engine/Source/Developer/RigVMDeveloper/Public/RigVMModel/RigVMController.h
@@ -86,6 +86,7 @@ DECLARE_DELEGATE_RetVal_OneParam(bool, FRigVMController_RequestLocalizeFunctionD
 DECLARE_DELEGATE_RetVal_ThreeParams(FName, FRigVMController_RequestNewExternalVariableDelegate, FRigVMGraphVariableDescription, bool, bool);
 DECLARE_DELEGATE_RetVal_TwoParams(bool, FRigVMController_IsDependencyCyclicDelegate, UObject*, UObject*)
 DECLARE_DELEGATE_RetVal_TwoParams(FRigVMController_BulkEditResult, FRigVMController_RequestBulkEditDialogDelegate, URigVMLibraryNode*, ERigVMControllerBulkEditType)
+DECLARE_DELEGATE_RetVal_OneParam(bool, FRigVMController_RequestBreakLinksDialogDelegate, TArray<URigVMLink*>)
 DECLARE_DELEGATE_FiveParams(FRigVMController_OnBulkEditProgressDelegate, TSoftObjectPtr<URigVMFunctionReferenceNode>, ERigVMControllerBulkEditType, ERigVMControllerBulkEditProgress, int32, int32)
 DECLARE_DELEGATE_RetVal_TwoParams(FString, FRigVMController_PinPathRemapDelegate, const FString& /* InPinPath */, bool /* bIsInput */);
 DECLARE_DELEGATE_OneParam(FRigVMController_RequestJumpToHyperlinkDelegate, const UObject* InSubject);
@@ -307,6 +308,7 @@ public:
 	// Turns a resolved templated node(s) back into its template.
 	UFUNCTION(BlueprintCallable, Category = RigVMController)
 	bool UnresolveTemplateNodes(const TArray<FName>& InNodeNames, bool bSetupUndoRedo = true, bool bPrintPythonCommand = false);
+	bool UnresolveTemplateNodes(const TArray<URigVMTemplateNode*>& InNodes, bool bSetupUndoRedo);
 
 	// Upgrades a set of nodes with each corresponding next known version
 	UFUNCTION(BlueprintCallable, Category = RigVMController)
@@ -400,7 +402,8 @@ public:
 	// Resolves a wildcard pin on any node
 	UFUNCTION(BlueprintCallable, Category = RigVMController)
 	bool ResolveWildCardPin(const FString& InPinPath, const FString& InCPPType, const FName& InCPPTypeObjectPath, bool bSetupUndoRedo = true, bool bPrintPythonCommand = false);
-
+	bool ResolveWildCardPin(URigVMPin* InPin, const FRigVMTemplateArgument::FType& InType, bool bSetupUndoRedo = true, bool bPrintPythonCommand = false);
+		
 	// Adds a Function / Struct Node to the edited Graph as an injected node
 	// UnitNode represent a RIGVM_METHOD declaration on a USTRUCT.
 	// This causes a NodeAdded modified event.
@@ -646,8 +649,6 @@ public:
 	UFUNCTION(BlueprintCallable, Category = RigVMController, meta=(DeprecatedFunction))
 	bool RenameParameter(const FName& InOldName, const FName& InNewName, bool bSetupUndoRedo = true);\
 
-	bool UnresolveTemplateNodes(const TArray<URigVMTemplateNode*>& InNodes, bool bSetupUndoRedo);
-
 	// Upgrades a set of nodes with each corresponding next known version
 	TArray<URigVMNode*> UpgradeNodes(const TArray<URigVMNode*>& InNodes, bool bRecursive = true, bool bSetupUndoRedo = true);
 
@@ -744,7 +745,7 @@ public:
 	// Adds a link to the graph.
 	// This causes a LinkAdded modified event.
 	UFUNCTION(BlueprintCallable, Category = RigVMController)
-	bool AddLink(const FString& InOutputPinPath, const FString& InInputPinPath, bool bSetupUndoRedo = true, bool bPrintPythonCommand = false, ERigVMPinDirection InUserDirection = ERigVMPinDirection::Invalid);
+	bool AddLink(const FString& InOutputPinPath, const FString& InInputPinPath, bool bSetupUndoRedo = true, bool bPrintPythonCommand = false, ERigVMPinDirection InUserDirection = ERigVMPinDirection::Output);
 
 	// Removes a link from the graph.
 	// This causes a LinkRemoved modified event.
@@ -859,6 +860,9 @@ public:
 	// A delegate to ask the host / client for a dialog to confirm a bulk edit
 	FRigVMController_RequestBulkEditDialogDelegate RequestBulkEditDialogDelegate;
 
+	// A delegate to ask the host / client for a dialog to confirm a bulk edit
+	FRigVMController_RequestBreakLinksDialogDelegate RequestBreakLinksDialogDelegate;
+
 	// A delegate to inform the host / client about the progress during a bulk edit
 	FRigVMController_OnBulkEditProgressDelegate OnBulkEditProgressDelegate;
 
@@ -887,6 +891,12 @@ public:
 
 	// removes any orphan pins that no longer holds a link
 	bool RemoveUnusedOrphanedPins(URigVMNode* InNode, bool bNotify);
+
+	// Initializes and recomputes the filtered permutations of all template nodes in the graph 
+	void RecomputeAllTemplateFilteredTypes(bool bSetupUndoRedo);
+
+	// Initializes filtered permuations to be unresolved in all template nodes in graph
+	void InitializeAllTemplateFiltersInGraph(bool bSetupUndoRedo);
 	
 #endif
 
@@ -1008,7 +1018,7 @@ public:
 	bool BreakAllLinks(URigVMPin* Pin, bool bAsInput, bool bSetupUndoRedo = true);
 
 private:
-	void BreakAllLinksRecursive(URigVMPin* Pin, bool bAsInput, bool bTowardsParent, bool bSetupUndoRedo);
+	bool BreakAllLinksRecursive(URigVMPin* Pin, bool bAsInput, bool bTowardsParent, bool bSetupUndoRedo);
 	void UpdateRerouteNodeAfterChangingLinks(URigVMPin* PinChanged, bool bSetupUndoRedo = true);
 	bool SetPinExpansion(URigVMPin* InPin, bool bIsExpanded, bool bSetupUndoRedo = true);
 	void ExpandPinRecursively(URigVMPin* InPin, bool bSetupUndoRedo);
@@ -1048,11 +1058,30 @@ private:
 	static void CreateDefaultValueForStructIfRequired(UScriptStruct* InStruct, FString& InOutDefaultValue);
 	static void PostProcessDefaultValue(URigVMPin* Pin, FString& OutDefaultValue);
 
-	bool ResolveWildCardPin(URigVMPin* InPinToResolve, URigVMPin* InTemplatePin, bool bSetupUndoRedo, bool bTraverseNode = true, bool bTraverseParentPins = true, bool bTraverseLinks = true, bool bAllowFloatingPointCasts = false);
-	bool ResolveWildCardPinImpl(URigVMPin* InPinToResolve, const FString& InCPPType, UObject* InCPPTypeObject, bool bSetupUndoRedo, bool bTraverseNode = true, bool bTraverseParentPins = true, bool bTraverseLinks = true, bool bAllowFloatingPointCasts = false);
-	void ResolveTemplateNodePins(URigVMTemplateNode* InNode, bool bSetupUndoRedo);
 	void ResolveTemplateNodeMetaData(URigVMTemplateNode* InNode, bool bSetupUndoRedo);
 	bool FullyResolveTemplateNode(URigVMTemplateNode* InNode, int32 InPermutationIndex, bool bSetupUndoRedo);
+
+	// Prepare the graph for the change this template node is about to make
+	// If any of the types is supported (without breaking any links), then the filtered permutations will be updated and the change will
+	// propagate to other nodes in the graph.
+	// If it is not supported, we will attempt to find and break any links that do not support this change
+	bool PrepareTemplatePinForType(URigVMPin* InPin, const TArray<FRigVMTemplateArgument::FType>& InTypes, bool bSetupUndoRedo);
+
+	// Get filtered types for a wildcard node. If template node, that means just returning its filtered wildcard types, but if it's another type of node (select, if, rereoute), iterate
+	// its connections to figure out the filtered types
+	TArray<FRigVMTemplateArgument::FType> GetWildcardFilteredTypes(URigVMPin* InPin);
+	
+	// Updates the permutations allowed without having to break any links
+	bool UpdateFilteredPermutations(URigVMPin* InPin, URigVMPin* InLinkedPin, bool bSetupUndoRedo);
+	bool UpdateFilteredPermutations(URigVMPin* InPin, const TArray<FRigVMTemplateArgument::FType>& InTypes, bool bSetupUndoRedo);
+
+	// Changes Pin types if filtered types of a pin are unique
+	bool UpdateTemplateNodePinTypes(URigVMTemplateNode* InNode, bool bSetupUndoRedo);
+
+	// Reduces the filtered permutations of all templates in the graph to comply with the types filtered by InNode
+	// Returns false if a link had to be broken
+	bool PropagateTemplateFilteredTypes(URigVMTemplateNode* InNode, bool bSetupUndoRedo);
+
 	bool ChangePinType(const FString& InPinPath, const FString& InCPPType, const FName& InCPPTypeObjectPath, bool bSetupUndoRedo, bool bSetupOrphanPins = true, bool bBreakLinks = true, bool bRemoveSubPins = true);
 	bool ChangePinType(URigVMPin* InPin, const FString& InCPPType, const FName& InCPPTypeObjectPath, bool bSetupUndoRedo, bool bSetupOrphanPins = true, bool bBreakLinks = true, bool bRemoveSubPins = true);
 	bool ChangePinType(URigVMPin* InPin, const FString& InCPPType,UObject* InCPPTypeObject, bool bSetupUndoRedo, bool bSetupOrphanPins = true, bool bBreakLinks = true, bool bRemoveSubPins = true);
@@ -1165,6 +1194,7 @@ private:
 	bool bIsTransacting; // Performing undo/redo transaction
 	bool bIsRunningUnitTest;
 	bool bIsFullyResolvingTemplateNode;
+	bool bSuspendRecomputingTemplateFilters;
 
 	friend class URigVMGraph;
 	friend class URigVMPin;
diff --git a/Engine/Source/Developer/RigVMDeveloper/Public/RigVMModel/RigVMControllerActions.h b/Engine/Source/Developer/RigVMDeveloper/Public/RigVMModel/RigVMControllerActions.h
--- a/Engine/Source/Developer/RigVMDeveloper/Public/RigVMModel/RigVMControllerActions.h
+++ b/Engine/Source/Developer/RigVMDeveloper/Public/RigVMModel/RigVMControllerActions.h
@@ -1116,6 +1116,65 @@ public:
 	FString NewDefaultValue;
 };
 
+/**
+ * An action setting the filtered permutations on a template node
+ */
+USTRUCT()
+struct FRigVMSetTemplateFilteredPermutationsAction : public FRigVMBaseAction
+{
+	GENERATED_BODY()
+
+public:
+
+	FRigVMSetTemplateFilteredPermutationsAction() {}
+	FRigVMSetTemplateFilteredPermutationsAction(URigVMTemplateNode* InNode, URigVMLink* InLink, const TArray<int32>& InOldFilteredPermutations);
+	virtual ~FRigVMSetTemplateFilteredPermutationsAction() {};
+	virtual UScriptStruct* GetScriptStruct() const override { return FRigVMSetTemplateFilteredPermutationsAction::StaticStruct(); }
+	virtual bool Merge(const FRigVMBaseAction* Other);
+	virtual bool Undo(URigVMController* InController) override;
+	virtual bool Redo(URigVMController* InController) override;
+
+	UPROPERTY()
+	FString NodePath;
+
+	UPROPERTY()
+	FString LinkPath;
+
+	UPROPERTY()
+	TArray<int32> OldFilteredPermutations;
+
+	UPROPERTY()
+	TArray<int32> NewFilteredPermutations;
+};
+
+/**
+ * An action setting the filtered permutations on a template node
+ */
+USTRUCT()
+struct FRigVMSetPreferredTemplatePermutationsAction : public FRigVMBaseAction
+{
+	GENERATED_BODY()
+
+public:
+
+	FRigVMSetPreferredTemplatePermutationsAction() {}
+	FRigVMSetPreferredTemplatePermutationsAction(URigVMTemplateNode* InNode, const TArray<FString>& InPreferredPermutation);
+	virtual ~FRigVMSetPreferredTemplatePermutationsAction() {};
+	virtual UScriptStruct* GetScriptStruct() const override { return FRigVMSetTemplateFilteredPermutationsAction::StaticStruct(); }
+	virtual bool Merge(const FRigVMBaseAction* Other);
+	virtual bool Undo(URigVMController* InController) override;
+	virtual bool Redo(URigVMController* InController) override;
+
+	UPROPERTY()
+	FString NodePath;
+
+	UPROPERTY()
+	TArray<FString> OldPreferredPermutationTypes;
+
+	UPROPERTY()
+	TArray<FString> NewPreferredPermutationTypes;
+};
+
 /**
  * An action inserting a new array pin in the graph.
  */
diff --git a/Engine/Source/Developer/RigVMDeveloper/Public/RigVMModel/RigVMGraph.h b/Engine/Source/Developer/RigVMDeveloper/Public/RigVMModel/RigVMGraph.h
--- a/Engine/Source/Developer/RigVMDeveloper/Public/RigVMModel/RigVMGraph.h
+++ b/Engine/Source/Developer/RigVMDeveloper/Public/RigVMModel/RigVMGraph.h
@@ -155,7 +155,7 @@ public:
 	UScriptStruct* GetExecuteContextStruct() const;
 
 	void PrepareCycleChecking(URigVMPin* InPin, bool bAsInput);
-	virtual bool CanLink(URigVMPin* InSourcePin, URigVMPin* InTargetPin, FString* OutFailureReason, const FRigVMByteCode* InByteCode, ERigVMPinDirection InUserLinkDirection = ERigVMPinDirection::IO);
+	virtual bool CanLink(URigVMPin* InSourcePin, URigVMPin* InTargetPin, FString* OutFailureReason, const FRigVMByteCode* InByteCode, ERigVMPinDirection InUserLinkDirection = ERigVMPinDirection::IO, bool bInAllowWildcard = false);
 	
 	TSharedPtr<FRigVMParserAST> GetDiagnosticsAST(bool bForceRefresh = false, TArray<URigVMLink*> InLinksToSkip = TArray<URigVMLink*>());
 	TSharedPtr<FRigVMParserAST> GetRuntimeAST(const FRigVMParserASTSettings& InSettings = FRigVMParserASTSettings::Optimized(), bool bForceRefresh = false);
diff --git a/Engine/Source/Developer/RigVMDeveloper/Public/RigVMModel/RigVMLink.h b/Engine/Source/Developer/RigVMDeveloper/Public/RigVMModel/RigVMLink.h
--- a/Engine/Source/Developer/RigVMDeveloper/Public/RigVMModel/RigVMLink.h
+++ b/Engine/Source/Developer/RigVMDeveloper/Public/RigVMModel/RigVMLink.h
@@ -41,11 +41,15 @@ public:
 
 	// Returns the source Pin of this Link (or nullptr)
 	UFUNCTION(BlueprintCallable, Category = RigVMLink)
-	URigVMPin* GetSourcePin();
+	URigVMPin* GetSourcePin() const;
 
 	// Returns the target Pin of this Link (or nullptr)
 	UFUNCTION(BlueprintCallable, Category = RigVMLink)
-	URigVMPin* GetTargetPin();
+	URigVMPin* GetTargetPin() const;
+
+	// Returns the opposite Pin of this Link given one of its edges (or nullptr)
+	UFUNCTION(BlueprintCallable, Category = RigVMLink)
+	URigVMPin* GetOppositePin(const URigVMPin* InPin) const;
 
 	// Returns a string representation of the Link,
 	// for example: "NodeA.Color.R -> NodeB.Translation.X"
@@ -68,8 +72,8 @@ private:
 	UPROPERTY()
 	FString TargetPinPath;
 
-	URigVMPin* SourcePin;
-	URigVMPin* TargetPin;
+	mutable URigVMPin* SourcePin;
+	mutable URigVMPin* TargetPin;
 
 	friend class URigVMController;
 };
diff --git a/Engine/Source/Developer/RigVMDeveloper/Public/RigVMModel/RigVMPin.h b/Engine/Source/Developer/RigVMDeveloper/Public/RigVMModel/RigVMPin.h
--- a/Engine/Source/Developer/RigVMDeveloper/Public/RigVMModel/RigVMPin.h
+++ b/Engine/Source/Developer/RigVMDeveloper/Public/RigVMModel/RigVMPin.h
@@ -314,6 +314,12 @@ public:
 	UFUNCTION(BlueprintCallable, Category = RigVMPin)
 	URigVMPin* GetPinForLink() const;
 
+	// Returns the link that represents the connection
+	// between this pin and InOtherPin. nullptr is returned
+	// if the pins are not connected.
+	UFUNCTION(BlueprintCallable, Category = RigVMPin)
+	URigVMLink* FindLinkForPin(const URigVMPin* InOtherPin) const;
+
 	// Returns the original pin for a pin on an injected
 	// node. This can be used to determine where a link
 	// should go in the user interface
@@ -369,7 +375,7 @@ public:
 
 	// Returns true is the two provided source and target Pins
 	// can be linked to one another.
-	static bool CanLink(URigVMPin* InSourcePin, URigVMPin* InTargetPin, FString* OutFailureReason, const FRigVMByteCode* InByteCode, ERigVMPinDirection InUserLinkDirection = ERigVMPinDirection::IO);
+	static bool CanLink(URigVMPin* InSourcePin, URigVMPin* InTargetPin, FString* OutFailureReason, const FRigVMByteCode* InByteCode, ERigVMPinDirection InUserLinkDirection = ERigVMPinDirection::IO, bool bInAllowWildcard = false);
 
 	// Returns true if this pin has injected nodes
 	bool HasInjectedNodes() const { return InjectionInfos.Num() > 0; }
diff --git a/Engine/Source/Runtime/RigVM/Public/RigVMCore/RigVMTemplate.h b/Engine/Source/Runtime/RigVM/Public/RigVMCore/RigVMTemplate.h
--- a/Engine/Source/Runtime/RigVM/Public/RigVMCore/RigVMTemplate.h
+++ b/Engine/Source/Runtime/RigVM/Public/RigVMCore/RigVMTemplate.h
@@ -110,6 +110,16 @@ struct RIGVM_API FRigVMTemplateArgument
 			}
 			return CPPType;
 		}
+
+		void ConvertToArray() 
+		{
+			CPPType = RigVMTypeUtils::ArrayTypeFromBaseType(CPPType);
+		}
+
+		void ConvertToBaseElement() 
+		{
+			CPPType = RigVMTypeUtils::BaseTypeFromArrayType(CPPType);
+		}
 	};
 	
 	enum EArrayType
