Control Rig: Fix crash when opening certain rigs

diff --git a/Engine/Plugins/Animation/ControlRig/Source/ControlRigDeveloper/Private/ControlRigBlueprint.cpp b/Engine/Plugins/Animation/ControlRig/Source/ControlRigDeveloper/Private/ControlRigBlueprint.cpp
--- a/Engine/Plugins/Animation/ControlRig/Source/ControlRigDeveloper/Private/ControlRigBlueprint.cpp
+++ b/Engine/Plugins/Animation/ControlRig/Source/ControlRigDeveloper/Private/ControlRigBlueprint.cpp
@@ -913,7 +913,7 @@ void UControlRigBlueprint::RefreshAllModels()
 			Controller->RemoveUnusedOrphanedPins(ModelNode, false);
 		}
 
-		Controller->RecomputeAllTemplateFilteredTypes(false);
+		Controller->InitializeFilteredPermutationsFromTemplateTypes();
 	}
 }
 
diff --git a/Engine/Source/Developer/RigVMDeveloper/Private/RigVMModel/Nodes/RigVMTemplateNode.cpp b/Engine/Source/Developer/RigVMDeveloper/Private/RigVMModel/Nodes/RigVMTemplateNode.cpp
--- a/Engine/Source/Developer/RigVMDeveloper/Private/RigVMModel/Nodes/RigVMTemplateNode.cpp
+++ b/Engine/Source/Developer/RigVMDeveloper/Private/RigVMModel/Nodes/RigVMTemplateNode.cpp
@@ -266,12 +266,25 @@ bool URigVMTemplateNode::FilteredSupportsType(const URigVMPin* InPin, const FStr
 
 	const URigVMPin* RootPin = InPin;
 	bool bIsArrayElement = false;
+	bool bIsStructElement = false;
 	if (URigVMPin* ParentPin = InPin->GetParentPin())
 	{
 		RootPin = ParentPin;
-		bIsArrayElement = true;
+		if (ParentPin->IsArray())
+		{
+			bIsArrayElement = true;
+		}
+		else if (ParentPin->IsStruct())
+		{
+			bIsStructElement = true;
+		}
 	}
-	
+
+	if (bIsStructElement)
+	{
+		return InPin->GetCPPType() == InCPPType;
+	}
+
 	const FRigVMTemplateArgument* Argument = GetTemplate()->FindArgument(RootPin->GetFName());
 	if (Argument == nullptr)
 	{
@@ -527,11 +540,21 @@ TArray<int32> URigVMTemplateNode::GetNewFilteredPermutations(URigVMPin* InPin, U
 	NewFilteredPermutations.Reserve(FilteredPermutations.Num());
 	
 	bool bIsArrayElement = false;
+	bool bIsStructElement = false;
 	URigVMPin* RootPin = InPin;
 	if (URigVMPin* ParentPin = InPin->GetParentPin())
 	{
 		RootPin = ParentPin;
-		bIsArrayElement = true;
+		bIsArrayElement = RootPin->IsArray();
+		bIsStructElement = RootPin->IsStruct();
+	}
+
+	if (bIsStructElement)
+	{
+		if (InPin->GetCPPType() == LinkedPin->GetCPPType())
+		{
+			return FilteredPermutations;
+		}
 	}
 
 	TArray<int32> PermutationsToTry = PreferredPermutationTypes.IsEmpty()? FilteredPermutations : TArray<int32>({FindPermuationForTypes(PreferredPermutationTypes)});
@@ -740,6 +763,7 @@ void URigVMTemplateNode::InvalidateCache()
 {
 	SupportedTypesCache.Reset();
 	CachedFunction = nullptr;
+	CachedTemplate = nullptr;
 	ResolvedPermutations.Reset();
 
 	for(URigVMPin* Pin : GetPins())
@@ -770,3 +794,34 @@ void URigVMTemplateNode::InitializeFilteredPermutations()
 		}
 	}	
 }
+
+void URigVMTemplateNode::InitializeFilteredPermutationsFromTypes()
+{
+	if (IsSingleton())
+	{
+		return;
+	}
+
+	if (const FRigVMTemplate* Template = GetTemplate())
+	{
+		TArray<FString> ArgTypes;
+		for (int32 ArgIndex = 0; ArgIndex < Template->NumArguments(); ++ArgIndex)
+		{
+			const FRigVMTemplateArgument* Argument = Template->GetArgument(ArgIndex);
+			if (URigVMPin* Pin = FindPin(Argument->GetName().ToString()))
+			{
+				ArgTypes.Add(Argument->GetName().ToString() + TEXT(":") + Pin->GetCPPType());
+			}
+		}
+
+		int32 Permutation = FindPermuationForTypes(ArgTypes);
+		if (Permutation != INDEX_NONE)
+		{
+			FilteredPermutations = {Permutation};
+		}
+		else
+		{
+			InitializeFilteredPermutations();
+		}
+	}
+}
diff --git a/Engine/Source/Developer/RigVMDeveloper/Private/RigVMModel/RigVMController.cpp b/Engine/Source/Developer/RigVMDeveloper/Private/RigVMModel/RigVMController.cpp
--- a/Engine/Source/Developer/RigVMDeveloper/Private/RigVMModel/RigVMController.cpp
+++ b/Engine/Source/Developer/RigVMDeveloper/Private/RigVMModel/RigVMController.cpp
@@ -15539,6 +15539,30 @@ void URigVMController::RecomputeAllTemplateFilteredTypes(bool bSetupUndoRedo)
 	}
 }
 
+void URigVMController::InitializeFilteredPermutationsFromTemplateTypes()
+{
+	if (!IsValidGraph())
+	{
+		return;
+	}
+
+	if (!bIsTransacting && !IsGraphEditable())
+	{
+		return;
+	}
+
+	URigVMGraph* Graph = GetGraph();
+	check(Graph);
+
+	for (URigVMNode* Node : Graph->GetNodes())
+	{
+		if (URigVMTemplateNode* TemplateNode = Cast<URigVMTemplateNode>(Node))
+		{
+			TemplateNode->InitializeFilteredPermutationsFromTypes();
+		}
+	}
+}
+
 void URigVMController::InitializeAllTemplateFiltersInGraph(bool bSetupUndoRedo)
 {
 	if (!IsValidGraph())
diff --git a/Engine/Source/Developer/RigVMDeveloper/Public/RigVMModel/Nodes/RigVMTemplateNode.h b/Engine/Source/Developer/RigVMDeveloper/Public/RigVMModel/Nodes/RigVMTemplateNode.h
--- a/Engine/Source/Developer/RigVMDeveloper/Public/RigVMModel/Nodes/RigVMTemplateNode.h
+++ b/Engine/Source/Developer/RigVMDeveloper/Public/RigVMModel/Nodes/RigVMTemplateNode.h
@@ -87,6 +87,8 @@ public:
 
 	// initializes the filtered permutations to all possible permutations
 	void InitializeFilteredPermutations();
+
+	void InitializeFilteredPermutationsFromTypes();
 	
 protected:
 
diff --git a/Engine/Source/Developer/RigVMDeveloper/Public/RigVMModel/RigVMController.h b/Engine/Source/Developer/RigVMDeveloper/Public/RigVMModel/RigVMController.h
--- a/Engine/Source/Developer/RigVMDeveloper/Public/RigVMModel/RigVMController.h
+++ b/Engine/Source/Developer/RigVMDeveloper/Public/RigVMModel/RigVMController.h
@@ -895,6 +895,9 @@ public:
 	// Initializes and recomputes the filtered permutations of all template nodes in the graph 
 	void RecomputeAllTemplateFilteredTypes(bool bSetupUndoRedo);
 
+	// Try to initialize the filterd permutations from the pin types
+	void InitializeFilteredPermutationsFromTemplateTypes();
+
 	// Initializes filtered permuations to be unresolved in all template nodes in graph
 	void InitializeAllTemplateFiltersInGraph(bool bSetupUndoRedo);
 	
