Control Rig: Fix recompute filtered permutations from links

diff --git a/Engine/Plugins/Animation/ControlRig/Source/ControlRigDeveloper/Private/ControlRigBlueprint.cpp b/Engine/Plugins/Animation/ControlRig/Source/ControlRigDeveloper/Private/ControlRigBlueprint.cpp
--- a/Engine/Plugins/Animation/ControlRig/Source/ControlRigDeveloper/Private/ControlRigBlueprint.cpp
+++ b/Engine/Plugins/Animation/ControlRig/Source/ControlRigDeveloper/Private/ControlRigBlueprint.cpp
@@ -913,7 +913,14 @@ void UControlRigBlueprint::RefreshAllModels()
 			Controller->RemoveUnusedOrphanedPins(ModelNode, false);
 		}
 
-		Controller->InitializeFilteredPermutationsFromTemplateTypes();
+		for(URigVMNode* ModelNode : GraphToClean->GetNodes())
+		{
+			if (URigVMTemplateNode* TemplateNode = Cast<URigVMTemplateNode>(ModelNode))
+			{
+				TemplateNode->InvalidateCache();
+			}
+		}
+		Controller->RecomputeAllTemplateFilteredTypes(false);
 	}
 }
 
diff --git a/Engine/Source/Developer/RigVMDeveloper/Private/RigVMModel/Nodes/RigVMTemplateNode.cpp b/Engine/Source/Developer/RigVMDeveloper/Private/RigVMModel/Nodes/RigVMTemplateNode.cpp
--- a/Engine/Source/Developer/RigVMDeveloper/Private/RigVMModel/Nodes/RigVMTemplateNode.cpp
+++ b/Engine/Source/Developer/RigVMDeveloper/Private/RigVMModel/Nodes/RigVMTemplateNode.cpp
@@ -551,7 +551,7 @@ TArray<int32> URigVMTemplateNode::GetNewFilteredPermutations(URigVMPin* InPin, U
 
 	if (bIsStructElement)
 	{
-		if (InPin->GetCPPType() == LinkedPin->GetCPPType())
+		if (RigVMTypeUtils::AreCompatible(InPin->GetCPPType(), InPin->GetCPPTypeObject(), LinkedPin->GetCPPType(), LinkedPin->GetCPPTypeObject()))
 		{
 			return FilteredPermutations;
 		}
@@ -561,7 +561,7 @@ TArray<int32> URigVMTemplateNode::GetNewFilteredPermutations(URigVMPin* InPin, U
 	bool bLinkedIsTemplate = false;
 	if (URigVMTemplateNode* LinkedTemplate = Cast<URigVMTemplateNode>(LinkedPin->GetNode()))
 	{
-		if (!LinkedTemplate->IsSingleton())
+		if (!LinkedTemplate->IsSingleton() && !LinkedPin->IsStructMember())
 		{
 			bLinkedIsTemplate = true;
 			if (const FRigVMTemplateArgument* Argument = GetTemplate()->FindArgument(RootPin->GetFName()))
diff --git a/Engine/Source/Developer/RigVMDeveloper/Private/RigVMModel/RigVMController.cpp b/Engine/Source/Developer/RigVMDeveloper/Private/RigVMModel/RigVMController.cpp
--- a/Engine/Source/Developer/RigVMDeveloper/Private/RigVMModel/RigVMController.cpp
+++ b/Engine/Source/Developer/RigVMDeveloper/Private/RigVMModel/RigVMController.cpp
@@ -15063,7 +15063,7 @@ bool URigVMController::PrepareTemplatePinForType(URigVMPin* InPin, const TArray<
 				bBrokenLinks = false;
 				
 				// Initialize all template nodes in the graph
-				InitializeAllTemplateFiltersInGraph(true);
+				InitializeAllTemplateFiltersInGraph(true, false);
 
 				// Update our node with the requested type
 				UpdateFilteredPermutations(InPin, InTypes, true);
@@ -15450,6 +15450,7 @@ bool URigVMController::PropagateTemplateFilteredTypes(URigVMTemplateNode* InNode
 						}
 						else
 						{
+							ensure(bIsTransacting);
 							URigVMLink* Link = Pin->FindLinkForPin(OtherPin);
 							BreakLink(Link->GetSourcePin(), Link->GetTargetPin(), bSetupUndoRedo);
 							return false;
@@ -15501,13 +15502,86 @@ void URigVMController::RecomputeAllTemplateFilteredTypes(bool bSetupUndoRedo)
 
 	URigVMGraph* Graph = GetGraph();
 	check(Graph);
-	
-	InitializeAllTemplateFiltersInGraph(bSetupUndoRedo);
-	
-	for (URigVMLink* Link : Graph->GetLinks())
+
+	// Save all template pin types, in case we need them after initializing
+	TMap<URigVMPin*, FRigVMTemplateArgument::FType> TypesBeforeRecomputing;
+	for (URigVMNode* Node : Graph->GetNodes())
 	{
+		if (URigVMTemplateNode* TemplateNode = Cast<URigVMTemplateNode>(Node))
+		{
+			if (TemplateNode->IsSingleton())
+			{
+				continue;
+			}
+
+			if (!TemplateNode->PreferredPermutationTypes.IsEmpty())
+			{
+				continue;
+			}
+
+			if (const FRigVMTemplate* Template = TemplateNode->GetTemplate())
+			{
+				for (int32 i=0; i<Template->NumArguments(); ++i)
+				{
+					const FRigVMTemplateArgument* Argument = Template->GetArgument(i);
+					URigVMPin* Pin = TemplateNode->FindPin(Argument->GetName().ToString());
+					if (Pin->IsWildCard())
+					{
+						continue;
+					}
+
+					TypesBeforeRecomputing.Add(Pin, FRigVMTemplateArgument::FType(Pin->GetCPPType(), Pin->GetCPPTypeObject()));
+				}
+			}
+		}
+	}
+
+	// Initialize all filtered permutations to unresolved (except for preferred permutation)
+	InitializeAllTemplateFiltersInGraph(bSetupUndoRedo, false);
+
+	// Apply all links to update filtered permutations
+	for (int32 i=0; i<Graph->GetLinks().Num(); ++i)
+	{
+		URigVMLink* Link = Graph->GetLinks()[i];
+		
 		URigVMPin* OutputPin = Link->GetSourcePin();
 		URigVMPin* InputPin = Link->GetTargetPin();
+
+		// If pin is a struct member, we should resolve to that type
+		if (URigVMTemplateNode* OutputNode = Cast<URigVMTemplateNode>(OutputPin->GetNode()))
+		{
+			if (!OutputNode->IsSingleton())
+			{
+				if (OutputPin->IsStructMember())
+				{
+					URigVMPin* RootPin = OutputPin->GetRootPin();
+					const FRigVMTemplateArgument::FType Type = TypesBeforeRecomputing.FindChecked(RootPin);
+					if (UpdateFilteredPermutations(RootPin, {Type}, bSetupUndoRedo))
+					{
+						PropagateTemplateFilteredTypes(OutputNode, bSetupUndoRedo);
+					}
+				}
+			}
+		}
+
+		// If pin is a struct member, we should resolve to that type
+		if (URigVMTemplateNode* InputNode = Cast<URigVMTemplateNode>(InputPin->GetNode()))
+		{
+			if (!InputNode->IsSingleton())
+			{
+				if (InputPin->IsStructMember())
+				{
+					URigVMPin* RootPin = InputPin->GetRootPin();
+					const FRigVMTemplateArgument::FType Type = TypesBeforeRecomputing.FindChecked(RootPin);
+					if (UpdateFilteredPermutations(RootPin, {Type}, bSetupUndoRedo))
+					{
+						PropagateTemplateFilteredTypes(InputNode, bSetupUndoRedo);
+					}
+				}
+			}
+		}
+
+		// Propagate filtered types due to this link
 		if (URigVMTemplateNode* OutputNode = Cast<URigVMTemplateNode>(OutputPin->GetNode()))
 		{
 			if (!OutputNode->IsSingleton())
@@ -15516,7 +15590,6 @@ void URigVMController::RecomputeAllTemplateFilteredTypes(bool bSetupUndoRedo)
 				{
 					if (UpdateFilteredPermutations(OutputPin, InputPin, bSetupUndoRedo))
 					{
-						UpdateTemplateNodePinTypes(OutputNode, bSetupUndoRedo);
 						PropagateTemplateFilteredTypes(OutputNode, bSetupUndoRedo);
 					}
 				}
@@ -15530,40 +15603,29 @@ void URigVMController::RecomputeAllTemplateFilteredTypes(bool bSetupUndoRedo)
 				{
 					if (UpdateFilteredPermutations(InputPin, OutputPin, bSetupUndoRedo))
 					{
-						UpdateTemplateNodePinTypes(InputNode, bSetupUndoRedo);
 						PropagateTemplateFilteredTypes(InputNode, bSetupUndoRedo);
 					}
 				}
 			}
 		}				
 	}
-}
-
-void URigVMController::InitializeFilteredPermutationsFromTemplateTypes()
-{
-	if (!IsValidGraph())
-	{
-		return;
-	}
-
-	if (!bIsTransacting && !IsGraphEditable())
-	{
-		return;
-	}
-
-	URigVMGraph* Graph = GetGraph();
-	check(Graph);
 
+	// Now update all template nodes pin types
 	for (URigVMNode* Node : Graph->GetNodes())
 	{
 		if (URigVMTemplateNode* TemplateNode = Cast<URigVMTemplateNode>(Node))
 		{
-			TemplateNode->InitializeFilteredPermutationsFromTypes();
+			if (TemplateNode->IsSingleton())
+			{
+				continue;
+			}
+
+			UpdateTemplateNodePinTypes(TemplateNode, bSetupUndoRedo);
 		}
 	}
 }
 
-void URigVMController::InitializeAllTemplateFiltersInGraph(bool bSetupUndoRedo)
+void URigVMController::InitializeAllTemplateFiltersInGraph(bool bSetupUndoRedo, bool bChangePinTypes)
 {
 	if (!IsValidGraph())
 	{
@@ -15592,7 +15654,10 @@ void URigVMController::InitializeAllTemplateFiltersInGraph(bool bSetupUndoRedo)
 				
 				TArray<int32> OldPermutations = TemplateNode->FilteredPermutations;
 				TemplateNode->InitializeFilteredPermutations();
-				UpdateTemplateNodePinTypes(TemplateNode, bSetupUndoRedo);
+				if (bChangePinTypes)
+				{
+					UpdateTemplateNodePinTypes(TemplateNode, bSetupUndoRedo);
+				}
 				if (bSetupUndoRedo)
 				{
 					FRigVMSetTemplateFilteredPermutationsAction FilteringAction(TemplateNode, nullptr, OldPermutations);
diff --git a/Engine/Source/Developer/RigVMDeveloper/Public/RigVMModel/Nodes/RigVMTemplateNode.h b/Engine/Source/Developer/RigVMDeveloper/Public/RigVMModel/Nodes/RigVMTemplateNode.h
--- a/Engine/Source/Developer/RigVMDeveloper/Public/RigVMModel/Nodes/RigVMTemplateNode.h
+++ b/Engine/Source/Developer/RigVMDeveloper/Public/RigVMModel/Nodes/RigVMTemplateNode.h
@@ -119,6 +119,7 @@ protected:
 	mutable TArray<int32> ResolvedPermutations;
 
 	friend class URigVMController;
+	friend class UControlRigBlueprint;
 	friend struct FRigVMSetTemplateFilteredPermutationsAction;
 	friend struct FRigVMSetPreferredTemplatePermutationsAction;
 };
diff --git a/Engine/Source/Developer/RigVMDeveloper/Public/RigVMModel/RigVMController.h b/Engine/Source/Developer/RigVMDeveloper/Public/RigVMModel/RigVMController.h
--- a/Engine/Source/Developer/RigVMDeveloper/Public/RigVMModel/RigVMController.h
+++ b/Engine/Source/Developer/RigVMDeveloper/Public/RigVMModel/RigVMController.h
@@ -895,11 +895,8 @@ public:
 	// Initializes and recomputes the filtered permutations of all template nodes in the graph 
 	void RecomputeAllTemplateFilteredTypes(bool bSetupUndoRedo);
 
-	// Try to initialize the filterd permutations from the pin types
-	void InitializeFilteredPermutationsFromTemplateTypes();
-
 	// Initializes filtered permuations to be unresolved in all template nodes in graph
-	void InitializeAllTemplateFiltersInGraph(bool bSetupUndoRedo);
+	void InitializeAllTemplateFiltersInGraph(bool bSetupUndoRedo, bool bChangePinTypes);
 	
 #endif
 
