Control Rig: Back to initializing all template filtered permutations from types

diff --git a/Engine/Plugins/Animation/ControlRig/Source/ControlRigDeveloper/Private/ControlRigBlueprint.cpp b/Engine/Plugins/Animation/ControlRig/Source/ControlRigDeveloper/Private/ControlRigBlueprint.cpp
--- a/Engine/Plugins/Animation/ControlRig/Source/ControlRigDeveloper/Private/ControlRigBlueprint.cpp
+++ b/Engine/Plugins/Animation/ControlRig/Source/ControlRigDeveloper/Private/ControlRigBlueprint.cpp
@@ -920,7 +920,7 @@ void UControlRigBlueprint::RefreshAllModels()
 				TemplateNode->InvalidateCache();
 			}
 		}
-		Controller->RecomputeAllTemplateFilteredTypes(false);
+		Controller->InitializeFilteredPermutationsFromTemplateTypes();
 	}
 }
 
diff --git a/Engine/Source/Developer/RigVMDeveloper/Private/RigVMModel/RigVMController.cpp b/Engine/Source/Developer/RigVMDeveloper/Private/RigVMModel/RigVMController.cpp
--- a/Engine/Source/Developer/RigVMDeveloper/Private/RigVMModel/RigVMController.cpp
+++ b/Engine/Source/Developer/RigVMDeveloper/Private/RigVMModel/RigVMController.cpp
@@ -15450,7 +15450,7 @@ bool URigVMController::PropagateTemplateFilteredTypes(URigVMTemplateNode* InNode
 						}
 						else
 						{
-							ensure(bIsTransacting);
+							ensureMsgf(bIsTransacting, TEXT("Unexpected link broken in package %s"), *GetPackage()->GetPathName());
 							URigVMLink* Link = Pin->FindLinkForPin(OtherPin);
 							BreakLink(Link->GetSourcePin(), Link->GetTargetPin(), bSetupUndoRedo);
 							return false;
@@ -15625,6 +15625,30 @@ void URigVMController::RecomputeAllTemplateFilteredTypes(bool bSetupUndoRedo)
 	}
 }
 
+void URigVMController::InitializeFilteredPermutationsFromTemplateTypes()
+{
+	if (!IsValidGraph())
+	{
+		return;
+	}
+
+	if (!bIsTransacting && !IsGraphEditable())
+	{
+		return;
+	}
+
+	URigVMGraph* Graph = GetGraph();
+	check(Graph);
+
+	for (URigVMNode* Node : Graph->GetNodes())
+	{
+		if (URigVMTemplateNode* TemplateNode = Cast<URigVMTemplateNode>(Node))
+		{
+			TemplateNode->InitializeFilteredPermutationsFromTypes();
+		}
+	}
+}
+
 void URigVMController::InitializeAllTemplateFiltersInGraph(bool bSetupUndoRedo, bool bChangePinTypes)
 {
 	if (!IsValidGraph())
diff --git a/Engine/Source/Developer/RigVMDeveloper/Public/RigVMModel/RigVMController.h b/Engine/Source/Developer/RigVMDeveloper/Public/RigVMModel/RigVMController.h
--- a/Engine/Source/Developer/RigVMDeveloper/Public/RigVMModel/RigVMController.h
+++ b/Engine/Source/Developer/RigVMDeveloper/Public/RigVMModel/RigVMController.h
@@ -895,6 +895,9 @@ public:
 	// Initializes and recomputes the filtered permutations of all template nodes in the graph 
 	void RecomputeAllTemplateFilteredTypes(bool bSetupUndoRedo);
 
+	// Try to initialize the filterd permutations from the pin types
+	void InitializeFilteredPermutationsFromTemplateTypes();
+
 	// Initializes filtered permuations to be unresolved in all template nodes in graph
 	void InitializeAllTemplateFiltersInGraph(bool bSetupUndoRedo, bool bChangePinTypes);
 	
