Control Rig: Test recompute all filtered permutations on templates

diff --git a/Engine/Source/Developer/RigVMDeveloper/Private/RigVMModel/Nodes/RigVMTemplateNode.cpp b/Engine/Source/Developer/RigVMDeveloper/Private/RigVMModel/Nodes/RigVMTemplateNode.cpp
--- a/Engine/Source/Developer/RigVMDeveloper/Private/RigVMModel/Nodes/RigVMTemplateNode.cpp
+++ b/Engine/Source/Developer/RigVMDeveloper/Private/RigVMModel/Nodes/RigVMTemplateNode.cpp
@@ -557,7 +557,7 @@ TArray<int32> URigVMTemplateNode::GetNewFilteredPermutations(URigVMPin* InPin, U
 		}
 	}
 
-	TArray<int32> PermutationsToTry = PreferredPermutationTypes.IsEmpty()? FilteredPermutations : TArray<int32>({FindPermuationForTypes(PreferredPermutationTypes)});
+	TArray<int32> PermutationsToTry = PreferredPermutationTypes.IsEmpty()? FilteredPermutations : TArray<int32>(FindPermuationsForTypes(PreferredPermutationTypes));
 	bool bLinkedIsTemplate = false;
 	if (URigVMTemplateNode* LinkedTemplate = Cast<URigVMTemplateNode>(LinkedPin->GetNode()))
 	{
@@ -619,7 +619,7 @@ TArray<int32> URigVMTemplateNode::GetNewFilteredPermutations(URigVMPin* InPin, c
 		}
 	}
 
-	TArray<int32> PermutationsToTry = PreferredPermutationTypes.IsEmpty()? FilteredPermutations : TArray<int32>({FindPermuationForTypes(PreferredPermutationTypes)});
+	TArray<int32> PermutationsToTry = PreferredPermutationTypes.IsEmpty()? FilteredPermutations : TArray<int32>(FindPermuationsForTypes(PreferredPermutationTypes));
 	
 	if (const FRigVMTemplateArgument* Argument = GetTemplate()->FindArgument(RootPin->GetFName()))
 	{
@@ -639,8 +639,9 @@ TArray<int32> URigVMTemplateNode::GetNewFilteredPermutations(URigVMPin* InPin, c
 	return NewFilteredPermutations;
 }
 
-int32 URigVMTemplateNode::FindPermuationForTypes(const TArray<FString>& ArgumentTypes)
+TArray<int32> URigVMTemplateNode::FindPermuationsForTypes(const TArray<FString>& ArgumentTypes)
 {
+	TArray<int32> Permutations;
 	if (const FRigVMTemplate* Template = GetTemplate())
 	{
 		for (int32 i=0; i<Template->NumPermutations(); ++i)
@@ -651,7 +652,7 @@ int32 URigVMTemplateNode::FindPermuationForTypes(const TArray<FString>& Argument
 				FString ArgName, Type;
 				if (!TypePairString.Split(TEXT(":"), &ArgName, &Type))
 				{
-					return INDEX_NONE;
+					return {};
 				}
 				
 				if (const FRigVMTemplateArgument* Argument = Template->FindArgument(*ArgName))
@@ -664,17 +665,17 @@ int32 URigVMTemplateNode::FindPermuationForTypes(const TArray<FString>& Argument
 				}
 				else
 				{
-					return INDEX_NONE;
+					return {};
 				}
 			}
 
 			if (bAllArgsMatched)
 			{
-				return i;
+				Permutations.Add(i);
 			}
 		}
 	}
-	return INDEX_NONE;
+	return Permutations;
 }
 
 TArray<FString> URigVMTemplateNode::GetArgumentTypesForPermutation(const int32 InPermutationIndex)
@@ -782,7 +783,7 @@ void URigVMTemplateNode::InitializeFilteredPermutations()
 	{
 		if (!PreferredPermutationTypes.IsEmpty())
 		{
-			FilteredPermutations = {FindPermuationForTypes(PreferredPermutationTypes)};
+			FilteredPermutations = FindPermuationsForTypes(PreferredPermutationTypes);
 		}
 		else
 		{
@@ -810,14 +811,18 @@ void URigVMTemplateNode::InitializeFilteredPermutationsFromTypes()
 			const FRigVMTemplateArgument* Argument = Template->GetArgument(ArgIndex);
 			if (URigVMPin* Pin = FindPin(Argument->GetName().ToString()))
 			{
-				ArgTypes.Add(Argument->GetName().ToString() + TEXT(":") + Pin->GetCPPType());
+				if (!Pin->IsWildCard())
+				{
+					ArgTypes.Add(Argument->GetName().ToString() + TEXT(":") + Pin->GetCPPType());
+				}
 			}
 		}
 
-		int32 Permutation = FindPermuationForTypes(ArgTypes);
-		if (Permutation != INDEX_NONE)
+		TArray<int32> Permutations = FindPermuationsForTypes(ArgTypes);
+		if (!Permutations.IsEmpty())
 		{
-			FilteredPermutations = {Permutation};
+			FilteredPermutations = Permutations;
+			PreferredPermutationTypes = ArgTypes;
 		}
 		else
 		{
diff --git a/Engine/Source/Developer/RigVMDeveloper/Private/RigVMModel/RigVMController.cpp b/Engine/Source/Developer/RigVMDeveloper/Private/RigVMModel/RigVMController.cpp
--- a/Engine/Source/Developer/RigVMDeveloper/Private/RigVMModel/RigVMController.cpp
+++ b/Engine/Source/Developer/RigVMDeveloper/Private/RigVMModel/RigVMController.cpp
@@ -3541,7 +3541,7 @@ TArray<FName> URigVMController::ImportNodesFromText(const FString& InText, bool
 					}
 					if (!TemplateNode->PreferredPermutationTypes.IsEmpty())
 					{
-						TemplateNode->FilteredPermutations = {TemplateNode->FindPermuationForTypes(TemplateNode->PreferredPermutationTypes)};
+						TemplateNode->FilteredPermutations = TemplateNode->FindPermuationsForTypes(TemplateNode->PreferredPermutationTypes);
 					}
 					UpdateTemplateNodePinTypes(TemplateNode, bSetupUndoRedo);
 				}
@@ -15470,7 +15470,7 @@ bool URigVMController::PropagateTemplateFilteredTypes(URigVMTemplateNode* InNode
 			return false;
 		}
 
-		if (Pin->GetSubPins().Num() > 0)
+		if (Pin->IsArray())
 		{
 			for (FRigVMTemplateArgument::FType& Type : Types)
 			{
@@ -15488,7 +15488,7 @@ bool URigVMController::PropagateTemplateFilteredTypes(URigVMTemplateNode* InNode
 	return true;
 }
 
-void URigVMController::RecomputeAllTemplateFilteredTypes(bool bSetupUndoRedo)
+void URigVMController::RecomputeAllTemplateFilteredTypes(bool bSetupUndoRedo, bool bKeepResolvedTypes)
 {
 	if (!IsValidGraph())
 	{
@@ -15620,6 +15620,26 @@ void URigVMController::RecomputeAllTemplateFilteredTypes(bool bSetupUndoRedo)
 				continue;
 			}
 
+			if (bKeepResolvedTypes && TemplateNode->FilteredPermutations.Num() > 1)
+			{
+				if (const FRigVMTemplate* Template = TemplateNode->GetTemplate())
+				{
+					TArray<FString> Types;
+					for (int32 i=0; i<Template->NumArguments(); ++i)
+					{
+						const FRigVMTemplateArgument* Argument = Template->GetArgument(i);
+						URigVMPin* Pin = TemplateNode->FindPin(Argument->GetName().ToString());
+						if (const FRigVMTemplateArgument::FType* Type = TypesBeforeRecomputing.Find(Pin))
+						{
+							Types.Add(Argument->GetName().ToString() + TEXT(":") + Type->CPPType);
+						}						
+					}
+
+					TemplateNode->FilteredPermutations = TemplateNode->FindPermuationsForTypes(Types);
+					TemplateNode->PreferredPermutationTypes = Types;
+				}
+			}
+
 			UpdateTemplateNodePinTypes(TemplateNode, bSetupUndoRedo);
 		}
 	}
diff --git a/Engine/Source/Developer/RigVMDeveloper/Public/RigVMModel/Nodes/RigVMTemplateNode.h b/Engine/Source/Developer/RigVMDeveloper/Public/RigVMModel/Nodes/RigVMTemplateNode.h
--- a/Engine/Source/Developer/RigVMDeveloper/Public/RigVMModel/Nodes/RigVMTemplateNode.h
+++ b/Engine/Source/Developer/RigVMDeveloper/Public/RigVMModel/Nodes/RigVMTemplateNode.h
@@ -88,6 +88,7 @@ public:
 	// initializes the filtered permutations to all possible permutations
 	void InitializeFilteredPermutations();
 
+	// Initializes the filtered permutations and preferred permutation from the types of the pins
 	void InitializeFilteredPermutationsFromTypes();
 	
 protected:
@@ -97,7 +98,7 @@ protected:
 	TArray<int32> GetNewFilteredPermutations(URigVMPin* InPin, URigVMPin* LinkedPin);
 	TArray<int32> GetNewFilteredPermutations(URigVMPin* InPin, const TArray<FRigVMTemplateArgument::FType>& InTypes);
 
-	int32 FindPermuationForTypes(const TArray<FString>& ArgumentTypes);
+	TArray<int32> FindPermuationsForTypes(const TArray<FString>& ArgumentTypes);
 	TArray<FString> GetArgumentTypesForPermutation(const int32 InPermutationIndex);
 
 	UPROPERTY()
diff --git a/Engine/Source/Developer/RigVMDeveloper/Public/RigVMModel/RigVMController.h b/Engine/Source/Developer/RigVMDeveloper/Public/RigVMModel/RigVMController.h
--- a/Engine/Source/Developer/RigVMDeveloper/Public/RigVMModel/RigVMController.h
+++ b/Engine/Source/Developer/RigVMDeveloper/Public/RigVMModel/RigVMController.h
@@ -893,7 +893,7 @@ public:
 	bool RemoveUnusedOrphanedPins(URigVMNode* InNode, bool bNotify);
 
 	// Initializes and recomputes the filtered permutations of all template nodes in the graph 
-	void RecomputeAllTemplateFilteredTypes(bool bSetupUndoRedo);
+	void RecomputeAllTemplateFilteredTypes(bool bSetupUndoRedo, bool bKeepResolvedTypes = false);
 
 	// Try to initialize the filterd permutations from the pin types
 	void InitializeFilteredPermutationsFromTemplateTypes();
