Control Rig: Fix crash when recomputing filtered permutations and preferred types is set

diff --git a/Engine/Source/Developer/RigVMDeveloper/Private/RigVMModel/RigVMController.cpp b/Engine/Source/Developer/RigVMDeveloper/Private/RigVMModel/RigVMController.cpp
--- a/Engine/Source/Developer/RigVMDeveloper/Private/RigVMModel/RigVMController.cpp
+++ b/Engine/Source/Developer/RigVMDeveloper/Private/RigVMModel/RigVMController.cpp
@@ -15472,15 +15472,18 @@ bool URigVMController::PropagateTemplateFilteredTypes(URigVMTemplateNode* InNode
 
 		if (Pin->IsArray())
 		{
-			for (FRigVMTemplateArgument::FType& Type : Types)
+			if (Pin->GetSubPins().Num() > 0)
 			{
-				Type.ConvertToBaseElement();
-			}
-			for (URigVMPin* SubPin : Pin->GetSubPins())
-			{
-				if (!UpdateAndPropagete(SubPin, Types))
+				for (FRigVMTemplateArgument::FType& Type : Types)
 				{
-					return false;
+					Type.ConvertToBaseElement();
+				}
+				for (URigVMPin* SubPin : Pin->GetSubPins())
+				{
+					if (!UpdateAndPropagete(SubPin, Types))
+					{
+						return false;
+					}
 				}
 			}
 		}
@@ -15550,7 +15553,7 @@ void URigVMController::RecomputeAllTemplateFilteredTypes(bool bSetupUndoRedo, bo
 		// If pin is a struct member, we should resolve to that type
 		if (URigVMTemplateNode* OutputNode = Cast<URigVMTemplateNode>(OutputPin->GetNode()))
 		{
-			if (!OutputNode->IsSingleton())
+			if (!OutputNode->IsSingleton() && OutputNode->PreferredPermutationTypes.IsEmpty())
 			{
 				if (OutputPin->IsStructMember())
 				{
@@ -15569,7 +15572,7 @@ void URigVMController::RecomputeAllTemplateFilteredTypes(bool bSetupUndoRedo, bo
 		{
 			if (!InputNode->IsSingleton())
 			{
-				if (InputPin->IsStructMember())
+				if (InputPin->IsStructMember() && InputNode->PreferredPermutationTypes.IsEmpty())
 				{
 					URigVMPin* RootPin = InputPin->GetRootPin();
 					const FRigVMTemplateArgument::FType Type = TypesBeforeRecomputing.FindChecked(RootPin);
