Control Rig: Fix adding/resolving/unresolving template nodes either not printing Python commands or printing wrong python command

diff --git a/Engine/Plugins/Animation/ControlRig/Source/ControlRigEditor/Private/ControlRigEditorModule.cpp b/Engine/Plugins/Animation/ControlRig/Source/ControlRigEditor/Private/ControlRigEditorModule.cpp
--- a/Engine/Plugins/Animation/ControlRig/Source/ControlRigEditor/Private/ControlRigEditorModule.cpp
+++ b/Engine/Plugins/Animation/ControlRig/Source/ControlRigEditor/Private/ControlRigEditorModule.cpp
@@ -1405,7 +1405,7 @@ void FControlRigEditorModule::GetContextMenuActions(const UControlRigGraphSchema
 								FSlateIcon(),
 								FUIAction(FExecuteAction::CreateLambda([Controller, ModelPin]() {
 									const TArray<FName> Nodes = ModelPin->GetGraph()->GetSelectNodes();
-									Controller->UnresolveTemplateNodes(Nodes);
+									Controller->UnresolveTemplateNodes(Nodes, true, true);
 								})
 							));
 						}
@@ -2341,7 +2341,7 @@ void FControlRigEditorModule::GetContextMenuActions(const UControlRigGraphSchema
 								FSlateIcon(),
 								FUIAction(FExecuteAction::CreateLambda([Controller, Model]() {
 									const TArray<FName> Nodes = Model->GetSelectNodes();
-									Controller->UnresolveTemplateNodes(Nodes);
+									Controller->UnresolveTemplateNodes(Nodes, true, true);
 								})
 							));
 						}
diff --git a/Engine/Plugins/Animation/ControlRig/Source/ControlRigEditor/Private/Graph/SControlRigGraphChangePinType.cpp b/Engine/Plugins/Animation/ControlRig/Source/ControlRigEditor/Private/Graph/SControlRigGraphChangePinType.cpp
--- a/Engine/Plugins/Animation/ControlRig/Source/ControlRigEditor/Private/Graph/SControlRigGraphChangePinType.cpp
+++ b/Engine/Plugins/Animation/ControlRig/Source/ControlRigEditor/Private/Graph/SControlRigGraphChangePinType.cpp
@@ -394,7 +394,7 @@ void SControlRigChangePinType::HandlePinTypeChanged(FRigVMTemplateArgument::FTyp
 			continue;
 		}
 
-		Controller->ResolveWildCardPin(ModelPin->GetPinPath(), InType.CPPType, InType.GetCPPTypeObjectPath());
+		Controller->ResolveWildCardPin(ModelPin->GetPinPath(), InType.CPPType, InType.GetCPPTypeObjectPath(), true, true);
 	}
 }
 
diff --git a/Engine/Source/Developer/RigVMDeveloper/Private/RigVMModel/RigVMController.cpp b/Engine/Source/Developer/RigVMDeveloper/Private/RigVMModel/RigVMController.cpp
--- a/Engine/Source/Developer/RigVMDeveloper/Private/RigVMModel/RigVMController.cpp
+++ b/Engine/Source/Developer/RigVMDeveloper/Private/RigVMModel/RigVMController.cpp
@@ -428,7 +428,7 @@ TArray<FString> URigVMController::GetAddNodePythonCommands(URigVMNode* Node) con
 					*InjectionInfoOutputPinName, 
 					*UnitNode->GetName()));
 		}
-		else
+		else if (UnitNode->IsSingleton())
 		{
 			// add_struct_node_from_struct_path(script_struct_path, method_name, position=[0.0, 0.0], node_name='', undo=True)
 			Commands.Add(FString::Printf(TEXT("blueprint.get_controller_by_name('%s').add_unit_node_from_struct_path('%s', 'Execute', %s, '%s')"),
@@ -437,6 +437,49 @@ TArray<FString> URigVMController::GetAddNodePythonCommands(URigVMNode* Node) con
 					*RigVMPythonUtils::Vector2DToPythonString(UnitNode->GetPosition()),
 					*NodeName));
 		}
+		else
+		{
+			// add_template_node(notation, position=[0.0, 0.0], node_name='', undo=True)
+			Commands.Add(FString::Printf(TEXT("blueprint.get_controller_by_name('%s').add_template_node('%s', %s, '%s')"),
+							*GraphName,
+							*UnitNode->GetNotation().ToString(),
+							*RigVMPythonUtils::Vector2DToPythonString(UnitNode->GetPosition()),
+							*NodeName));
+
+			// Try to resolve wildcard pins			
+			if (const FRigVMTemplate* Template = UnitNode->GetTemplate())
+			{
+				// Lets minimize the number of commands by stopping when the number of permutations left is 1 (or less)
+				TArray<int32> Permutations;
+				Permutations.SetNumUninitialized(Template->NumPermutations());
+				FRigVMTemplate::FTypeMap TypeMap;
+		
+				for (int32 ArgIndex = 0; ArgIndex < Template->NumArguments(); ++ArgIndex)
+				{
+					if (Permutations.Num() < 2)
+					{
+						break;
+					}
+			
+					const FRigVMTemplateArgument* Argument = Template->GetArgument(ArgIndex);
+					if (!Argument->IsSingleton())
+					{
+						URigVMPin* Pin = UnitNode->FindPin(Argument->GetName().ToString());
+						if (!Pin->IsWildCard())
+						{
+							Commands.Add(FString::Printf(TEXT("blueprint.get_controller_by_name('%s').resolve_wild_card_pin('%s', '%s', '%s')"),
+									*GraphName,
+									*Pin->GetPinPath(),
+									*Pin->GetCPPType(),
+									*Pin->GetCPPTypeObject()->GetPathName()));
+
+							TypeMap.Add(Argument->GetName(), FRigVMTemplateArgument::FType(Pin->GetCPPType(), Pin->GetCPPTypeObject()));
+							Template->Resolve(TypeMap, Permutations, false);
+						}
+					}
+				}
+			}			
+		}		
 	}
 	else if (const URigVMAggregateNode* AggregateNode = Cast<URigVMAggregateNode>(Node))
 	{
@@ -570,15 +613,6 @@ TArray<FString> URigVMController::GetAddNodePythonCommands(URigVMNode* Node) con
 					*NodeName));	
 		}
 	}
-	else if (const URigVMTemplateNode* TemplateNode = Cast<URigVMTemplateNode>(Node))
-	{
-		// add_template_node(notation, position=[0.0, 0.0], node_name='', undo=True)
-		Commands.Add(FString::Printf(TEXT("blueprint.get_controller_by_name('%s').add_template_node('%s', %s, '%s')"),
-						*GraphName,
-						*TemplateNode->GetNotation().ToString(),
-						*RigVMPythonUtils::Vector2DToPythonString(TemplateNode->GetPosition()),
-						*NodeName));
-	}
 	else if (const URigVMEnumNode* EnumNode = Cast<URigVMEnumNode>(Node))
 	{
 		// add_enum_node(cpp_type_object_path, position=[0.0, 0.0], node_name='', undo=True)
@@ -15190,7 +15224,23 @@ bool URigVMController::ResolveWildCardPin(const FString& InPinPath, const FStrin
 
 	if (URigVMPin* Pin = Graph->FindPin(InPinPath))
 	{
-		return ResolveWildCardPin(Pin, FRigVMTemplateArgument::FType(CPPType, CPPTypeObject), bSetupUndoRedo, bPrintPythonCommand);
+		if (ResolveWildCardPin(Pin, FRigVMTemplateArgument::FType(CPPType, CPPTypeObject), bSetupUndoRedo, bPrintPythonCommand))
+		{
+			if (bPrintPythonCommand)
+			{
+				const FString GraphName = GetSanitizedGraphName(GetGraph()->GetGraphName());
+
+				// bool ResolveWildCardPin(const FString& InPinPath, const FString& InCPPType, const FName& InCPPTypeObjectPath, bool bSetupUndoRedo = true, bool bPrintPythonCommand = false);
+				RigVMPythonUtils::Print(GetGraphOuterName(), 
+									FString::Printf(TEXT("blueprint.get_controller_by_name('%s').resolve_wild_card_pin('%s', '%s', '%s')"),
+													*GraphName,
+													*InPinPath,
+													*InCPPType,
+													*InCPPTypeObjectPath.ToString()));
+			}
+			
+			return true;
+		}
 	}
 
 	return false;
@@ -15243,11 +15293,20 @@ bool URigVMController::ResolveWildCardPin(URigVMPin* InPin, const FRigVMTemplate
 		{
 			ActionStack->EndAction(Action);
 		}
-
-		// Print python command
 		
 		return true;		
 	}
+	else
+	{
+		OpenUndoBracket(TEXT("Resolving wildcard pin"));
+		UnresolveTemplateNodes({*InPin->GetNode()->GetName()}, bSetupUndoRedo);
+		if (ResolveWildCardPin(InPin, InType, bSetupUndoRedo, bPrintPythonCommand))
+		{
+			CloseUndoBracket();
+			return true;
+		}
+		CancelUndoBracket();
+	}
 
 	return false;
 }
