RigVM: Fix deletion of functions

Resulted from bad merge

diff --git a/Engine/Source/Developer/RigVMDeveloper/Private/RigVMModel/RigVMController.cpp b/Engine/Source/Developer/RigVMDeveloper/Private/RigVMModel/RigVMController.cpp
--- a/Engine/Source/Developer/RigVMDeveloper/Private/RigVMModel/RigVMController.cpp
+++ b/Engine/Source/Developer/RigVMDeveloper/Private/RigVMModel/RigVMController.cpp
@@ -5912,15 +5912,8 @@ bool URigVMController::RemoveNode(URigVMNode* InNode, bool bSetupUndoRedo, bool
 						
 						if (FunctionReferencePtr.IsValid())
 						{
-							{
-								FRigVMControllerGraphGuard GraphGuard(this, FunctionReferencePtr->GetGraph(), bSetupUndoRedo);
-								RemoveNode(FunctionReferencePtr.Get());
-
-								TGuardValue<TSoftObjectPtr<URigVMLibraryNode>> ClearReferencedNodePtr(
-									FunctionReferencePtr->ReferencedNodePtr,
-									TSoftObjectPtr<URigVMLibraryNode>());
-							}
-							FunctionReferencePtr->ReferencedNodePtr.ResetWeakPtr();
+							FRigVMControllerGraphGuard GraphGuard(this, FunctionReferencePtr->GetGraph(), bSetupUndoRedo);
+							RemoveNode(FunctionReferencePtr.Get());
 						}
 					}
 				}
