RigVM: Promoting CollapseNode To Function to respect execute pin direction

diff --git a/Engine/Source/Developer/RigVMDeveloper/Private/RigVMModel/RigVMController.cpp b/Engine/Source/Developer/RigVMDeveloper/Private/RigVMModel/RigVMController.cpp
--- a/Engine/Source/Developer/RigVMDeveloper/Private/RigVMModel/RigVMController.cpp
+++ b/Engine/Source/Developer/RigVMDeveloper/Private/RigVMModel/RigVMController.cpp
@@ -5443,7 +5443,13 @@ URigVMFunctionReferenceNode* URigVMController::PromoteCollapseNodeToFunctionRefe
 		{
 			FRigVMControllerGraphGuard GraphGuard(this, FunctionLibrary, false);
 			const FString FunctionName = GetValidNodeName(InCollapseNode->GetName());
-			FunctionDefinition = AddFunctionToLibrary(*FunctionName, InCollapseNode->IsMutable(), FVector2D::ZeroVector, false);		
+			FunctionDefinition = AddFunctionToLibrary(
+				*FunctionName,
+				InCollapseNode->GetPins().ContainsByPredicate([](const URigVMPin* Pin) -> bool
+				{
+					return Pin->IsExecuteContext() && (Pin->GetDirection() == ERigVMPinDirection::IO);
+				}),
+				FVector2D::ZeroVector, false);		
 		}
 	
 		// Add interface pins in function
