Control Rig: Fix crash when pasting function reference node from another project

diff --git a/Engine/Source/Developer/RigVMDeveloper/Private/RigVMModel/RigVMController.cpp b/Engine/Source/Developer/RigVMDeveloper/Private/RigVMModel/RigVMController.cpp
--- a/Engine/Source/Developer/RigVMDeveloper/Private/RigVMModel/RigVMController.cpp
+++ b/Engine/Source/Developer/RigVMDeveloper/Private/RigVMModel/RigVMController.cpp
@@ -3501,7 +3501,10 @@ TArray<FName> URigVMController::ImportNodesFromText(const FString& InText, bool
 	{
 		if (URigVMLibraryNode* LibraryNode = Cast<URigVMLibraryNode>(CreatedNode))
 		{
-			EditGuards.Emplace(LibraryNode->GetContainedGraph()->bEditable, true);
+			if (URigVMGraph* ContainedGraph = LibraryNode->GetContainedGraph())
+			{
+				EditGuards.Emplace(ContainedGraph->bEditable, true);
+			}
 		}
 	}
 
