RigVMDeveloper: Centralize pythonize for enum values

diff --git a/Engine/Plugins/Animation/ControlRig/Source/ControlRig/Private/Rigs/RigHierarchy.cpp b/Engine/Plugins/Animation/ControlRig/Source/ControlRig/Private/Rigs/RigHierarchy.cpp
--- a/Engine/Plugins/Animation/ControlRig/Source/ControlRig/Private/Rigs/RigHierarchy.cpp
+++ b/Engine/Plugins/Animation/ControlRig/Source/ControlRig/Private/Rigs/RigHierarchy.cpp
@@ -3168,18 +3168,11 @@ void URigHierarchy::SetControlValue(FRigControlElement* InControlElement, const
 					}
 					if (!BlueprintName.IsEmpty())
 					{
-						FString TypeStr;
-						switch (InValueType)
-						{
-						case ERigControlValueType::Minimum: TypeStr = TEXT("MINIMUM"); break;
-						case ERigControlValueType::Maximum: TypeStr = TEXT("MAXIMUM"); break;
-						default: ensure(false);
-						}
 						RigVMPythonUtils::Print(BlueprintName,
-							FString::Printf(TEXT("hierarchy.set_control_value(%s, %s, unreal.RigControlValueType.%s)"),
+							FString::Printf(TEXT("hierarchy.set_control_value(%s, %s, %s)"),
 							*InControlElement->GetKey().ToPythonString(),
 							*InValue.ToPythonString(InControlElement->Settings.ControlType),
-							*TypeStr));
+							*RigVMPythonUtils::EnumValueToPythonString<ERigControlValueType>((int64)InValueType)));
 					}
 				}
 #endif
@@ -5423,9 +5416,24 @@ TArray<FString> URigHierarchy::ControlSettingsToPythonCommands(const FRigControl
 	TArray<FString> Commands;
 	Commands.Add(FString::Printf(TEXT("%s = unreal.RigControlSettings()"),
 			*NameSettings));
-	
-	const FString AnimationTypeStr = StaticEnum<ERigControlAnimationType>()->GetNameStringByValue((int64)Settings.AnimationType).ToUpper();
-	const FString ControlTypeStr = StaticEnum<ERigControlType>()->GetNameStringByValue((int64)Settings.ControlType).ToUpper();
+
+	ERigControlType ControlType = Settings.ControlType;
+	switch(ControlType)
+	{
+		case ERigControlType::Transform:
+		case ERigControlType::TransformNoScale:
+		{
+			ControlType = ERigControlType::EulerTransform;
+			break;
+		}
+		default:
+		{
+			break;
+		}
+	}
+
+	const FString AnimationTypeStr = RigVMPythonUtils::EnumValueToPythonString<ERigControlAnimationType>((int64)Settings.AnimationType);
+	const FString ControlTypeStr = RigVMPythonUtils::EnumValueToPythonString<ERigControlType>((int64)ControlType);
 
 	static const TCHAR* TrueText = TEXT("True");
 	static const TCHAR* FalseText = TEXT("False");
@@ -5440,10 +5448,10 @@ TArray<FString> URigHierarchy::ControlSettingsToPythonCommands(const FRigControl
 	
 	const FString LimitEnabledStr = FString::Join(LimitEnabledParts, TEXT(", "));
 	
-	Commands.Add(FString::Printf(TEXT("%s.animation_type = unreal.RigControlAnimationType.%s"),
+	Commands.Add(FString::Printf(TEXT("%s.animation_type = %s"),
 									*NameSettings,
 									*AnimationTypeStr));
-	Commands.Add(FString::Printf(TEXT("%s.control_type = unreal.RigControlType.%s"),
+	Commands.Add(FString::Printf(TEXT("%s.control_type = %s"),
 									*NameSettings,
 									*ControlTypeStr));
 	Commands.Add(FString::Printf(TEXT("%s.display_name = '%s'"),
@@ -5473,9 +5481,9 @@ TArray<FString> URigHierarchy::ControlSettingsToPythonCommands(const FRigControl
 	Commands.Add(FString::Printf(TEXT("%s.maximum_value = %s"),
 		*NameSettings,
 		*Settings.MaximumValue.ToPythonString(Settings.ControlType)));
-	Commands.Add(FString::Printf(TEXT("%s.primary_axis = unreal.RigControlAxis.%s"),
+	Commands.Add(FString::Printf(TEXT("%s.primary_axis = %s"),
 		*NameSettings,
-		Settings.PrimaryAxis == ERigControlAxis::X ? TEXT("X") : Settings.PrimaryAxis == ERigControlAxis::Y ? TEXT("Y") : TEXT("Z")));
+		*RigVMPythonUtils::EnumValueToPythonString<ERigControlAxis>((int64)Settings.PrimaryAxis)));
 
 	return Commands;
 }
diff --git a/Engine/Plugins/Animation/ControlRig/Source/ControlRig/Private/Rigs/RigHierarchyController.cpp b/Engine/Plugins/Animation/ControlRig/Source/ControlRig/Private/Rigs/RigHierarchyController.cpp
--- a/Engine/Plugins/Animation/ControlRig/Source/ControlRig/Private/Rigs/RigHierarchyController.cpp
+++ b/Engine/Plugins/Animation/ControlRig/Source/ControlRig/Private/Rigs/RigHierarchyController.cpp
@@ -1425,11 +1425,12 @@ TArray<FString> URigHierarchyController::GetAddBonePythonCommands(FRigBoneElemen
 	}
 	
 	// AddBone(FName InName, FRigElementKey InParent, FTransform InTransform, bool bTransformInGlobal = true, ERigBoneType InBoneType = ERigBoneType::User, bool bSetupUndo = false);
-	Commands.Add(FString::Printf(TEXT("hierarchy_controller.add_bone('%s', %s, %s, False, unreal.RigBoneType.%s)"),
+	Commands.Add(FString::Printf(TEXT("hierarchy_controller.add_bone('%s', %s, %s, False, %s)"),
 		*Bone->GetName().ToString(),
 		*ParentKeyStr,
 		*TransformStr,
-		Bone->BoneType == ERigBoneType::Imported ? TEXT("IMPORTED") : TEXT("USER")));
+		*RigVMPythonUtils::EnumValueToPythonString<ERigBoneType>((int64)Bone->BoneType)
+	));
 
 	return Commands;
 }
@@ -1538,19 +1539,10 @@ TArray<FString> URigHierarchyController::GetAddRigidBodyPythonCommands(FRigRigid
 TArray<FString> URigHierarchyController::GetSetControlValuePythonCommands(const FRigControlElement* Control, const FRigControlValue& Value,
                                                                           const ERigControlValueType& Type) const
 {
-	FString TypeStr;
-	switch (Type)
-	{
-		case ERigControlValueType::Initial: TypeStr = TEXT("INITIAL"); break;
-		case ERigControlValueType::Current: TypeStr = TEXT("CURRENT"); break;
-		case ERigControlValueType::Minimum: TypeStr = TEXT("MINIMUM"); break;
-		case ERigControlValueType::Maximum: TypeStr = TEXT("MAXIMUM"); break;
-		default: ensure(false);
-	}
-	return {FString::Printf(TEXT("hierarchy.set_control_value(%s, %s, unreal.RigControlValueType.%s)"),
+	return {FString::Printf(TEXT("hierarchy.set_control_value(%s, %s, %s)"),
 		*Control->GetKey().ToPythonString(),
 		*Value.ToPythonString(Control->Settings.ControlType),
-		*TypeStr)};
+		*RigVMPythonUtils::EnumValueToPythonString<ERigControlValueType>((int64)Type))};
 }
 
 TArray<FString> URigHierarchyController::GetSetControlOffsetTransformPythonCommands(const FRigControlElement* Control,
@@ -2573,10 +2565,10 @@ TArray<FRigElementKey> URigHierarchyController::MirrorElements(TArray<FRigElemen
 			ArrayStr += TEXT("]");
 
 			RigVMPythonUtils::Print(Blueprint->GetFName().ToString(),
-				FString::Printf(TEXT("hierarchy_controller.mirror_elements(%s, unreal.RigMirrorSettings(unreal.AxisType.%s, unreal.AxisType.%s, '%s', '%s'), %s)"),
+				FString::Printf(TEXT("hierarchy_controller.mirror_elements(%s, unreal.RigMirrorSettings(%s, %s, '%s', '%s'), %s)"),
 				*ArrayStr,
-				(InSettings.MirrorAxis.GetValue() == EAxis::X) ? TEXT("X") : (InSettings.MirrorAxis.GetValue() == EAxis::Y) ? TEXT("Y") : TEXT("Z"),
-				(InSettings.AxisToFlip.GetValue() == EAxis::X) ? TEXT("X") : (InSettings.AxisToFlip.GetValue() == EAxis::Y) ? TEXT("Y") : TEXT("Z"),
+				*RigVMPythonUtils::EnumValueToPythonString<EAxis::Type>((int64)InSettings.MirrorAxis.GetValue()),
+				*RigVMPythonUtils::EnumValueToPythonString<EAxis::Type>((int64)InSettings.AxisToFlip.GetValue()),
 				*InSettings.SearchString,
 				*InSettings.ReplaceString,
 				(bSelectNewElements) ? TEXT("True") : TEXT("False")));
diff --git a/Engine/Plugins/Animation/ControlRig/Source/ControlRig/Private/Rigs/RigHierarchyDefines.cpp b/Engine/Plugins/Animation/ControlRig/Source/ControlRig/Private/Rigs/RigHierarchyDefines.cpp
--- a/Engine/Plugins/Animation/ControlRig/Source/ControlRig/Private/Rigs/RigHierarchyDefines.cpp
+++ b/Engine/Plugins/Animation/ControlRig/Source/ControlRig/Private/Rigs/RigHierarchyDefines.cpp
@@ -3,6 +3,10 @@
 #include "Rigs/RigHierarchyDefines.h"
 #include "Rigs/RigHierarchy.h"
 
+#if WITH_EDITOR
+#include "RigVMPythonUtils.h"
+#endif
+
 ////////////////////////////////////////////////////////////////////////////////
 // FRigControlLimitEnabled
 ////////////////////////////////////////////////////////////////////////////////
@@ -76,6 +80,17 @@ void FRigElementKey::Load(FArchive& Ar)
 	Ar << Name;
 }
 
+FString FRigElementKey::ToPythonString() const
+{
+#if WITH_EDITOR
+	return FString::Printf(TEXT("unreal.RigElementKey(type=%s, name='%s')"),
+		*RigVMPythonUtils::EnumValueToPythonString<ERigElementType>((int64)Type),
+		*Name.ToString());
+#else
+	return FString();
+#endif
+}
+
 ////////////////////////////////////////////////////////////////////////////////
 // FRigElementKeyCollection
 ////////////////////////////////////////////////////////////////////////////////
diff --git a/Engine/Plugins/Animation/ControlRig/Source/ControlRig/Public/Rigs/RigHierarchyDefines.h b/Engine/Plugins/Animation/ControlRig/Source/ControlRig/Public/Rigs/RigHierarchyDefines.h
--- a/Engine/Plugins/Animation/ControlRig/Source/ControlRig/Public/Rigs/RigHierarchyDefines.h
+++ b/Engine/Plugins/Animation/ControlRig/Source/ControlRig/Public/Rigs/RigHierarchyDefines.h
@@ -1515,13 +1515,7 @@ public:
 		return FString();
 	}
 
-	FORCEINLINE_DEBUGGABLE FString ToPythonString() const
-	{
-		FString TypeStr = StaticEnum<ERigElementType>()->GetDisplayNameTextByIndex((int32)Type).ToUpper().ToString();
-		return FString::Printf(TEXT("unreal.RigElementKey(type=unreal.RigElementType.%s, name='%s')"),
-			*TypeStr,
-			*Name.ToString());
-	}
+	FString ToPythonString() const;
 };
 
 USTRUCT(BlueprintType)
diff --git a/Engine/Plugins/Animation/ControlRig/Source/ControlRigDeveloper/Private/ControlRigBlueprint.cpp b/Engine/Plugins/Animation/ControlRig/Source/ControlRigDeveloper/Private/ControlRigBlueprint.cpp
--- a/Engine/Plugins/Animation/ControlRig/Source/ControlRigDeveloper/Private/ControlRigBlueprint.cpp
+++ b/Engine/Plugins/Animation/ControlRig/Source/ControlRigDeveloper/Private/ControlRigBlueprint.cpp
@@ -1840,9 +1840,10 @@ TArray<FString> UControlRigBlueprint::GeneratePythonCommands(const FString InNew
 							}
 
 							// FName AddExposedPin(const FName& InPinName, ERigVMPinDirection InDirection, const FString& InCPPType, const FName& InCPPTypeObjectPath, const FString& InDefaultValue, bool bSetupUndoRedo = true);
-							InternalCommands.Add(FString::Printf(TEXT("blueprint.get_controller_by_name('%s').add_exposed_pin('%s', unreal.RigVMPinDirection.INPUT, '%s', '%s', '%s')"),
+							InternalCommands.Add(FString::Printf(TEXT("blueprint.get_controller_by_name('%s').add_exposed_pin('%s', %s, '%s', '%s', '%s')"),
 									*Graph->GetGraphName(),
 									*Pin->GetName(),
+									*RigVMPythonUtils::EnumValueToPythonString<ERigVMPinDirection>((int64)ERigVMPinDirection::Input),
 									*Pin->GetCPPType(),
 									Pin->GetCPPTypeObject() ? *Pin->GetCPPTypeObject()->GetPathName() : TEXT(""),
 									*Pin->GetDefaultValue()));
@@ -1866,9 +1867,10 @@ TArray<FString> UControlRigBlueprint::GeneratePythonCommands(const FString InNew
 							}
 
 							// FName AddExposedPin(const FName& InPinName, ERigVMPinDirection InDirection, const FString& InCPPType, const FName& InCPPTypeObjectPath, const FString& InDefaultValue, bool bSetupUndoRedo = true);
-							InternalCommands.Add(FString::Printf(TEXT("blueprint.get_controller_by_name('%s').add_exposed_pin('%s', unreal.RigVMPinDirection.OUTPUT, '%s', '%s', '%s')"),
+							InternalCommands.Add(FString::Printf(TEXT("blueprint.get_controller_by_name('%s').add_exposed_pin('%s', %s, '%s', '%s', '%s')"),
 									*Graph->GetGraphName(),
-									*Pin->GetName(), 
+									*Pin->GetName(),
+									*RigVMPythonUtils::EnumValueToPythonString<ERigVMPinDirection>((int64)ERigVMPinDirection::Output),
 									*Pin->GetCPPType(),
 									Pin->GetCPPTypeObject() ? *Pin->GetCPPTypeObject()->GetPathName() : TEXT(""),
 									*Pin->GetDefaultValue()));
diff --git a/Engine/Source/Developer/RigVMDeveloper/Private/RigVMModel/RigVMController.cpp b/Engine/Source/Developer/RigVMDeveloper/Private/RigVMModel/RigVMController.cpp
--- a/Engine/Source/Developer/RigVMDeveloper/Private/RigVMModel/RigVMController.cpp
+++ b/Engine/Source/Developer/RigVMDeveloper/Private/RigVMModel/RigVMController.cpp
@@ -588,16 +588,13 @@ TArray<FString> URigVMController::GetAddNodePythonCommands(URigVMNode* Node) con
 	}
 	else if (const URigVMArrayNode* ArrayNode = Cast<URigVMArrayNode>(Node))
 	{
-		FString OpCodeString = StaticEnum<ERigVMOpCode>()->GetNameStringByValue((int64)ArrayNode->GetOpCode());
-		OpCodeString = RigVMPythonUtils::NameToPep8(OpCodeString);
-		OpCodeString.ToUpperInline();
-		
 		// add_array_node(opcode, cpp_type, cpp_type_object, position=[0.0, 0.0], node_name='', undo=True)
 		if (ArrayNode->GetCPPTypeObject())
 		{
-			Commands.Add(FString::Printf(TEXT("blueprint.get_controller_by_name('%s').add_array_node_from_object_path(unreal.RigVMOpCode.%s, '%s', '%s', %s, '%s')"),
+			static constexpr TCHAR ArrayNodeFormat[] = TEXT("blueprint.get_controller_by_name('%s').add_array_node_from_object_path(%s, '%s', '%s', %s, '%s')");
+			Commands.Add(FString::Printf(ArrayNodeFormat,
 					*GraphName,
-					*OpCodeString,
+					*RigVMPythonUtils::EnumValueToPythonString<ERigVMOpCode>((int64)ArrayNode->GetOpCode()),
 					*ArrayNode->GetCPPType(),
 					*ArrayNode->GetCPPTypeObject()->GetPathName(),
 					*RigVMPythonUtils::Vector2DToPythonString(ArrayNode->GetPosition()),
@@ -605,9 +602,10 @@ TArray<FString> URigVMController::GetAddNodePythonCommands(URigVMNode* Node) con
 		}
 		else
 		{
-			Commands.Add(FString::Printf(TEXT("blueprint.get_controller_by_name('%s').add_array_node(unreal.RigVMOpCode.%s, '%s', None, %s, '%s')"),
+			static constexpr TCHAR ArrayNodeFormat[] = TEXT("blueprint.get_controller_by_name('%s').add_array_node(%s, '%s', None, %s, '%s')");
+			Commands.Add(FString::Printf(ArrayNodeFormat,
 					*GraphName,
-					*OpCodeString,
+					*RigVMPythonUtils::EnumValueToPythonString<ERigVMOpCode>((int64)ArrayNode->GetOpCode()),
 					*ArrayNode->GetCPPType(),
 					*RigVMPythonUtils::Vector2DToPythonString(ArrayNode->GetPosition()),
 					*NodeName));	
@@ -9768,25 +9766,16 @@ FName URigVMController::AddExposedPin(const FName& InPinName, ERigVMPinDirection
 
 	if (bPrintPythonCommand)
 	{
-		FString DirectionStr;
-		switch (InDirection)
-		{
-			case ERigVMPinDirection::Hidden: DirectionStr = TEXT("unreal.RigVMPinDirection.HIDDEN"); break;
-			case ERigVMPinDirection::Input: DirectionStr = TEXT("unreal.RigVMPinDirection.INPUT"); break;
-			case ERigVMPinDirection::Output: DirectionStr = TEXT("unreal.RigVMPinDirection.OUTPUT"); break;
-			case ERigVMPinDirection::Visible: DirectionStr = TEXT("unreal.RigVMPinDirection.VISIBLE"); break;
-			case ERigVMPinDirection::IO: DirectionStr = TEXT("unreal.RigVMPinDirection.IO"); break;
-		}
-		
 		//AddExposedPin(const FName& InPinName, ERigVMPinDirection InDirection, const FString& InCPPType, const FName& InCPPTypeObjectPath, const FString& InDefaultValue, bool bSetupUndoRedo, bool bPrintPythonCommand)
 
 		const FString GraphName = GetSanitizedGraphName(GetGraph()->GetGraphName());
+		static constexpr TCHAR AddExposedPinFormat[] = TEXT("blueprint.get_controller_by_name('%s').add_exposed_pin('%s', %s, '%s', '%s', '%s')");
 		
 		RigVMPythonUtils::Print(GetGraphOuterName(), 
-			FString::Printf(TEXT("blueprint.get_controller_by_name('%s').add_exposed_pin('%s', %s, '%s', '%s', '%s')"),
+			FString::Printf(AddExposedPinFormat,
 				*GraphName,
 				*GetSanitizedPinName(InPinName.ToString()),
-				*DirectionStr,
+				*RigVMPythonUtils::EnumValueToPythonString<ERigVMPinDirection>((int64)InDirection),
 				*InCPPType,
 				*InCPPTypeObjectPath.ToString(),
 				*InDefaultValue));
@@ -14780,15 +14769,17 @@ void URigVMController::ResolveTemplateNodeMetaData(URigVMTemplateNode* InNode, b
 {
 #if WITH_EDITOR
 	check(InNode);
-
-	const FLinearColor PreviousColor = InNode->NodeColor;
 	
 	const TArray<int32> FilteredPermutationIndices = InNode->GetFilteredPermutationsIndices();
-	InNode->NodeColor = InNode->GetTemplate()->GetColor(FilteredPermutationIndices);
 
-	if(!InNode->NodeColor.Equals(PreviousColor, 0.01f))
+	if(InNode->IsA<URigVMUnitNode>())
 	{
-		Notify(ERigVMGraphNotifType::NodeColorChanged, InNode);
+		const FLinearColor PreviousColor = InNode->NodeColor;
+		InNode->NodeColor = InNode->GetTemplate()->GetColor(FilteredPermutationIndices);
+		if(!InNode->NodeColor.Equals(PreviousColor, 0.01f))
+		{
+			Notify(ERigVMGraphNotifType::NodeColorChanged, InNode);
+		}
 	}
 #endif
 
diff --git a/Engine/Source/Developer/RigVMDeveloper/Private/RigVMPythonUtils.cpp b/Engine/Source/Developer/RigVMDeveloper/Private/RigVMPythonUtils.cpp
--- a/Engine/Source/Developer/RigVMDeveloper/Private/RigVMPythonUtils.cpp
+++ b/Engine/Source/Developer/RigVMDeveloper/Private/RigVMPythonUtils.cpp
@@ -38,7 +38,8 @@ FString RigVMPythonUtils::NameToPep8(const FString& Name)
 
 FString RigVMPythonUtils::TransformToPythonString(const FTransform& Transform)
 {
-	return FString::Printf(TEXT("unreal.Transform(location=[%f,%f,%f],rotation=[%f,%f,%f],scale=[%f,%f,%f])"),
+	static constexpr TCHAR TransformFormat[] = TEXT("unreal.Transform(location=[%f,%f,%f],rotation=[%f,%f,%f],scale=[%f,%f,%f])");
+	return FString::Printf(TransformFormat,
 	                       Transform.GetLocation().X,
 	                       Transform.GetLocation().Y,
 	                       Transform.GetLocation().Z,
@@ -52,17 +53,34 @@ FString RigVMPythonUtils::TransformToPythonString(const FTransform& Transform)
 
 FString RigVMPythonUtils::Vector2DToPythonString(const FVector2D& Vector)
 {
-	return FString::Printf(TEXT("unreal.Vector2D(%f, %f)"),
+	static constexpr TCHAR Vector2DFormat[] = TEXT("unreal.Vector2D(%f, %f)");
+	return FString::Printf(Vector2DFormat,
 	                       Vector.X,
 	                       Vector.Y);
 }
 
 FString RigVMPythonUtils::LinearColorToPythonString(const FLinearColor& Color)
 {
-	return FString::Printf(TEXT("unreal.LinearColor(%f, %f, %f, %f)"),
+	static constexpr TCHAR LinearColorFormat[] = TEXT("unreal.LinearColor(%f, %f, %f, %f)");
+	return FString::Printf(LinearColorFormat,
 	                       Color.R, Color.G, Color.B, Color.A);
 }
 
+FString RigVMPythonUtils::EnumValueToPythonString(UEnum* Enum, int64 Value)
+{
+	static constexpr TCHAR EnumPrefix[] = TEXT("E");
+	static constexpr TCHAR EnumValueFormat[] = TEXT("unreal.%s.%s");
+
+	FString EnumName = Enum->GetName();
+	EnumName.RemoveFromStart(EnumPrefix, ESearchCase::CaseSensitive);
+	
+	return FString::Printf(
+		EnumValueFormat,
+		*EnumName,
+		*NameToPep8(Enum->GetNameStringByValue((int64)Value)).ToUpper()
+	);
+}
+
 #if WITH_EDITOR
 void RigVMPythonUtils::Print(const FString& BlueprintTitle, const FString& InMessage)
 {
@@ -92,14 +110,21 @@ void RigVMPythonUtils::PrintPythonContext(const FString& InBlueprintPath)
 	{
 		BlueprintName = BlueprintName.Right(BlueprintName.Len() - DotIndex - 1);
 	}
+
+	static constexpr TCHAR LoadObjectFormat[] = TEXT("blueprint = unreal.load_object(name = '%s', outer = None)");
+	static constexpr TCHAR DefineFunctionLibraryFormat[] = TEXT("library = blueprint.get_local_function_library()");
+	static constexpr TCHAR DefineLibraryControllerFormat[] = TEXT("library_controller = blueprint.get_controller(library)");
+	static constexpr TCHAR DefineHierarchyFormat[] = TEXT("hierarchy = blueprint.hierarchy");
+	static constexpr TCHAR DefineHierarchyControllerFormat[] = TEXT("hierarchy_controller = hierarchy.get_controller()");
 		
 	TArray<FString> PyCommands = {
 		TEXT("import unreal"),
-		FString::Printf(TEXT("blueprint = unreal.load_object(name = '%s', outer = None)"), *InBlueprintPath),
-		TEXT("library = blueprint.get_local_function_library()"),
-		TEXT("library_controller = blueprint.get_controller(library)"),
-		TEXT("hierarchy = blueprint.hierarchy"),
-		TEXT("hierarchy_controller = hierarchy.get_controller()")};
+		FString::Printf(LoadObjectFormat, *InBlueprintPath),
+		DefineFunctionLibraryFormat,
+		DefineLibraryControllerFormat,
+		DefineHierarchyFormat,
+		DefineHierarchyControllerFormat
+	};
 
 	for (FString& Command : PyCommands)
 	{
diff --git a/Engine/Source/Developer/RigVMDeveloper/Public/RigVMPythonUtils.h b/Engine/Source/Developer/RigVMDeveloper/Public/RigVMPythonUtils.h
--- a/Engine/Source/Developer/RigVMDeveloper/Public/RigVMPythonUtils.h
+++ b/Engine/Source/Developer/RigVMDeveloper/Public/RigVMPythonUtils.h
@@ -13,6 +13,14 @@ namespace RigVMPythonUtils
 	RIGVMDEVELOPER_API FString Vector2DToPythonString(const FVector2D& Vector);
 
 	RIGVMDEVELOPER_API FString LinearColorToPythonString(const FLinearColor& Color);
+
+	RIGVMDEVELOPER_API FString EnumValueToPythonString(UEnum* Enum, int64 Value);
+
+	template<typename T>
+	FORCEINLINE RIGVMDEVELOPER_API FString EnumValueToPythonString(int64 Value)
+	{
+		return EnumValueToPythonString(StaticEnum<T>(), Value);
+	}
 	
 #if WITH_EDITOR
 	RIGVMDEVELOPER_API void Print(const FString& BlueprintTitle, const FString& InMessage);
