Control Rig: Fix hanging of editor when adding/breaking links

diff --git a/Engine/Source/Developer/RigVMDeveloper/Private/RigVMModel/Nodes/RigVMTemplateNode.cpp b/Engine/Source/Developer/RigVMDeveloper/Private/RigVMModel/Nodes/RigVMTemplateNode.cpp
--- a/Engine/Source/Developer/RigVMDeveloper/Private/RigVMModel/Nodes/RigVMTemplateNode.cpp
+++ b/Engine/Source/Developer/RigVMDeveloper/Private/RigVMModel/Nodes/RigVMTemplateNode.cpp
@@ -639,7 +639,7 @@ TArray<int32> URigVMTemplateNode::GetNewFilteredPermutations(URigVMPin* InPin, c
 	return NewFilteredPermutations;
 }
 
-TArray<int32> URigVMTemplateNode::FindPermuationsForTypes(const TArray<FString>& ArgumentTypes)
+TArray<int32> URigVMTemplateNode::FindPermuationsForTypes(const TArray<FString>& ArgumentTypes, bool bAllowCasting)
 {
 	TArray<int32> Permutations;
 	if (const FRigVMTemplate* Template = GetTemplate())
@@ -657,7 +657,8 @@ TArray<int32> URigVMTemplateNode::FindPermuationsForTypes(const TArray<FString>&
 				
 				if (const FRigVMTemplateArgument* Argument = Template->FindArgument(*ArgName))
 				{
-					if (!Argument->GetTypes()[i].Matches(Type))
+					if ((bAllowCasting && !Argument->GetTypes()[i].Matches(Type)) ||
+						(!bAllowCasting && Argument->GetTypes()[i].CPPType != Type))
 					{
 						bAllArgsMatched = false;
 						break;
diff --git a/Engine/Source/Developer/RigVMDeveloper/Public/RigVMModel/Nodes/RigVMTemplateNode.h b/Engine/Source/Developer/RigVMDeveloper/Public/RigVMModel/Nodes/RigVMTemplateNode.h
--- a/Engine/Source/Developer/RigVMDeveloper/Public/RigVMModel/Nodes/RigVMTemplateNode.h
+++ b/Engine/Source/Developer/RigVMDeveloper/Public/RigVMModel/Nodes/RigVMTemplateNode.h
@@ -98,7 +98,7 @@ protected:
 	TArray<int32> GetNewFilteredPermutations(URigVMPin* InPin, URigVMPin* LinkedPin);
 	TArray<int32> GetNewFilteredPermutations(URigVMPin* InPin, const TArray<FRigVMTemplateArgument::FType>& InTypes);
 
-	TArray<int32> FindPermuationsForTypes(const TArray<FString>& ArgumentTypes);
+	TArray<int32> FindPermuationsForTypes(const TArray<FString>& ArgumentTypes, bool bAllowCasting = false);
 	TArray<FString> GetArgumentTypesForPermutation(const int32 InPermutationIndex);
 
 	UPROPERTY()
