ControlRig: rigvmpythonuitl.cpp non-unity compilation issue, missing class.h header include

diff --git a/Engine/Source/Developer/RigVMDeveloper/Public/RigVMPythonUtils.h b/Engine/Source/Developer/RigVMDeveloper/Public/RigVMPythonUtils.h
--- a/Engine/Source/Developer/RigVMDeveloper/Public/RigVMPythonUtils.h
+++ b/Engine/Source/Developer/RigVMDeveloper/Public/RigVMPythonUtils.h
@@ -3,6 +3,7 @@
 #pragma once
 
 #include "RigVMDeveloperModule.h"
+#include "UObject/Class.h"
 
 namespace RigVMPythonUtils
 {
@@ -15,7 +16,7 @@ namespace RigVMPythonUtils
 	RIGVMDEVELOPER_API FString LinearColorToPythonString(const FLinearColor& Color);
 
 	RIGVMDEVELOPER_API FString EnumValueToPythonString(UEnum* Enum, int64 Value);
-
+	
 	template<typename T>
 	FORCEINLINE RIGVMDEVELOPER_API FString EnumValueToPythonString(int64 Value)
 	{
