Control Rig: Fix crash when connecting different type to template node

diff --git a/Engine/Source/Developer/RigVMDeveloper/Private/RigVMModel/RigVMController.cpp b/Engine/Source/Developer/RigVMDeveloper/Private/RigVMModel/RigVMController.cpp
--- a/Engine/Source/Developer/RigVMDeveloper/Private/RigVMModel/RigVMController.cpp
+++ b/Engine/Source/Developer/RigVMDeveloper/Private/RigVMModel/RigVMController.cpp
@@ -1995,8 +1995,11 @@ bool URigVMController::UnresolveTemplateNodes(const TArray<URigVMTemplateNode*>&
 		{
 			BreakLink(Link->GetSourcePin(), Link->GetTargetPin(), bSetupUndoRedo);
 		}
-		
-		ResolveWildCardPin(Pin.Key, Pin.Value, bSetupUndoRedo);
+
+		if (Pin.Key->IsWildCard())
+		{
+			ResolveWildCardPin(Pin.Key, Pin.Value, bSetupUndoRedo);
+		}
 	}
 
 	if(bSetupUndoRedo)
@@ -15479,17 +15482,19 @@ bool URigVMController::UpdateTemplateNodePinTypes(URigVMTemplateNode* InNode, bo
 
 bool URigVMController::PropagateTemplateFilteredTypes(URigVMTemplateNode* InNode, bool bSetupUndoRedo) 
 {
-	auto UpdateAndPropagete = [&](URigVMPin* Pin, TArray<FRigVMTemplateArgument::FType>& Types)
+	auto UpdateAndPropagate = [&](URigVMPin* Pin)
 	{
 		TArray<URigVMPin*> OtherPins = Pin->GetLinkedSourcePins();
 		OtherPins.Append(Pin->GetLinkedTargetPins());
 		for (URigVMPin* OtherPin : OtherPins)
 		{
 			bool bPropagate = false;
+			bool bIsTemplate = false;
 			if (URigVMTemplateNode* OtherTemplate = Cast<URigVMTemplateNode>(OtherPin->GetNode()))
 			{
 				if (!OtherTemplate->IsSingleton())
 				{
+					bIsTemplate = true;
 					if (OtherTemplate->PinNeedsFilteredTypesUpdate(OtherPin, Pin))
 					{
 						if (UpdateFilteredPermutations(OtherPin, Pin, bSetupUndoRedo))
@@ -15509,15 +15514,25 @@ bool URigVMController::PropagateTemplateFilteredTypes(URigVMTemplateNode* InNode
 						}
 					}
 				}								
-			}			
+			}
+
+			if (!bIsTemplate)
+			{
+				if (!InNode->FilteredSupportsType(Pin, OtherPin->GetCPPType()))
+				{
+					URigVMLink* Link = Pin->FindLinkForPin(OtherPin);
+					ensureMsgf(!ActionStack->BracketActions.IsEmpty(), TEXT("Unexpected link broken %s in package %s"), *Link->GetPinPathRepresentation(), *GetPackage()->GetPathName());
+					BreakLink(Link->GetSourcePin(), Link->GetTargetPin(), bSetupUndoRedo);
+					return false;
+				}
+			}
 		}
 		return true;
 	};
 	
 	for(URigVMPin* Pin : InNode->GetPins())
 	{
-		TArray<FRigVMTemplateArgument::FType> Types = InNode->GetFilteredTypesForPin(Pin);
-		if (!UpdateAndPropagete(Pin, Types))
+		if (!UpdateAndPropagate(Pin))
 		{
 			return false;
 		}
@@ -15526,13 +15541,9 @@ bool URigVMController::PropagateTemplateFilteredTypes(URigVMTemplateNode* InNode
 		{
 			if (Pin->GetSubPins().Num() > 0)
 			{
-				for (FRigVMTemplateArgument::FType& Type : Types)
-				{
-					Type.ConvertToBaseElement();
-				}
 				for (URigVMPin* SubPin : Pin->GetSubPins())
 				{
-					if (!UpdateAndPropagete(SubPin, Types))
+					if (!UpdateAndPropagate(SubPin))
 					{
 						return false;
 					}
