Control Rig: Performance improvements on template nodes

diff --git a/Engine/Plugins/Animation/ControlRig/Source/ControlRig/Public/ControlRigObjectVersion.h b/Engine/Plugins/Animation/ControlRig/Source/ControlRig/Public/ControlRigObjectVersion.h
--- a/Engine/Plugins/Animation/ControlRig/Source/ControlRig/Public/ControlRigObjectVersion.h
+++ b/Engine/Plugins/Animation/ControlRig/Source/ControlRig/Public/ControlRigObjectVersion.h
@@ -63,6 +63,9 @@ struct CONTROLRIG_API FControlRigObjectVersion
 		// Added the notion of a per control animation type
 		ControlAnimationType,
 
+		// Added preferred permutation for templates
+		TemplatesPreferredPermutatation,
+
 		// -----<new versions can be added above this line>-------------------------------------------------
 		VersionPlusOne,
 		LatestVersion = VersionPlusOne - 1
diff --git a/Engine/Plugins/Animation/ControlRig/Source/ControlRigDeveloper/Private/ControlRigBlueprint.cpp b/Engine/Plugins/Animation/ControlRig/Source/ControlRigDeveloper/Private/ControlRigBlueprint.cpp
--- a/Engine/Plugins/Animation/ControlRig/Source/ControlRigDeveloper/Private/ControlRigBlueprint.cpp
+++ b/Engine/Plugins/Animation/ControlRig/Source/ControlRigDeveloper/Private/ControlRigBlueprint.cpp
@@ -479,6 +479,8 @@ void UControlRigBlueprint::PostLoad()
 		PatchBoundVariables();
 		PatchPropagateToChildren();
 		PatchParameterNodesOnLoad();
+		PatchTemplateNodesWithPreferredPermutation();
+		
 
 #if WITH_EDITOR
 
@@ -920,7 +922,7 @@ void UControlRigBlueprint::RefreshAllModels()
 				TemplateNode->InvalidateCache();
 			}
 		}
-		Controller->InitializeFilteredPermutationsFromTemplateTypes();
+		Controller->RecomputeAllTemplateFilteredTypes(false);
 	}
 }
 
@@ -4085,6 +4087,28 @@ void UControlRigBlueprint::PatchParameterNodesOnLoad()
 #endif	
 }
 
+void UControlRigBlueprint::PatchTemplateNodesWithPreferredPermutation()
+{
+#if WITH_EDITOR
+
+	if (GetLinkerCustomVersion(FControlRigObjectVersion::GUID) < FControlRigObjectVersion::TemplatesPreferredPermutatation)
+	{
+		for (URigVMGraph* Graph : GetAllModels())
+		{
+			for(URigVMNode* ModelNode : Graph->GetNodes())
+			{
+				if (URigVMTemplateNode* TemplateNode = Cast<URigVMTemplateNode>(ModelNode))
+				{
+					TemplateNode->InvalidateCache();
+				}
+			}
+			
+			GetController(Graph)->InitializeFilteredPermutationsFromTemplateTypes();
+		}
+	}
+#endif
+}
+
 void UControlRigBlueprint::PropagatePoseFromInstanceToBP(UControlRig* InControlRig)
 {
 	check(InControlRig);
diff --git a/Engine/Plugins/Animation/ControlRig/Source/ControlRigDeveloper/Public/ControlRigBlueprint.h b/Engine/Plugins/Animation/ControlRig/Source/ControlRigDeveloper/Public/ControlRigBlueprint.h
--- a/Engine/Plugins/Animation/ControlRig/Source/ControlRigDeveloper/Public/ControlRigBlueprint.h
+++ b/Engine/Plugins/Animation/ControlRig/Source/ControlRigDeveloper/Public/ControlRigBlueprint.h
@@ -573,6 +573,7 @@ private:
 	void PatchVariableNodesWithIncorrectType();
 	void PatchPropagateToChildren();
 	void PatchParameterNodesOnLoad();
+	void PatchTemplateNodesWithPreferredPermutation();
 
 	TMap<FName, int32> AddedMemberVariableMap;
 	TArray<FBPVariableDescription> LastNewVariables;
diff --git a/Engine/Source/Developer/RigVMDeveloper/Private/RigVMModel/Nodes/RigVMTemplateNode.cpp b/Engine/Source/Developer/RigVMDeveloper/Private/RigVMModel/Nodes/RigVMTemplateNode.cpp
--- a/Engine/Source/Developer/RigVMDeveloper/Private/RigVMModel/Nodes/RigVMTemplateNode.cpp
+++ b/Engine/Source/Developer/RigVMDeveloper/Private/RigVMModel/Nodes/RigVMTemplateNode.cpp
@@ -199,51 +199,16 @@ bool URigVMTemplateNode::SupportsType(const URigVMPin* InPin, const FString& InC
 			return CachedResult->Key;
 		}
 
-		FRigVMTemplate::FTypeMap Types;
-
-		for (URigVMPin* Pin : GetPins())
+		FRigVMTemplateArgument::FType OutType;
+		if (Template->ArgumentSupportsType(RootPin->GetFName(), CPPType, &OutType))
 		{
-			if (Pin == RootPin)
-			{
-				continue;
-			}
-			if (Pin->IsWildCard())
-			{
-				continue;
-			}
-			if (Pin->GetDirection() == ERigVMPinDirection::Hidden)
+			if (OutCPPType)
 			{
-				continue;
-			}
-			Types.Add(Pin->GetFName(), Pin->GetTemplateArgumentType());
-		}
-
-		FRigVMTemplateArgument::FType ResolvedType;
-		bool bSupportsType = false;
-		if(RootPin->IsWildCard())
-		{
-			bSupportsType = Template->ArgumentSupportsType(RootPin->GetFName(), CPPType, Types, &ResolvedType);
-		}
-		else
-		{
-			FRigVMTemplate::FTypeMap ResolvedTypes;
-			bSupportsType = Template->ResolveArgument(RootPin->GetFName(), FRigVMTemplateArgument::FType(CPPType), ResolvedTypes);
-			if(bSupportsType)
-			{
-				ResolvedType = ResolvedTypes.FindChecked(RootPin->GetFName());
-			}
-		}
-
-		if (bSupportsType)
-		{
-			ensure(ResolvedType.CPPTypeObject != RigVMTypeUtils::GetWildCardCPPTypeObject());
-			SupportedTypesCache.Add(CacheKey, TPair<bool, FRigVMTemplateArgument::FType>(bSupportsType, ResolvedType));
-			if(OutCPPType)
-			{
-				*OutCPPType = ResolvedType.CPPType;
+				(*OutCPPType) = OutType.CPPType;
 			}
+			return true;
 		}
-		return bSupportsType;
+		return false;
 	}
 	
 	if(RootPin->GetCPPType() == CPPType)
@@ -291,15 +256,22 @@ bool URigVMTemplateNode::FilteredSupportsType(const URigVMPin* InPin, const FStr
 		return false;
 	}
 
+	FString RootCPPType = InCPPType;
+	if (bIsArrayElement)
+	{
+		RootCPPType = RigVMTypeUtils::ArrayTypeFromBaseType(InCPPType);
+	}
+
+	if (FilteredPermutations.Num() == GetTemplate()->NumPermutations())
+	{
+		return GetTemplate()->ArgumentSupportsType(RootPin->GetFName(), InCPPType);
+	}
+
 	const TArray<FRigVMTemplateArgument::FType>& Types = Argument->GetTypes();
 	for (int32 PermutationIndex : FilteredPermutations)
 	{
-		FRigVMTemplateArgument::FType FilteredType = Types[PermutationIndex];
-		if (bIsArrayElement)
-		{
-			FilteredType.ConvertToBaseElement();
-		}
-		if (FilteredType.Matches(InCPPType, bAllowFloatingPointCasts))
+		const FRigVMTemplateArgument::FType& FilteredType = Types[PermutationIndex];
+		if (FilteredType.Matches(RootCPPType, bAllowFloatingPointCasts))
 		{
 			return true;
 		}
@@ -507,6 +479,26 @@ TArray<FRigVMTemplateArgument::FType> URigVMTemplateNode::GetFilteredTypesForPin
 	{
 		return FilteredTypes;	
 	}
+
+	if (!PreferredPermutationTypes.IsEmpty())
+	{
+		for (const FString& PreferredType : PreferredPermutationTypes)
+		{
+			FString ArgName, CPPType;
+			PreferredType.Split(TEXT(":"), &ArgName, &CPPType);
+			if (InPin->GetName() == ArgName)
+			{
+				const FRigVMTemplateArgument* Argument = GetTemplate()->FindArgument(*ArgName);
+				for (const FRigVMTemplateArgument::FType& Type : Argument->GetTypes())
+				{
+					if (Type.CPPType == CPPType)
+					{
+						return {Type};						
+					}
+				}				
+			}
+		}
+	}
 	
 	URigVMPin* RootPin = InPin;
 	bool bIsArrayElement = false;
@@ -557,7 +549,7 @@ TArray<int32> URigVMTemplateNode::GetNewFilteredPermutations(URigVMPin* InPin, U
 		}
 	}
 
-	TArray<int32> PermutationsToTry = PreferredPermutationTypes.IsEmpty()? FilteredPermutations : TArray<int32>(FindPermuationsForTypes(PreferredPermutationTypes));
+	TArray<int32> PermutationsToTry = PreferredPermutationTypes.IsEmpty()? FilteredPermutations : TArray<int32>(FindPermutationsForTypes(PreferredPermutationTypes));
 	bool bLinkedIsTemplate = false;
 	if (URigVMTemplateNode* LinkedTemplate = Cast<URigVMTemplateNode>(LinkedPin->GetNode()))
 	{
@@ -586,6 +578,7 @@ TArray<int32> URigVMTemplateNode::GetNewFilteredPermutations(URigVMPin* InPin, U
 	{
 		if (const FRigVMTemplateArgument* Argument = GetTemplate()->FindArgument(RootPin->GetFName()))
 		{
+			const FString LinkedCPPType = LinkedPin->GetCPPType();
 			for(int32 PermIndex : PermutationsToTry)
 			{
 				FRigVMTemplateArgument::FType Type = Argument->GetTypes()[PermIndex];
@@ -593,7 +586,7 @@ TArray<int32> URigVMTemplateNode::GetNewFilteredPermutations(URigVMPin* InPin, U
 				{
 					Type.ConvertToBaseElement();
 				}
-				if (Type.Matches(LinkedPin->GetCPPType()))
+				if (Type.Matches(LinkedCPPType))
 				{
 					NewFilteredPermutations.Add(PermIndex);
 				}
@@ -619,7 +612,7 @@ TArray<int32> URigVMTemplateNode::GetNewFilteredPermutations(URigVMPin* InPin, c
 		}
 	}
 
-	TArray<int32> PermutationsToTry = PreferredPermutationTypes.IsEmpty()? FilteredPermutations : TArray<int32>(FindPermuationsForTypes(PreferredPermutationTypes));
+	TArray<int32> PermutationsToTry = PreferredPermutationTypes.IsEmpty()? FilteredPermutations : TArray<int32>(FindPermutationsForTypes(PreferredPermutationTypes));
 	
 	if (const FRigVMTemplateArgument* Argument = GetTemplate()->FindArgument(RootPin->GetFName()))
 	{
@@ -639,35 +632,46 @@ TArray<int32> URigVMTemplateNode::GetNewFilteredPermutations(URigVMPin* InPin, c
 	return NewFilteredPermutations;
 }
 
-TArray<int32> URigVMTemplateNode::FindPermuationsForTypes(const TArray<FString>& ArgumentTypes, bool bAllowCasting)
+TArray<int32> URigVMTemplateNode::FindPermutationsForTypes(const TArray<FString>& ArgumentTypes, bool bAllowCasting)
 {
 	TArray<int32> Permutations;
 	if (const FRigVMTemplate* Template = GetTemplate())
 	{
+		TArray<const FRigVMTemplateArgument*> Args;
+		TArray<FString> CPPTypes;
+		for (const FString& TypePairString : ArgumentTypes)
+		{
+			FString ArgName, Type;
+			if (!TypePairString.Split(TEXT(":"), &ArgName, &Type))
+			{
+				return {};
+			}
+
+			if (const FRigVMTemplateArgument* Argument = Template->FindArgument(*ArgName))
+			{
+				Args.Add(Argument);
+				CPPTypes.Add(Type);
+			}
+			else
+			{
+				return {};
+			}
+		}
+		
 		for (int32 i=0; i<Template->NumPermutations(); ++i)
 		{
 			bool bAllArgsMatched = true;
-			for (const FString& TypePairString : ArgumentTypes)
+			for (int32 ArgIndex = 0; ArgIndex < Args.Num(); ++ArgIndex)
 			{
-				FString ArgName, Type;
-				if (!TypePairString.Split(TEXT(":"), &ArgName, &Type))
-				{
-					return {};
-				}
-				
-				if (const FRigVMTemplateArgument* Argument = Template->FindArgument(*ArgName))
+				const FRigVMTemplateArgument* Argument = Args[ArgIndex];
 				{
-					if ((bAllowCasting && !Argument->GetTypes()[i].Matches(Type)) ||
-						(!bAllowCasting && Argument->GetTypes()[i].CPPType != Type))
+					if ((bAllowCasting && !Argument->GetTypes()[i].Matches(CPPTypes[ArgIndex])) ||
+						(!bAllowCasting && Argument->GetTypes()[i].CPPType != CPPTypes[ArgIndex]))
 					{
 						bAllArgsMatched = false;
 						break;
 					}
 				}
-				else
-				{
-					return {};
-				}
 			}
 
 			if (bAllArgsMatched)
@@ -784,7 +788,7 @@ void URigVMTemplateNode::InitializeFilteredPermutations()
 	{
 		if (!PreferredPermutationTypes.IsEmpty())
 		{
-			FilteredPermutations = FindPermuationsForTypes(PreferredPermutationTypes);
+			FilteredPermutations = FindPermutationsForTypes(PreferredPermutationTypes);
 		}
 		else
 		{
@@ -819,7 +823,7 @@ void URigVMTemplateNode::InitializeFilteredPermutationsFromTypes()
 			}
 		}
 
-		TArray<int32> Permutations = FindPermuationsForTypes(ArgTypes);
+		TArray<int32> Permutations = FindPermutationsForTypes(ArgTypes);
 		if (!Permutations.IsEmpty())
 		{
 			FilteredPermutations = Permutations;
diff --git a/Engine/Source/Developer/RigVMDeveloper/Private/RigVMModel/RigVMController.cpp b/Engine/Source/Developer/RigVMDeveloper/Private/RigVMModel/RigVMController.cpp
--- a/Engine/Source/Developer/RigVMDeveloper/Private/RigVMModel/RigVMController.cpp
+++ b/Engine/Source/Developer/RigVMDeveloper/Private/RigVMModel/RigVMController.cpp
@@ -3579,7 +3579,7 @@ TArray<FName> URigVMController::ImportNodesFromText(const FString& InText, bool
 					}
 					if (!TemplateNode->PreferredPermutationTypes.IsEmpty())
 					{
-						TemplateNode->FilteredPermutations = TemplateNode->FindPermuationsForTypes(TemplateNode->PreferredPermutationTypes);
+						TemplateNode->FilteredPermutations = TemplateNode->FindPermutationsForTypes(TemplateNode->PreferredPermutationTypes);
 					}
 					UpdateTemplateNodePinTypes(TemplateNode, bSetupUndoRedo);
 				}
@@ -15369,7 +15369,8 @@ bool URigVMController::UpdateFilteredPermutations(URigVMPin* InPin, const TArray
 
 bool URigVMController::UpdateTemplateNodePinTypes(URigVMTemplateNode* InNode, bool bSetupUndoRedo)
 {
-	TArray<URigVMPin*> RemainingWildcardPins;
+	bool bAnyWildcard = false;;
+	bool bAnyPinChanged = false;
 	for (URigVMPin* Pin : InNode->GetPins())
 	{
 		TArray<FRigVMTemplateArgument::FType> Types = InNode->GetFilteredTypesForPin(Pin);
@@ -15377,6 +15378,38 @@ bool URigVMController::UpdateTemplateNodePinTypes(URigVMTemplateNode* InNode, bo
 		{
 			continue;
 		}
+
+		if (Types.Num() > 1)
+		{
+			
+			bool bCanReduceToSingleType = true;
+			for (int32 i=1; i<Types.Num(); ++i)
+			{
+				if (!RigVMTypeUtils::AreCompatible(Types[0].CPPType, Types[0].CPPTypeObject, Types[i].CPPType, Types[i].CPPTypeObject))
+				{
+					bCanReduceToSingleType = false;
+					break;
+				}
+			}
+			
+			// Select the same Pin type if it has one
+			if (bCanReduceToSingleType)
+			{
+				int32 PreferredIndex = 0;
+				if (!Pin->IsWildCard())
+				{
+					for (int32 i=0; i<Types.Num(); ++i)
+					{
+						if (Types[i].CPPType == Pin->GetCPPType())
+						{
+							PreferredIndex = i;
+							break;
+						}
+					}
+				}
+				Types = {Types[PreferredIndex]};
+			}
+		}
 		
 		if (Types.Num() > 1)
 		{
@@ -15398,9 +15431,10 @@ bool URigVMController::UpdateTemplateNodePinTypes(URigVMTemplateNode* InNode, bo
 			
 			if (Pin->GetCPPType() != CPPType || Pin->GetCPPTypeObject() != CPPObjectType)
 			{
+				bAnyPinChanged = true;
 				ChangePinType(Pin, CPPType, CPPObjectType, bSetupUndoRedo, false, false, false);
 			}
-			RemainingWildcardPins.Add(Pin);
+			bAnyWildcard = true;
 		}
 		else if (Types.Num() == 1)
 		{
@@ -15409,6 +15443,7 @@ bool URigVMController::UpdateTemplateNodePinTypes(URigVMTemplateNode* InNode, bo
 			UObject* CPPObjectType = Types[0].CPPTypeObject;
 			if (Pin->GetCPPType() != CPPType || Pin->GetCPPTypeObject() != CPPObjectType)
 			{
+				bAnyPinChanged = true;
 				ChangePinType(Pin, CPPType, CPPObjectType, bSetupUndoRedo, false, false, false);
 			}
 		}
@@ -15419,60 +15454,7 @@ bool URigVMController::UpdateTemplateNodePinTypes(URigVMTemplateNode* InNode, bo
 		}
 	}
 
-	if (RemainingWildcardPins.Num() > 0)
-	{
-		TArray<int32> ResolvePermutations;
-		bool bCanAllReduceToSingleType = true;
-		for (URigVMPin* WildcardPin : RemainingWildcardPins)
-		{
-			TArray<FRigVMTemplateArgument::FType> Types = InNode->GetFilteredTypesForPin(WildcardPin);
-			if (Types.IsEmpty())
-			{
-				continue;
-			}
-
-			// Find if all types are compatible (for example, float and double) 
-			bool bCanReduceToSingleType=true;
-			if (Types.Num() > 1)
-			{
-				for (int32 i=1; i<Types.Num(); ++i)
-				{
-					if (!Types[0].Matches(Types[i].CPPType))
-					{
-						bCanReduceToSingleType=false;
-						break;
-					}
-				}
-			}
-	
-			if (!bCanReduceToSingleType)
-			{
-				bCanAllReduceToSingleType = false;
-				break;
-			}			
-		}
-
-		// Choose the first permutation
-		if (bCanAllReduceToSingleType)
-		{
-			TArray<int32> OldPermutations;
-			if (bSetupUndoRedo)
-			{
-				OldPermutations = InNode->GetFilteredPermutationsIndices();
-			}
-			InNode->FilteredPermutations = {InNode->FilteredPermutations[0]};
-			if (bSetupUndoRedo)
-			{
-				FRigVMSetTemplateFilteredPermutationsAction Action(InNode, nullptr, OldPermutations);
-				ActionStack->BeginAction(Action);
-				ActionStack->AddAction(Action);
-				ActionStack->EndAction(Action);
-			}
-			
-			return UpdateTemplateNodePinTypes(InNode, bSetupUndoRedo);
-		}
-	}
-	else
+	if (!bAnyWildcard && bAnyPinChanged)
 	{
 		FullyResolveTemplateNode(InNode, InNode->FilteredPermutations[0], bSetupUndoRedo);
 	}
@@ -15554,7 +15536,7 @@ bool URigVMController::PropagateTemplateFilteredTypes(URigVMTemplateNode* InNode
 	return true;
 }
 
-void URigVMController::RecomputeAllTemplateFilteredTypes(bool bSetupUndoRedo, bool bKeepResolvedTypes)
+void URigVMController::RecomputeAllTemplateFilteredTypes(bool bSetupUndoRedo)
 {
 	if (!IsValidGraph())
 	{
@@ -15606,9 +15588,33 @@ void URigVMController::RecomputeAllTemplateFilteredTypes(bool bSetupUndoRedo, bo
 	InitializeAllTemplateFiltersInGraph(bSetupUndoRedo, false);
 
 	// Apply all links to update filtered permutations
-	for (int32 i=0; i<Graph->GetLinks().Num(); ++i)
+	TArray<URigVMLink*> SortedLinks = Graph->GetLinks();
+
+	// Solve for unit nodes first, other templates are more expensive to solve. This way we are avoiding solving
+	// links like (reroute-reroute) before their filtered permutations are reduced
+	Algo::Sort(SortedLinks, [](const URigVMLink* A, const URigVMLink* B)
+	{
+		const bool bASourceIsUnitNode = A->GetSourcePin()->GetNode()->IsA<URigVMUnitNode>();
+		const bool bATargetIsUnitNode = A->GetTargetPin()->GetNode()->IsA<URigVMUnitNode>();
+		const bool bBSourceIsUnitNode = B->GetSourcePin()->GetNode()->IsA<URigVMUnitNode>();
+		const bool bBTargetIsUnitNode = B->GetTargetPin()->GetNode()->IsA<URigVMUnitNode>();
+		
+		if (bASourceIsUnitNode && bATargetIsUnitNode && (!bBSourceIsUnitNode || !bBTargetIsUnitNode))
+		{
+			return true;
+		}
+
+		if ((bASourceIsUnitNode || bATargetIsUnitNode) && (!bBSourceIsUnitNode && !bBTargetIsUnitNode))
+		{
+			return true;
+		}
+
+		return false;		
+	});
+	
+	for (int32 i=0; i<SortedLinks.Num(); ++i)
 	{
-		URigVMLink* Link = Graph->GetLinks()[i];
+		URigVMLink* Link = SortedLinks[i];
 		
 		URigVMPin* OutputPin = Link->GetSourcePin();
 		URigVMPin* InputPin = Link->GetTargetPin();
@@ -15686,26 +15692,6 @@ void URigVMController::RecomputeAllTemplateFilteredTypes(bool bSetupUndoRedo, bo
 				continue;
 			}
 
-			if (bKeepResolvedTypes && TemplateNode->FilteredPermutations.Num() > 1)
-			{
-				if (const FRigVMTemplate* Template = TemplateNode->GetTemplate())
-				{
-					TArray<FString> Types;
-					for (int32 i=0; i<Template->NumArguments(); ++i)
-					{
-						const FRigVMTemplateArgument* Argument = Template->GetArgument(i);
-						URigVMPin* Pin = TemplateNode->FindPin(Argument->GetName().ToString());
-						if (const FRigVMTemplateArgument::FType* Type = TypesBeforeRecomputing.Find(Pin))
-						{
-							Types.Add(Argument->GetName().ToString() + TEXT(":") + Type->CPPType);
-						}						
-					}
-
-					TemplateNode->FilteredPermutations = TemplateNode->FindPermuationsForTypes(Types);
-					TemplateNode->PreferredPermutationTypes = Types;
-				}
-			}
-
 			UpdateTemplateNodePinTypes(TemplateNode, bSetupUndoRedo);
 		}
 	}
@@ -15731,6 +15717,10 @@ void URigVMController::InitializeFilteredPermutationsFromTemplateTypes()
 		if (URigVMTemplateNode* TemplateNode = Cast<URigVMTemplateNode>(Node))
 		{
 			TemplateNode->InitializeFilteredPermutationsFromTypes();
+			if (!TemplateNode->IsSingleton())
+			{
+				FullyResolveTemplateNode(TemplateNode, TemplateNode->FilteredPermutations[0], false);
+			}
 		}
 	}
 }
diff --git a/Engine/Source/Developer/RigVMDeveloper/Public/RigVMModel/Nodes/RigVMTemplateNode.h b/Engine/Source/Developer/RigVMDeveloper/Public/RigVMModel/Nodes/RigVMTemplateNode.h
--- a/Engine/Source/Developer/RigVMDeveloper/Public/RigVMModel/Nodes/RigVMTemplateNode.h
+++ b/Engine/Source/Developer/RigVMDeveloper/Public/RigVMModel/Nodes/RigVMTemplateNode.h
@@ -98,7 +98,7 @@ protected:
 	TArray<int32> GetNewFilteredPermutations(URigVMPin* InPin, URigVMPin* LinkedPin);
 	TArray<int32> GetNewFilteredPermutations(URigVMPin* InPin, const TArray<FRigVMTemplateArgument::FType>& InTypes);
 
-	TArray<int32> FindPermuationsForTypes(const TArray<FString>& ArgumentTypes, bool bAllowCasting = false);
+	TArray<int32> FindPermutationsForTypes(const TArray<FString>& ArgumentTypes, bool bAllowCasting = false);
 	TArray<FString> GetArgumentTypesForPermutation(const int32 InPermutationIndex);
 
 	UPROPERTY()
diff --git a/Engine/Source/Developer/RigVMDeveloper/Public/RigVMModel/RigVMController.h b/Engine/Source/Developer/RigVMDeveloper/Public/RigVMModel/RigVMController.h
--- a/Engine/Source/Developer/RigVMDeveloper/Public/RigVMModel/RigVMController.h
+++ b/Engine/Source/Developer/RigVMDeveloper/Public/RigVMModel/RigVMController.h
@@ -893,7 +893,7 @@ public:
 	bool RemoveUnusedOrphanedPins(URigVMNode* InNode, bool bNotify);
 
 	// Initializes and recomputes the filtered permutations of all template nodes in the graph 
-	void RecomputeAllTemplateFilteredTypes(bool bSetupUndoRedo, bool bKeepResolvedTypes = false);
+	void RecomputeAllTemplateFilteredTypes(bool bSetupUndoRedo);
 
 	// Try to initialize the filterd permutations from the pin types
 	void InitializeFilteredPermutationsFromTemplateTypes();
diff --git a/Engine/Source/Runtime/RigVM/Private/RigVMCore/RigVMRegistry.cpp b/Engine/Source/Runtime/RigVM/Private/RigVMCore/RigVMRegistry.cpp
--- a/Engine/Source/Runtime/RigVM/Private/RigVMCore/RigVMRegistry.cpp
+++ b/Engine/Source/Runtime/RigVM/Private/RigVMCore/RigVMRegistry.cpp
@@ -155,6 +155,12 @@ const FRigVMTemplate* FRigVMRegistry::GetOrAddTemplateFromArguments(const FName&
 			{
 				Argument.Types = FRigVMTemplateArgument::GetCompatibleTypes(FRigVMTemplateArgument::ETypeCategory_SingleAnyValue);
 			}
+
+			for (int32 i = 0; i < Argument.Types.Num(); ++i)
+			{
+				Argument.TypeToPermutations.Add(Argument.Types[i].CPPType, {i});				
+			}
+			
 			Argument.bSingleton = false;
 			NumPermutations = FMath::Max(NumPermutations, Argument.Types.Num()); 
 		}
@@ -169,10 +175,14 @@ const FRigVMTemplate* FRigVMRegistry::GetOrAddTemplateFromArguments(const FName&
 			{
 				const FRigVMTemplateArgument::FType Type = Argument.Types[0];
 				Argument.Types.SetNum(NumPermutations);
+				TArray<int32> ArgTypePermutations;
+				ArgTypePermutations.SetNum(NumPermutations);
 				for(int32 Index=0;Index<NumPermutations;Index++)
 				{
 					Argument.Types[Index] = Type;
+					ArgTypePermutations[Index] = Index;
 				}
+				Argument.TypeToPermutations.Add(Type.CPPType, ArgTypePermutations);
 				Argument.bSingleton = true;
 			}
 		}
diff --git a/Engine/Source/Runtime/RigVM/Private/RigVMCore/RigVMTemplate.cpp b/Engine/Source/Runtime/RigVM/Private/RigVMCore/RigVMTemplate.cpp
--- a/Engine/Source/Runtime/RigVM/Private/RigVMCore/RigVMTemplate.cpp
+++ b/Engine/Source/Runtime/RigVM/Private/RigVMCore/RigVMTemplate.cpp
@@ -51,6 +51,7 @@ FRigVMTemplateArgument::FRigVMTemplateArgument(FProperty* InProperty)
 	Type.CPPType = RigVMTypeUtils::PostProcessCPPType(Type.CPPType, Type.CPPTypeObject);
 
 	Types.Add(Type);
+	TypeToPermutations.Add(Type.CPPType, {0});
 }
 
 FRigVMTemplateArgument::FRigVMTemplateArgument(const FName& InName, ERigVMPinDirection InDirection, const FType& InType)
@@ -60,6 +61,7 @@ FRigVMTemplateArgument::FRigVMTemplateArgument(const FName& InName, ERigVMPinDir
 , bSingleton(true)
 , Types({InType})
 {
+	TypeToPermutations.Add(InType.CPPType, {0});	
 }
 
 FRigVMTemplateArgument::FRigVMTemplateArgument(const FName& InName, ERigVMPinDirection InDirection, const TArray<FType>& InTypes)
@@ -70,6 +72,7 @@ FRigVMTemplateArgument::FRigVMTemplateArgument(const FName& InName, ERigVMPinDir
 , Types(InTypes)
 {
 	check(Types.Num() > 0);
+
 	for(int32 TypeIndex=1;TypeIndex<Types.Num();TypeIndex++)
 	{
 		if(Types[TypeIndex] != Types[0])
@@ -78,47 +81,46 @@ FRigVMTemplateArgument::FRigVMTemplateArgument(const FName& InName, ERigVMPinDir
 			break;
 		}
 	}
-}
 
-bool FRigVMTemplateArgument::SupportsType(const FString& InCPPType, const TArray<int32>& InPermutationIndices, FType* OutType) const
-{
-	bool bFoundMatch = false;
-	bool bFoundPerfectMatch = false;
-	
-	auto VisitType = [&bFoundMatch, &bFoundPerfectMatch, InCPPType, &OutType](const FType& InType)
+	for(int32 TypeIndex=0;TypeIndex<Types.Num();TypeIndex++)
 	{
-		if (InType.Matches(InCPPType))
+		if (TArray<int32>* Permutations = TypeToPermutations.Find(Types[TypeIndex].CPPType))
 		{
-			if(!bFoundPerfectMatch)
-			{
-				if(OutType)
-				{
-					*OutType = InType;
-				}
-				if(InType.CPPType == InCPPType)
-				{
-					bFoundPerfectMatch = true;
-				}
-			}
-			bFoundMatch = true;
+			Permutations->Add(TypeIndex);
 		}
-	};
-	
-	if(InPermutationIndices.IsEmpty())
+		else
+		{
+			TypeToPermutations.Add(Types[TypeIndex].CPPType, {TypeIndex});
+		}
+	}
+}
+
+bool FRigVMTemplateArgument::SupportsType(const FString& InCPPType, FType* OutType) const
+{
+	if (const TArray<int32>* Permutations = TypeToPermutations.Find(InCPPType))
 	{
-		for (const FType& Type : Types)
+		if(OutType)
 		{
-			VisitType(Type);
+			(*OutType) = Types[(*Permutations)[0]];
 		}
+		return true;		
 	}
-	else
+
+	// Try to find compatible type
+	TArray<FString> CompatibleTypes = FType::GetCompatibleTypes(InCPPType);
+	for (const FString& Compatible : CompatibleTypes)
 	{
-		for (const int32 PermutationIndex : InPermutationIndices)
+		if (const TArray<int32>* Permutations = TypeToPermutations.Find(Compatible))
 		{
-			VisitType(Types[PermutationIndex]);
+			if(OutType)
+			{
+				(*OutType) = Types[(*Permutations)[0]];
+			}
+			return true;		
 		}
 	}
-	return bFoundMatch;
+
+	return false;	
 }
 
 bool FRigVMTemplateArgument::IsSingleton(const TArray<int32>& InPermutationIndices) const
@@ -852,6 +854,20 @@ bool FRigVMTemplate::Merge(const FRigVMTemplate& InOther)
 		}
 
 		NewArgs.Add(Arguments[ArgumentIndex]);
+
+		// Add Other argument information into the TypeToPermutations map
+		{
+			const FString& OtherCPPType = InOther.Arguments[ArgumentIndex].Types[0].CPPType;
+			const int32 NewPermutationIndex = NewArgs[ArgumentIndex].Types.Num();
+			if (TArray<int32>* ArgTypePermutations = NewArgs[ArgumentIndex].TypeToPermutations.Find(OtherCPPType))
+			{
+				ArgTypePermutations->Add(NewPermutationIndex);
+			}
+			else
+			{
+				NewArgs[ArgumentIndex].TypeToPermutations.Add(OtherCPPType, {NewPermutationIndex});
+			}
+		}
 		NewArgs[ArgumentIndex].Types.Add(InOther.Arguments[ArgumentIndex].Types[0]);
 
 		if (!Arguments[ArgumentIndex].Types.Contains(InOther.Arguments[ArgumentIndex].Types[0]))
@@ -874,28 +890,12 @@ const FRigVMTemplateArgument* FRigVMTemplate::FindArgument(const FName& InArgume
 	});
 }
 
-bool FRigVMTemplate::ArgumentSupportsType(const FName& InArgumentName, const FString& InCPPType, const FRigVMTemplate::FTypeMap& InTypes, FRigVMTemplateArgument::FType* OutType) const
+bool FRigVMTemplate::ArgumentSupportsType(const FName& InArgumentName, const FString& InCPPType, FRigVMTemplateArgument::FType* OutType) const
 {
 	if (const FRigVMTemplateArgument* Argument = FindArgument(InArgumentName))
 	{
-		if(InTypes.Num() == 0)
-		{
-			return Argument->SupportsType(InCPPType, TArray<int32>(), OutType);
-		}
-
-		FTypeMap Types = InTypes;
-		Types.FindOrAdd(InArgumentName) = InCPPType;
-
-		TArray<int32> PermutationIndices;
-		if(Resolve(Types, PermutationIndices, true))
-		{
-			if(OutType)
-			{
-				*OutType = Types.FindChecked(InArgumentName);
-			}
-			return true;
-		}
-		return false;
+		return Argument->SupportsType(InCPPType, OutType);		
+		
 	}
 	return false;
 }
diff --git a/Engine/Source/Runtime/RigVM/Public/RigVMCore/RigVMTemplate.h b/Engine/Source/Runtime/RigVM/Public/RigVMCore/RigVMTemplate.h
--- a/Engine/Source/Runtime/RigVM/Public/RigVMCore/RigVMTemplate.h
+++ b/Engine/Source/Runtime/RigVM/Public/RigVMCore/RigVMTemplate.h
@@ -81,6 +81,27 @@ struct RIGVM_API FRigVMTemplateArgument
 			return false;
 		}
 
+		static TArray<FString> GetCompatibleTypes(const FString& InCPPType)
+		{
+			if(InCPPType == RigVMTypeUtils::FloatType)
+			{
+				return {RigVMTypeUtils::DoubleType};
+			}
+			if(InCPPType == RigVMTypeUtils::DoubleType)
+			{
+				return {RigVMTypeUtils::FloatType};
+			}
+			if(InCPPType == RigVMTypeUtils::FloatArrayType)
+			{
+				return {RigVMTypeUtils::DoubleArrayType};
+			}
+			if(InCPPType == RigVMTypeUtils::DoubleArrayType)
+			{
+				return {RigVMTypeUtils::FloatArrayType};
+			}
+			return {};
+		}
+
 		FName GetCPPTypeObjectPath() const
 		{
 			if(CPPTypeObject)
@@ -157,7 +178,7 @@ struct RIGVM_API FRigVMTemplateArgument
 	ERigVMPinDirection GetDirection() const { return Direction; }
 
 	// returns true if this argument supports a given type across a set of permutations
-	bool SupportsType(const FString& InCPPType, const TArray<int32>& InPermutationIndices = TArray<int32>(), FType* OutType = nullptr) const;
+	bool SupportsType(const FString& InCPPType, FType* OutType = nullptr) const;
 
 	// returns the flat list of types (including duplicates) of this argument
 	const TArray<FType>& GetTypes() const;
@@ -184,6 +205,7 @@ protected:
 	ERigVMPinDirection Direction;
 	bool bSingleton;
 	TArray<FType> Types;
+	TMap<FString, TArray<int32>> TypeToPermutations;
 
 	// default constructor
 	FRigVMTemplateArgument();
@@ -234,7 +256,7 @@ public:
 	const FRigVMTemplateArgument* FindArgument(const FName& InArgumentName) const;
 
 	// returns true if a given arg supports a type
-	bool ArgumentSupportsType(const FName& InArgumentName, const FString& InCPPType, const FTypeMap& InTypes = FTypeMap(), FRigVMTemplateArgument::FType* OutType = nullptr) const;
+	bool ArgumentSupportsType(const FName& InArgumentName, const FString& InCPPType, FRigVMTemplateArgument::FType* OutType = nullptr) const;
 
 	// returns the number of permutations supported by this template
 	int32 NumPermutations() const { return Permutations.Num(); }
