Control Rig: Fix loading template nodes with wildcard pins

diff --git a/Engine/Source/Developer/RigVMDeveloper/Private/RigVMModel/RigVMController.cpp b/Engine/Source/Developer/RigVMDeveloper/Private/RigVMModel/RigVMController.cpp
--- a/Engine/Source/Developer/RigVMDeveloper/Private/RigVMModel/RigVMController.cpp
+++ b/Engine/Source/Developer/RigVMDeveloper/Private/RigVMModel/RigVMController.cpp
@@ -15717,10 +15717,11 @@ void URigVMController::InitializeFilteredPermutationsFromTemplateTypes()
 		if (URigVMTemplateNode* TemplateNode = Cast<URigVMTemplateNode>(Node))
 		{
 			TemplateNode->InitializeFilteredPermutationsFromTypes();
-			if (!TemplateNode->IsSingleton())
-			{
-				FullyResolveTemplateNode(TemplateNode, TemplateNode->FilteredPermutations[0], false);
-			}
+			// if (!TemplateNode->IsSingleton())
+			// {
+			// 	FullyResolveTemplateNode(TemplateNode, TemplateNode->FilteredPermutations[0], false);
+			// }
+			//
 		}
 	}
 }
