Control Rig: Fix crash when undo/redo upgrade nodes

diff --git a/Engine/Source/Developer/RigVMDeveloper/Private/RigVMModel/RigVMController.cpp b/Engine/Source/Developer/RigVMDeveloper/Private/RigVMModel/RigVMController.cpp
--- a/Engine/Source/Developer/RigVMDeveloper/Private/RigVMModel/RigVMController.cpp
+++ b/Engine/Source/Developer/RigVMDeveloper/Private/RigVMModel/RigVMController.cpp
@@ -14822,11 +14822,29 @@ bool URigVMController::FullyResolveTemplateNode(URigVMTemplateNode* InNode, int3
 	TGuardValue<bool> ReentryGuard(bIsFullyResolvingTemplateNode, true);
 	
 	check(InNode);
-	check(InPermutationIndex != INDEX_NONE);
 
 	const FRigVMTemplate* Template = InNode->GetTemplate();
-	const FRigVMFunction* ResolvedFunction = Template->GetPermutation(InPermutationIndex);
-	const TArray<int32> PermutationIndices = {InPermutationIndex};
+	int32 InputPermutation = InPermutationIndex;
+	
+	// Figure out the permutation index fromt the pin types
+	if (InPermutationIndex == INDEX_NONE)
+	{
+		FRigVMTemplate::FTypeMap TypeMap;
+		for (URigVMPin* Pin : InNode->GetPins())
+		{
+			check(!Pin->IsWildCard());
+			TypeMap.Add(Pin->GetFName(), FRigVMTemplateArgument::FType(Pin->GetCPPType(), Pin->GetCPPTypeObject()));
+		}
+
+		TArray<int32> Permutations;
+		Template->Resolve(TypeMap, Permutations, false);
+		check(Permutations.Num() == 1);
+		InputPermutation = Permutations[0];
+	}
+	
+
+	const FRigVMFunction* ResolvedFunction = Template->GetPermutation(InputPermutation);
+	const TArray<int32> PermutationIndices = {InputPermutation};
 	InNode->FilteredPermutations = PermutationIndices;
 
 	// find all existing pins that we may need to change
@@ -15369,10 +15387,14 @@ bool URigVMController::UpdateFilteredPermutations(URigVMPin* InPin, const TArray
 
 bool URigVMController::UpdateTemplateNodePinTypes(URigVMTemplateNode* InNode, bool bSetupUndoRedo)
 {
-	bool bAnyWildcard = false;;
-	bool bAnyPinChanged = false;
-	for (URigVMPin* Pin : InNode->GetPins())
+	for(int32 PinIndex=0; PinIndex < InNode->GetPins().Num(); ++PinIndex)
 	{
+		URigVMPin* Pin = InNode->GetPins()[PinIndex];
+		if (Pin->GetDirection() == ERigVMPinDirection::Hidden)
+		{
+			continue;
+		}
+		
 		TArray<FRigVMTemplateArgument::FType> Types = InNode->GetFilteredTypesForPin(Pin);
 		if (Types.IsEmpty())
 		{
@@ -15414,6 +15436,11 @@ bool URigVMController::UpdateTemplateNodePinTypes(URigVMTemplateNode* InNode, bo
 		if (Types.Num() > 1)
 		{
 			// Unresolve
+			if (Pin->HasInjectedNodes())
+			{
+				EjectNodeFromPin(Pin, bSetupUndoRedo);
+			}
+			
 			const FRigVMTemplateArgument* Argument = InNode->GetTemplate()->FindArgument(*Pin->GetName());
 			const FRigVMTemplateArgument::EArrayType ArrayType = Argument->GetArrayType();
 			
@@ -15431,10 +15458,8 @@ bool URigVMController::UpdateTemplateNodePinTypes(URigVMTemplateNode* InNode, bo
 			
 			if (Pin->GetCPPType() != CPPType || Pin->GetCPPTypeObject() != CPPObjectType)
 			{
-				bAnyPinChanged = true;
 				ChangePinType(Pin, CPPType, CPPObjectType, bSetupUndoRedo, false, false, false);
 			}
-			bAnyWildcard = true;
 		}
 		else if (Types.Num() == 1)
 		{
@@ -15443,7 +15468,6 @@ bool URigVMController::UpdateTemplateNodePinTypes(URigVMTemplateNode* InNode, bo
 			UObject* CPPObjectType = Types[0].CPPTypeObject;
 			if (Pin->GetCPPType() != CPPType || Pin->GetCPPTypeObject() != CPPObjectType)
 			{
-				bAnyPinChanged = true;
 				ChangePinType(Pin, CPPType, CPPObjectType, bSetupUndoRedo, false, false, false);
 			}
 		}
@@ -15454,11 +15478,6 @@ bool URigVMController::UpdateTemplateNodePinTypes(URigVMTemplateNode* InNode, bo
 		}
 	}
 
-	if (!bAnyWildcard && bAnyPinChanged)
-	{
-		FullyResolveTemplateNode(InNode, InNode->FilteredPermutations[0], bSetupUndoRedo);
-	}
-
 	return true;
 }
 
@@ -15717,11 +15736,6 @@ void URigVMController::InitializeFilteredPermutationsFromTemplateTypes()
 		if (URigVMTemplateNode* TemplateNode = Cast<URigVMTemplateNode>(Node))
 		{
 			TemplateNode->InitializeFilteredPermutationsFromTypes();
-			// if (!TemplateNode->IsSingleton())
-			// {
-			// 	FullyResolveTemplateNode(TemplateNode, TemplateNode->FilteredPermutations[0], false);
-			// }
-			//
 		}
 	}
 }
@@ -15960,8 +15974,13 @@ bool URigVMController::ChangePinType(URigVMPin* InPin, const FString& InCPPType,
 	if (InPin->IsArray())
 	{
 		const FString BaseCPPType = RigVMTypeUtils::BaseTypeFromArrayType(CPPType);
-		for (URigVMPin* SubPin : InPin->GetSubPins())
+		for (int32 i=0; i<InPin->GetSubPins().Num(); ++i)
 		{
+			URigVMPin* SubPin = InPin->GetSubPins()[i];
+			if (SubPin->GetDirection() == ERigVMPinDirection::Hidden)
+			{
+				continue;
+			}
 			ChangePinType(SubPin, BaseCPPType, InCPPTypeObject, bSetupUndoRedo, bSetupOrphanPins, bBreakLinks, bRemoveSubPins);
 		}
 	}
@@ -16098,6 +16117,13 @@ bool URigVMController::ChangePinType(URigVMPin* InPin, const FString& InCPPType,
 					Local::ApplyResolvedDefaultValue(this, Pin, FString(), SimpleTypeDefaultValue, bSetupUndoRedo);
 				}
 			}
+
+			if (URigVMTemplateNode* TemplateNode = Cast<URigVMTemplateNode>(InPin->GetNode()))
+			{
+				// Figure out the permutation from the pin types. During undo, the filtered permutations are not
+				// reliable as to which permutation we are resolving to.
+				FullyResolveTemplateNode(TemplateNode, INDEX_NONE, bSetupUndoRedo);
+			}
 		}
 	}
 
