Control Rig: Fix load of pre-template assets

diff --git a/Engine/Plugins/Animation/ControlRig/Source/ControlRigDeveloper/Private/ControlRigBlueprint.cpp b/Engine/Plugins/Animation/ControlRig/Source/ControlRigDeveloper/Private/ControlRigBlueprint.cpp
--- a/Engine/Plugins/Animation/ControlRig/Source/ControlRigDeveloper/Private/ControlRigBlueprint.cpp
+++ b/Engine/Plugins/Animation/ControlRig/Source/ControlRigDeveloper/Private/ControlRigBlueprint.cpp
@@ -479,8 +479,6 @@ void UControlRigBlueprint::PostLoad()
 		PatchBoundVariables();
 		PatchPropagateToChildren();
 		PatchParameterNodesOnLoad();
-		PatchTemplateNodesWithPreferredPermutation();
-		
 
 #if WITH_EDITOR
 
@@ -918,6 +916,7 @@ void UControlRigBlueprint::RefreshAllModels()
 		Controller->ReattachLinksToPinObjects(true /* follow redirectors */, nullptr, false, true);
 	}
 
+	PatchTemplateNodesWithPreferredPermutation();
 	
 	TArray<URigVMGraph*> GraphsToClean = GetAllModels();
 	
@@ -935,9 +934,15 @@ void UControlRigBlueprint::RefreshAllModels()
 			if (URigVMTemplateNode* TemplateNode = Cast<URigVMTemplateNode>(ModelNode))
 			{
 				TemplateNode->InvalidateCache();
+				TemplateNode->PostLoad();
 			}
 		}
-		//Controller->RecomputeAllTemplateFilteredTypes(false);
+#if WITH_EDITOR
+		// if (Controller->RecomputeAllTemplateFilteredTypes(false))
+		// {
+		// 	ensureMsgf(false, TEXT("Pin type changed during load %s"), *GetPackage()->GetPathName());
+		// }
+#endif
 	}
 }
 
@@ -4123,6 +4128,7 @@ void UControlRigBlueprint::PatchTemplateNodesWithPreferredPermutation()
 				if (URigVMTemplateNode* TemplateNode = Cast<URigVMTemplateNode>(ModelNode))
 				{
 					TemplateNode->InvalidateCache();
+					TemplateNode->PostLoad();
 				}
 			}
 			
diff --git a/Engine/Source/Developer/RigVMDeveloper/Private/RigVMModel/RigVMController.cpp b/Engine/Source/Developer/RigVMDeveloper/Private/RigVMModel/RigVMController.cpp
--- a/Engine/Source/Developer/RigVMDeveloper/Private/RigVMModel/RigVMController.cpp
+++ b/Engine/Source/Developer/RigVMDeveloper/Private/RigVMModel/RigVMController.cpp
@@ -14837,8 +14837,8 @@ bool URigVMController::FullyResolveTemplateNode(URigVMTemplateNode* InNode, int3
 		}
 
 		TArray<int32> Permutations;
-		Template->Resolve(TypeMap, Permutations, false);
-		check(Permutations.Num() == 1);
+		Template->Resolve(TypeMap, Permutations, true);
+		check(!Permutations.IsEmpty());
 		InputPermutation = Permutations[0];
 	}
 	
@@ -15387,6 +15387,7 @@ bool URigVMController::UpdateFilteredPermutations(URigVMPin* InPin, const TArray
 
 bool URigVMController::UpdateTemplateNodePinTypes(URigVMTemplateNode* InNode, bool bSetupUndoRedo)
 {
+	bool bAnyTypeChanged = false;
 	for(int32 PinIndex=0; PinIndex < InNode->GetPins().Num(); ++PinIndex)
 	{
 		URigVMPin* Pin = InNode->GetPins()[PinIndex];
@@ -15458,6 +15459,7 @@ bool URigVMController::UpdateTemplateNodePinTypes(URigVMTemplateNode* InNode, bo
 			
 			if (Pin->GetCPPType() != CPPType || Pin->GetCPPTypeObject() != CPPObjectType)
 			{
+				bAnyTypeChanged = true;
 				ChangePinType(Pin, CPPType, CPPObjectType, bSetupUndoRedo, false, false, false);
 			}
 		}
@@ -15468,6 +15470,7 @@ bool URigVMController::UpdateTemplateNodePinTypes(URigVMTemplateNode* InNode, bo
 			UObject* CPPObjectType = Types[0].CPPTypeObject;
 			if (Pin->GetCPPType() != CPPType || Pin->GetCPPTypeObject() != CPPObjectType)
 			{
+				bAnyTypeChanged = true;
 				ChangePinType(Pin, CPPType, CPPObjectType, bSetupUndoRedo, false, false, false);
 			}
 		}
@@ -15478,7 +15481,7 @@ bool URigVMController::UpdateTemplateNodePinTypes(URigVMTemplateNode* InNode, bo
 		}
 	}
 
-	return true;
+	return bAnyTypeChanged;
 }
 
 bool URigVMController::PropagateTemplateFilteredTypes(URigVMTemplateNode* InNode, bool bSetupUndoRedo) 
@@ -15555,16 +15558,16 @@ bool URigVMController::PropagateTemplateFilteredTypes(URigVMTemplateNode* InNode
 	return true;
 }
 
-void URigVMController::RecomputeAllTemplateFilteredTypes(bool bSetupUndoRedo)
+bool URigVMController::RecomputeAllTemplateFilteredTypes(bool bSetupUndoRedo)
 {
 	if (!IsValidGraph())
 	{
-		return;
+		return false;
 	}
 
 	if (!bIsTransacting && !IsGraphEditable())
 	{
-		return;
+		return false;
 	}
 
 	URigVMGraph* Graph = GetGraph();
@@ -15702,6 +15705,7 @@ void URigVMController::RecomputeAllTemplateFilteredTypes(bool bSetupUndoRedo)
 	}
 
 	// Now update all template nodes pin types
+	bool bAnyTypeChanged = false;
 	for (URigVMNode* Node : Graph->GetNodes())
 	{
 		if (URigVMTemplateNode* TemplateNode = Cast<URigVMTemplateNode>(Node))
@@ -15711,9 +15715,11 @@ void URigVMController::RecomputeAllTemplateFilteredTypes(bool bSetupUndoRedo)
 				continue;
 			}
 
-			UpdateTemplateNodePinTypes(TemplateNode, bSetupUndoRedo);
+			bAnyTypeChanged |= UpdateTemplateNodePinTypes(TemplateNode, bSetupUndoRedo);
 		}
 	}
+
+	return bAnyTypeChanged;
 }
 
 void URigVMController::InitializeFilteredPermutationsFromTemplateTypes()
diff --git a/Engine/Source/Developer/RigVMDeveloper/Public/RigVMModel/RigVMController.h b/Engine/Source/Developer/RigVMDeveloper/Public/RigVMModel/RigVMController.h
--- a/Engine/Source/Developer/RigVMDeveloper/Public/RigVMModel/RigVMController.h
+++ b/Engine/Source/Developer/RigVMDeveloper/Public/RigVMModel/RigVMController.h
@@ -892,8 +892,9 @@ public:
 	// removes any orphan pins that no longer holds a link
 	bool RemoveUnusedOrphanedPins(URigVMNode* InNode, bool bNotify);
 
-	// Initializes and recomputes the filtered permutations of all template nodes in the graph 
-	void RecomputeAllTemplateFilteredTypes(bool bSetupUndoRedo);
+	// Initializes and recomputes the filtered permutations of all template nodes in the graph
+	// Returns true if any pin has change it's type
+	bool RecomputeAllTemplateFilteredTypes(bool bSetupUndoRedo);
 
 	// Try to initialize the filterd permutations from the pin types
 	void InitializeFilteredPermutationsFromTemplateTypes();
diff --git a/Engine/Source/Runtime/RigVM/Private/RigVMCore/RigVMTemplate.cpp b/Engine/Source/Runtime/RigVM/Private/RigVMCore/RigVMTemplate.cpp
--- a/Engine/Source/Runtime/RigVM/Private/RigVMCore/RigVMTemplate.cpp
+++ b/Engine/Source/Runtime/RigVM/Private/RigVMCore/RigVMTemplate.cpp
@@ -991,10 +991,19 @@ bool FRigVMTemplate::Resolve(FTypeMap& InOutTypes, TArray<int32>& OutPermutation
 				// if we found a perfect match - remove all permutations which don't match this one
 				if(bFoundPerfectMatch)
 				{
+					const TArray<int32> BeforePermutationIndices = OutPermutationIndices;
 					OutPermutationIndices.RemoveAll([Argument, MatchedType](int32 PermutationIndex) -> bool
 					{
 						return Argument.Types[PermutationIndex] != MatchedType;
 					});
+					if (OutPermutationIndices.IsEmpty() && bAllowFloatingPointCasts)
+					{
+						OutPermutationIndices = BeforePermutationIndices;
+						OutPermutationIndices.RemoveAll([Argument, MatchedType, InputType](int32 PermutationIndex) -> bool
+						{
+							return !Argument.Types[PermutationIndex].Matches(InputType->CPPType, true);
+						});
+					}
 				}
 				continue;
 			}
