Control Rig: Allow change pin types on wildcard pins

diff --git a/Engine/Source/Developer/RigVMDeveloper/Private/RigVMModel/RigVMControllerActions.cpp b/Engine/Source/Developer/RigVMDeveloper/Private/RigVMModel/RigVMControllerActions.cpp
--- a/Engine/Source/Developer/RigVMDeveloper/Private/RigVMModel/RigVMControllerActions.cpp
+++ b/Engine/Source/Developer/RigVMDeveloper/Private/RigVMModel/RigVMControllerActions.cpp
@@ -1276,10 +1276,13 @@ FRigVMSetPinDefaultValueAction::FRigVMSetPinDefaultValueAction(URigVMPin* InPin,
 , OldDefaultValue(InPin->GetDefaultValue())
 , NewDefaultValue(InNewDefaultValue)
 {
+	/* Since for template we are chaning types - it is possible that the
+	 * pin is no longer compliant with the old value
 	if(!OldDefaultValue.IsEmpty())
 	{
 		check(InPin->IsValidDefaultValue(OldDefaultValue));
 	}
+	*/
 	if(!NewDefaultValue.IsEmpty())
 	{
 		check(InPin->IsValidDefaultValue(NewDefaultValue));
