RigVM: Restructure entry invocation for nativized VMs

diff --git a/Engine/Plugins/Animation/ControlRig/Source/ControlRig/Private/ControlRig.cpp b/Engine/Plugins/Animation/ControlRig/Source/ControlRig/Private/ControlRig.cpp
--- a/Engine/Plugins/Animation/ControlRig/Source/ControlRig/Private/ControlRig.cpp
+++ b/Engine/Plugins/Animation/ControlRig/Source/ControlRig/Private/ControlRig.cpp
@@ -1256,14 +1256,27 @@ void UControlRig::ExecuteUnits(FRigUnitContext& InOutContext, const FName& InEve
 {
 	if (VM)
 	{
-		// sanity check the validity of the VM to ensure stability.
-		if(!IsValidLowLevel() ||
-			!VM->IsValidLowLevel() ||
-			!VM->GetLiteralMemory()->IsValidLowLevel() ||
-			!VM->GetWorkMemory()->IsValidLowLevel())
+		static constexpr TCHAR InvalidatedVMFormat[] = TEXT("%s: Invalidated VM - aborting execution.");
+		if(VM->IsNativized())
 		{
-			UE_LOG(LogControlRig, Warning, TEXT("%s: Invalidated VM - aborting execution."), *GetClass()->GetName());
-			return;
+			if(!IsValidLowLevel() ||
+				!VM->IsValidLowLevel())
+			{
+				UE_LOG(LogControlRig, Warning, InvalidatedVMFormat, *GetClass()->GetName());
+				return;
+			}
+		}
+		else
+		{
+			// sanity check the validity of the VM to ensure stability.
+			if(!IsValidLowLevel() ||
+				!VM->IsValidLowLevel() ||
+				!VM->GetLiteralMemory()->IsValidLowLevel() ||
+				!VM->GetWorkMemory()->IsValidLowLevel())
+			{
+				UE_LOG(LogControlRig, Warning, InvalidatedVMFormat, *GetClass()->GetName());
+				return;
+			}
 		}
 		
 #if UE_CONTROLRIG_PROFILE_EXECUTE_UNITS_NUM
diff --git a/Engine/Source/Developer/RigVMDeveloper/Private/RigVMCompiler/RigVMCodeGenerator.cpp b/Engine/Source/Developer/RigVMDeveloper/Private/RigVMCompiler/RigVMCodeGenerator.cpp
--- a/Engine/Source/Developer/RigVMDeveloper/Private/RigVMCompiler/RigVMCodeGenerator.cpp
+++ b/Engine/Source/Developer/RigVMDeveloper/Private/RigVMCompiler/RigVMCodeGenerator.cpp
@@ -7,6 +7,7 @@
 #include "Algo/Count.h"
 #include "Animation/Rig.h"
 
+static constexpr TCHAR RigVM_CommaSeparator[] = TEXT(", ");
 static constexpr TCHAR RigVM_NewLineFormat[] = TEXT("\r\n");
 static constexpr TCHAR RigVM_IncludeBracketFormat[] = TEXT("#include <{0}>");
 static constexpr TCHAR RigVM_IncludeQuoteFormat[] = TEXT("#include \"{0}.h\"");
@@ -16,7 +17,7 @@ static constexpr TCHAR RigVM_DeclareExternalVariableFormat[] = TEXT("\t{0}* {1}
 static constexpr TCHAR RigVM_UpdateExternalVariableFormat[] = TEXT("\t{0} = &GetExternalVariableRef<{1}>(TEXT(\"{2}\"), TEXT(\"{1}\"));");
 static constexpr TCHAR RigVM_MemberPropertyFormat[] = TEXT("\t{0} {1};");
 static constexpr TCHAR RigVM_DeclareEntryNameFormat[] = TEXT("\tstatic const FName EntryName_{0};");
-static constexpr TCHAR RigVM_DefineEntryNameFormat[] = TEXT("const FName U{0}::EntryName_{1} = TEXT(\"{1}\");");
+static constexpr TCHAR RigVM_DefineEntryNameFormat[] = TEXT("const FName U{0}::EntryName_{1} = TEXT(\"{2}\");");
 static constexpr TCHAR RigVM_DefineConstFormatNoDefault[] = TEXT("\tstatic const {0} {1};");
 static constexpr TCHAR RigVM_DefineStructConstFormat[] = TEXT("\tstatic const {0} {1} = URigVMNativized::GetStructConstant<{0}>(TEXT(\"{2}\"));");
 static constexpr TCHAR RigVM_DefineStructArrayConstFormat[] = TEXT("\tstatic const {0} {1} = URigVMNativized::GetStructArrayConstant<{2}>(TEXT(\"{3}\"));");
@@ -40,6 +41,11 @@ static constexpr TCHAR RigVM_IncrementOpFormat[] = TEXT("\t{0}++;");
 static constexpr TCHAR RigVM_DecrementOpFormat[] = TEXT("\t{0}--;");
 static constexpr TCHAR RigVM_EqualsOpFormat[] = TEXT("\t{0} = {1} == {2};");
 static constexpr TCHAR RigVM_InvokeEntryFormat[] = TEXT("\tif (InEntryName == EntryName_{0}) return ExecuteEntry_{0}({1});");
+static constexpr TCHAR RigVM_InvokeEntryByNameFormat[] = TEXT("\treturn InvokeEntryByName(InEntryName{0});");
+static constexpr TCHAR RigVM_InvokeEntryByNameFormat2[] = TEXT("\tif(!InvokeEntryByName({0}{1})) return false;");
+static constexpr TCHAR RigVM_CanExecuteEntryFormat[] = TEXT("\tif(!CanExecuteEntry(InEntryName)) { return false; }");
+static constexpr TCHAR RigVM_EntryExecuteGuardFormat[] = TEXT("\tFEntryExecuteGuard EntryExecuteGuard(EntriesBeingExecuted, FindEntry(InEntryName));");
+static constexpr TCHAR RigVM_PublicContextGuardFormat[] = TEXT("\tTGuardValue<FRigVMExecuteContext> PublicContextGuard(Context.PublicData, PublicContext);");
 static constexpr TCHAR RigVM_EntryNameFormat[] = TEXT("EntryName_{0}");
 static constexpr TCHAR RigVM_UpdateContextFormat[] = TEXT("\tconst FRigVMExecuteContext& PublicContext = UpdateContext(AdditionalArguments, {0}, InEntryName);");
 static constexpr TCHAR RigVM_TrueFormat[] = TEXT("true");  
@@ -90,10 +96,12 @@ static constexpr TCHAR RigVM_GeneratedIncludeFormat[] = TEXT("#include \"{0}.gen
 static constexpr TCHAR RigVM_UClassDefinitionFormat[] = TEXT("UCLASS()\r\nclass {0}_API U{1} : public URigVMNativized\r\n{\r\n\tGENERATED_BODY()\r\npublic:\r\n\tU{1}() {}\r\n\tvirtual ~U{1}() override {}\r\n");
 static constexpr TCHAR RigVM_ProtectedFormat[] = TEXT("protected:");
 static constexpr TCHAR RigVM_GetVMHashFormat[] = TEXT("\tvirtual uint32 GetVMHash() const override { return {0}; }");
-static constexpr TCHAR RigVM_GetEntryNamesFormat[] = TEXT("\tvirtual TArray<FName> GetEntryNames() const override { return { {0} }; }");
+static constexpr TCHAR RigVM_GetEntryNamesFormat[] = TEXT("\tvirtual const TArray<FName>& GetEntryNames() const override\r\n\t{\r\n\t\tstatic const TArray<FName> StaticEntryNames = { {0} };\r\n\t\treturn StaticEntryNames;\r\n\t}");
 static constexpr TCHAR RigVM_DeclareUpdateExternalVariablesFormat[] = TEXT("\tvirtual void UpdateExternalVariables() override;");
+static constexpr TCHAR RigVM_DeclareInvokeEntryByNameFormat[] = TEXT("\tbool InvokeEntryByName(const FName& InEntryName{0});");
 static constexpr TCHAR RigVM_DeclareExecuteFormat[] = TEXT("\tvirtual bool Execute(TArrayView<URigVMMemoryStorage*> Memory, TArrayView<void*> AdditionalArguments, const FName& InEntryName) override;");
 static constexpr TCHAR RigVM_DefineUpdateExternalVariablesFormat[] = TEXT("void U{0}::UpdateExternalVariables()\r\n{");
+static constexpr TCHAR RigVM_DefineInvokeEntryByNameFormat[] = TEXT("bool U{0}::InvokeEntryByName(const FName& InEntryName{1})\r\n{");
 static constexpr TCHAR RigVM_DefineExecuteFormat[] = TEXT("bool U{0}::Execute(TArrayView<URigVMMemoryStorage*> Memory, TArrayView<void*> AdditionalArguments, const FName& InEntryName)\r\n{");
 static constexpr TCHAR RigVM_DeclareExecuteEntryFormat[] = TEXT("\tbool ExecuteEntry_{0}({1});");
 static constexpr TCHAR RigVM_DefineExecuteEntryFormat[] = TEXT("bool U{0}::ExecuteEntry_{1}({2})\r\n{");
@@ -107,6 +115,7 @@ static constexpr TCHAR RigVM_GetOperandSliceFormat[] = TEXT("GetOperandSlice<{0}
 static constexpr TCHAR RigVM_ExternalVariableFormat[] = TEXT("(*External_{0})");
 static constexpr TCHAR RigVM_JoinSegmentPathFormat[] = TEXT("{0}.{1}");
 static constexpr TCHAR RigVM_GetArrayElementSafeFormat[] = TEXT("GetArrayElementSafe<{0}>({1}, {2})");
+static constexpr TCHAR RigVM_InvokeEntryOpFormat[] = TEXT("\tif(!InvokeEntryByName(EntryName_{0})) return false;");
 
 FString FRigVMCodeGenerator::DumpIncludes(bool bLog)
 {
@@ -172,14 +181,14 @@ FString FRigVMCodeGenerator::DumpEntries(bool bForHeader, bool bLog)
 	for(int32 EntryIndex = 0; EntryIndex < ByteCode.NumEntries(); EntryIndex++)
 	{
 		const FRigVMByteCodeEntry& Entry = ByteCode.GetEntry(EntryIndex);
-		const FString EntryName = Entry.Name.ToString();
+		const FString EntryName = Entry.GetSanitizedName();
 		if(bForHeader)
 		{
 			Lines.Add(Format(RigVM_DeclareEntryNameFormat, *EntryName));
 		}
 		else
 		{
-			Lines.Add(Format(RigVM_DefineEntryNameFormat, *ClassName, *EntryName));
+			Lines.Add(Format(RigVM_DefineEntryNameFormat, *ClassName, *EntryName, *Entry.Name.ToString()));
 		}
 	}
 
@@ -254,7 +263,7 @@ FString FRigVMCodeGenerator::DumpProperties(bool bForHeader, int32 InOperationGr
 							{
 								DefaultValues[DefaultValueIndex] = (*StructConstGenerator)(DefaultValues[DefaultValueIndex]); 
 							}
-							DefaultValue = Format(RigVM_CurlyBracesFormat, FString::Join(DefaultValues, TEXT(", ")));
+							DefaultValue = Format(RigVM_CurlyBracesFormat, FString::Join(DefaultValues, RigVM_CommaSeparator));
 						}
 						else
 						{
@@ -445,61 +454,105 @@ FString FRigVMCodeGenerator::DumpInstructions(int32 InOperationGroup, bool bLog)
 
 	const FOperationGroup& Group = GetGroup(InOperationGroup);
 
-	if(Group.Entry.IsNone())
+	auto GetEntryParameters = [&]() -> FString
 	{
-		Lines.Add(Format(RigVM_UpdateContextFormat, Operations.Num()));
+		const int32 EntryIndex = ByteCode.NumEntries() - 1;
+		const FRigVMByteCodeEntry& Entry = ByteCode.GetEntry(EntryIndex);
+		const FString EntryName = Entry.GetSanitizedName();
+		FString Parameters;
 
-		// let's get the additional arguments for all sub groups
-		TArray<FString> OpaqueArgumentHit;
+		// find the entry's group and provide the needed arguments
 		for(const FOperationGroup& EntryGroup : OperationGroups)
 		{
-			if(EntryGroup.Entry.IsNone() || EntryGroup.Depth > 0)
+			if(EntryGroup.Depth <= 0 && EntryGroup.Entry == EntryName)
 			{
-				continue;
+				TArray<FString> ArgumentNames = {RigVM_ContextPublicFormat};
+				for(const FOpaqueArgument& OpaqueArgument : EntryGroup.OpaqueArguments)
+				{
+					ArgumentNames.Add(OpaqueArgument.Get<1>());
+				}
+				Parameters = FString::Join(ArgumentNames, RigVM_CommaSeparator);
+				break;
 			}
+		}
+
+		if(!Parameters.IsEmpty())
+		{
+			Parameters = RigVM_CommaSeparator + Parameters;
+		}
+		
+		return Parameters;
+	};
+
+	if(Group.Entry.IsEmpty())
+	{
+		if(InOperationGroup == INDEX_NONE)
+		{
+			Lines.Add(Format(RigVM_UpdateContextFormat, Operations.Num()));
 
-			for(const FOpaqueArgument& OpaqueArgument : EntryGroup.OpaqueArguments)
+			// let's get the additional arguments for all sub groups
+			TArray<FString> OpaqueArgumentHit;
+			for(const FOperationGroup& EntryGroup : OperationGroups)
 			{
-				if(OpaqueArgumentHit.Contains(OpaqueArgument.Get<0>()))
+				if(EntryGroup.Entry.IsEmpty() || EntryGroup.Depth > 0)
 				{
 					continue;
 				}
-				const int32 ArgumentIndex = OpaqueArgumentHit.Add(OpaqueArgument.Get<0>());
-				FString ArgumentTypeNoRef = OpaqueArgument.Get<2>();
-				if (ArgumentTypeNoRef.EndsWith(TEXT("&")))
+
+				for(const FOpaqueArgument& OpaqueArgument : EntryGroup.OpaqueArguments)
 				{
-					ArgumentTypeNoRef.LeftChopInline(1);
+					if(OpaqueArgumentHit.Contains(OpaqueArgument.Get<0>()))
+					{
+						continue;
+					}
+					const int32 ArgumentIndex = OpaqueArgumentHit.Add(OpaqueArgument.Get<0>());
+					FString ArgumentTypeNoRef = OpaqueArgument.Get<2>();
+					if (ArgumentTypeNoRef.EndsWith(TEXT("&")))
+					{
+						ArgumentTypeNoRef.LeftChopInline(1);
+					}
+					Lines.Add(Format(RigVM_FullOpaqueArgumentFormat, *OpaqueArgument.Get<2>(), *OpaqueArgument.Get<1>(), *ArgumentTypeNoRef, ArgumentIndex));
 				}
-				Lines.Add(Format(RigVM_FullOpaqueArgumentFormat, *OpaqueArgument.Get<2>(), *OpaqueArgument.Get<1>(), *ArgumentTypeNoRef, ArgumentIndex));
 			}
-		}
 
-		Lines.Emplace();
-		for(int32 EntryIndex = 0; EntryIndex < ByteCode.NumEntries(); EntryIndex++)
-		{
-			const FRigVMByteCodeEntry& Entry = ByteCode.GetEntry(EntryIndex);
-			const FString EntryName = Entry.Name.ToString();
-			FString OpaqueArguments;
+			Lines.Add(Format(RigVM_InvokeEntryByNameFormat, *GetEntryParameters()));
 
-			// find the entry's group and provide the needed arguments
-			for(const FOperationGroup& EntryGroup : OperationGroups)
+			return DumpLines(Lines, bLog);
+		}
+		else
+		{
+			Lines.Add(FString(RigVM_CanExecuteEntryFormat));
+			Lines.Emplace();
+			Lines.Add(FString(RigVM_EntryExecuteGuardFormat));
+			Lines.Add(FString(RigVM_PublicContextGuardFormat));
+			Lines.Emplace();
+			
+			for(int32 EntryIndex = 0; EntryIndex < ByteCode.NumEntries(); EntryIndex++)
 			{
-				if(EntryGroup.Depth <= 0 && EntryGroup.Entry == Entry.Name)
+				const FRigVMByteCodeEntry& Entry = ByteCode.GetEntry(EntryIndex);
+				const FString EntryName = Entry.GetSanitizedName();
+				FString OpaqueArguments;
+
+				// find the entry's group and provide the needed arguments
+				for(const FOperationGroup& EntryGroup : OperationGroups)
 				{
-					TArray<FString> ArgumentNames = {RigVM_ContextPublicFormat};
-					for(const FOpaqueArgument& OpaqueArgument : EntryGroup.OpaqueArguments)
+					if(EntryGroup.Depth <= 0 && EntryGroup.Entry == EntryName)
 					{
-						ArgumentNames.Add(OpaqueArgument.Get<1>());
+						TArray<FString> ArgumentNames = {RigVM_ContextPublicFormat};
+						for(const FOpaqueArgument& OpaqueArgument : EntryGroup.OpaqueArguments)
+						{
+							ArgumentNames.Add(OpaqueArgument.Get<1>());
+						}
+						OpaqueArguments = FString::Join(ArgumentNames, RigVM_CommaSeparator);
+						break;
 					}
-					OpaqueArguments = FString::Join(ArgumentNames, TEXT(", "));
-					break;
 				}
+				Lines.Add(Format(RigVM_InvokeEntryFormat, *EntryName, *OpaqueArguments));
 			}
-			Lines.Add(Format(RigVM_InvokeEntryFormat, *EntryName, *OpaqueArguments));
 		}
 	}
 
-	if(Group.Entry.IsNone() && OperationGroups.Num() > 0)
+	if(Group.Entry.IsEmpty() && OperationGroups.Num() > 0)
 	{
 		Lines.Add(RigVM_ReturnFalseFormat);
 		return DumpLines(Lines, bLog);
@@ -622,7 +675,7 @@ FString FRigVMCodeGenerator::DumpInstructions(int32 InOperationGroup, bool bLog)
 						}
 					}
 
-					const FString JoinedArguments = FString::Join(Arguments, TEXT(", "));
+					const FString JoinedArguments = FString::Join(Arguments, RigVM_CommaSeparator);
 					Lines.Add(Format(RigVM_CallExternOpFormat, *Function->Struct->GetStructCPPName(), *Function->GetMethodName().ToString(), *JoinedArguments));
 					break;
 				}
@@ -876,6 +929,13 @@ FString FRigVMCodeGenerator::DumpInstructions(int32 InOperationGroup, bool bLog)
 					Lines.Add(Format(RigVM_ArrayReverseOpFormat, *TemplateSuffix, *GetOperandName(Op.Arg, false)));
 					break;
 				}
+				case ERigVMOpCode::InvokeEntry:
+				{
+					const FRigVMInvokeEntryOp& Op = ByteCode.GetOpAt<FRigVMInvokeEntryOp>(Operation);
+					const FString EntryName = Op.EntryName.ToString();
+					Lines.Add(Format(RigVM_InvokeEntryByNameFormat2, *EntryName, *GetEntryParameters()));
+					break;
+				}
 				case ERigVMOpCode::Invalid:
 				case ERigVMOpCode::ChangeType:
 				default:
@@ -902,8 +962,8 @@ FString FRigVMCodeGenerator::DumpInstructions(int32 InOperationGroup, bool bLog)
 					ParameterArray.Add(OpaqueArgument.Get<1>());
 				}
 			}
-			Parameters = FString::Join(ParameterArray, TEXT(", "));
-			Lines.Add(Format(RigVM_InvokeExecuteGroupFormat, *ChildGroup.Entry.ToString(), ChildGroupIndex, *Parameters));
+			Parameters = FString::Join(ParameterArray, RigVM_CommaSeparator);
+			Lines.Add(Format(RigVM_InvokeExecuteGroupFormat, *ChildGroup.Entry, ChildGroupIndex, *Parameters));
 		}
 	}
 	
@@ -925,7 +985,7 @@ FString FRigVMCodeGenerator::DumpHeader(bool bLog)
 	for(int32 EntryIndex = 0; EntryIndex < ByteCode.NumEntries(); EntryIndex++)
 	{
 		const FRigVMByteCodeEntry& Entry = ByteCode.GetEntry(EntryIndex);
-		FormattedEntries.Add(Format(RigVM_EntryNameFormat, *Entry.Name.ToString()));
+		FormattedEntries.Add(Format(RigVM_EntryNameFormat, *Entry.GetSanitizedName()));
 	}
 
 	FStringArray Lines;
@@ -945,7 +1005,7 @@ FString FRigVMCodeGenerator::DumpHeader(bool bLog)
 	}
 	Lines.Add(Format(RigVM_UClassDefinitionFormat, *ModuleName.ToUpper(), *ClassName));
 	Lines.Add(Format(RigVM_GetVMHashFormat, VM->GetVMHash()));
-	Lines.Add(Format(RigVM_GetEntryNamesFormat, *FString::Join(FormattedEntries, TEXT(", "))));
+	Lines.Add(Format(RigVM_GetEntryNamesFormat, *FString::Join(FormattedEntries, RigVM_CommaSeparator)));
 	Lines.Emplace();
 	Lines.Add(FString(RigVM_DeclareExecuteFormat));
 	Lines.Emplace();
@@ -968,15 +1028,24 @@ FString FRigVMCodeGenerator::DumpHeader(bool bLog)
 				ParameterArray.Add(Format(RigVM_ParameterOpaqueArgumentFormat, *OpaqueArgument.Get<2>(), *OpaqueArgument.Get<1>()));
 			}
 		}
-		Parameters = FString::Join(ParameterArray, TEXT(", "));
+		Parameters = FString::Join(ParameterArray, RigVM_CommaSeparator);
 		
 		if(Group.Depth == 0)
 		{
-			Lines.Add(Format(RigVM_DeclareExecuteEntryFormat, *Group.Entry.ToString(), *Parameters));
+			Lines.Add(Format(RigVM_DeclareExecuteEntryFormat, *Group.Entry, *Parameters));
 		}
 		else
 		{
-			Lines.Add(Format(RigVM_DeclareExecuteGroupFormat, *Group.Entry.ToString(), GroupIndex, *Parameters));
+			Lines.Add(Format(RigVM_DeclareExecuteGroupFormat, *Group.Entry, GroupIndex, *Parameters));
+		}
+
+		if(GroupIndex == OperationGroups.Num() - 1)
+		{
+			if(!Parameters.IsEmpty())
+			{
+				Parameters = RigVM_CommaSeparator + Parameters;
+			}
+			Lines.Add(Format(RigVM_DeclareInvokeEntryByNameFormat, *Parameters));
 		}
 	}
 
@@ -1034,16 +1103,16 @@ FString FRigVMCodeGenerator::DumpSource(bool bLog)
 				ParameterArray.Add(Format(RigVM_ParameterOpaqueArgumentFormat, *OpaqueArgument.Get<2>(), *OpaqueArgument.Get<1>()));
 			}
 		}
-		Parameters = FString::Join(ParameterArray, TEXT(", "));
+		Parameters = FString::Join(ParameterArray, RigVM_CommaSeparator);
 
 		Lines.Emplace();
 		if(Group.Depth == 0)
 		{
-			Lines.Add(Format(RigVM_DefineExecuteEntryFormat, *ClassName, *Group.Entry.ToString(), *Parameters));
+			Lines.Add(Format(RigVM_DefineExecuteEntryFormat, *ClassName, *Group.Entry, *Parameters));
 		}
 		else
 		{
-			Lines.Add(Format(RigVM_DefineExecuteGroupFormat, *ClassName, *Group.Entry.ToString(), GroupIndex, *Parameters));
+			Lines.Add(Format(RigVM_DefineExecuteGroupFormat, *ClassName, *Group.Entry, GroupIndex, *Parameters));
 		}
 
 		const FString DumpedProperties = DumpProperties(false, GroupIndex);
@@ -1054,6 +1123,19 @@ FString FRigVMCodeGenerator::DumpSource(bool bLog)
 		}
 		Lines.Add(DumpInstructions(GroupIndex));
 		Lines.Add(TEXT("}"));
+
+		if(GroupIndex == OperationGroups.Num() - 1)
+		{
+			if(!Parameters.IsEmpty())
+			{
+				Parameters = RigVM_CommaSeparator + Parameters;
+			}
+
+			Lines.Emplace();
+			Lines.Add(Format(RigVM_DefineInvokeEntryByNameFormat, *ClassName, *Parameters));
+			Lines.Add(DumpInstructions(-2));
+			Lines.Add(TEXT("}"));
+		}
 	}
 	
 	return DumpLines(Lines, bLog);
@@ -1232,7 +1314,7 @@ void FRigVMCodeGenerator::ParseOperationGroups()
 		const FRigVMByteCodeEntry& Entry = ByteCode.GetEntry(EntryIndex);
 
 		FOperationGroup Group;
-		Group.Entry = Entry.Name;
+		Group.Entry = Entry.GetSanitizedName();
 		Group.Depth = 0;
 		Group.First = Entry.InstructionIndex;
 		Group.Last = Operations.Num() - 1;
diff --git a/Engine/Source/Developer/RigVMDeveloper/Public/RigVMCompiler/RigVMCodeGenerator.h b/Engine/Source/Developer/RigVMDeveloper/Public/RigVMCompiler/RigVMCodeGenerator.h
--- a/Engine/Source/Developer/RigVMDeveloper/Public/RigVMCompiler/RigVMCodeGenerator.h
+++ b/Engine/Source/Developer/RigVMDeveloper/Public/RigVMCompiler/RigVMCodeGenerator.h
@@ -58,7 +58,7 @@ private:
 
 	struct FOperationGroup
 	{
-		FName Entry;
+		FString Entry;
 		int32 First;
 		int32 Last;
 		int32 Depth;
@@ -68,7 +68,7 @@ private:
 		TArray<FOpaqueArgument> OpaqueArguments;
 
 		FOperationGroup()
-			: Entry(NAME_None)
+			: Entry()
 			, First(INDEX_NONE)
 			, Last(INDEX_NONE)
 			, Depth(INDEX_NONE)
diff --git a/Engine/Source/Runtime/RigVM/Private/RigVMCore/RigVM.cpp b/Engine/Source/Runtime/RigVM/Private/RigVMCore/RigVM.cpp
--- a/Engine/Source/Runtime/RigVM/Private/RigVMCore/RigVM.cpp
+++ b/Engine/Source/Runtime/RigVM/Private/RigVMCore/RigVM.cpp
@@ -72,7 +72,6 @@ URigVM::URigVM()
 	, DebugMemoryStorageObject(nullptr)
 	, ByteCodePtr(&ByteCodeStorage)
 	, NumExecutions(0)
-	, bHitErrorDuringExecution(false)
 #if WITH_EDITOR
 	, DebugInfo(nullptr)
 	, HaltedAtBreakpoint(nullptr)
@@ -845,10 +844,16 @@ bool URigVM::ContainsEntry(const FName& InEntryName) const
 	return ByteCode.FindEntryIndex(InEntryName) != INDEX_NONE;
 }
 
-TArray<FName> URigVM::GetEntryNames() const
+int32 URigVM::FindEntry(const FName& InEntryName) const
 {
-	TArray<FName> EntryNames;
+	const FRigVMByteCode& ByteCode = GetByteCode();
+	return ByteCode.FindEntryIndex(InEntryName);
+}
 
+const TArray<FName>& URigVM::GetEntryNames() const
+{
+	EntryNames.Reset();
+	
 	const FRigVMByteCode& ByteCode = GetByteCode();
 	for (int32 EntryIndex = 0; EntryIndex < ByteCode.NumEntries(); EntryIndex++)
 	{
@@ -858,6 +863,33 @@ TArray<FName> URigVM::GetEntryNames() const
 	return EntryNames;
 }
 
+bool URigVM::CanExecuteEntry(const FName& InEntryName) const
+{
+	const int32 EntryIndex = FindEntry(InEntryName);
+	if(EntryIndex == INDEX_NONE)
+	{
+		static constexpr TCHAR MissingEntry[] = TEXT("Entry('%s') cannot be found.");
+		Context.PublicData.Logf(EMessageSeverity::Error, MissingEntry, *InEntryName.ToString());
+		return false;
+	}
+	
+	if(EntriesBeingExecuted.Contains(EntryIndex))
+	{
+		TArray<FString> EntryNamesBeingExecuted;
+		for(const int32 EntryBeingExecuted : EntriesBeingExecuted)
+		{
+			EntryNamesBeingExecuted.Add(GetEntryNames()[EntryBeingExecuted].ToString());
+		}
+		EntryNamesBeingExecuted.Add(InEntryName.ToString());
+
+		static constexpr TCHAR RecursiveEntry[] = TEXT("Entry('%s') is being invoked recursively (%s).");
+		Context.PublicData.Logf(EMessageSeverity::Error, RecursiveEntry, *InEntryName.ToString(), *FString::Join(EntryNamesBeingExecuted, TEXT(" -> ")));
+		return false;
+	}
+
+	return true;
+}
+
 #if WITH_EDITOR
 
 bool URigVM::ResumeExecution()
@@ -1592,8 +1624,6 @@ bool URigVM::Execute(TArrayView<URigVMMemoryStorage*> Memory, TArrayView<void*>
 	const bool bIsRootEntry = EntriesBeingExecuted.IsEmpty(); 
 	if(bIsRootEntry)
 	{
-		bHitErrorDuringExecution = false;
-		
 		if (ExecutingThreadId != INDEX_NONE)
 		{
 			ensureMsgf(ExecutingThreadId == FPlatformTLS::GetCurrentThreadId(), TEXT("RigVM::Execute from multiple threads (%d and %d)"), ExecutingThreadId, (int32)FPlatformTLS::GetCurrentThreadId());
@@ -1660,6 +1690,7 @@ bool URigVM::Execute(TArrayView<URigVMMemoryStorage*> Memory, TArrayView<void*>
 		ClearDebugMemory();
 	}
 
+	int32 EntryIndexToPush = INDEX_NONE;
 	if (!InEntryName.IsNone())
 	{
 		int32 EntryIndex = ByteCode.FindEntryIndex(InEntryName);
@@ -1667,8 +1698,9 @@ bool URigVM::Execute(TArrayView<URigVMMemoryStorage*> Memory, TArrayView<void*>
 		{
 			return false;
 		}
+		
 		SetInstructionIndex((uint16)ByteCode.GetEntry(EntryIndex).InstructionIndex);
-		EntriesBeingExecuted.Push(EntryIndex);
+		EntryIndexToPush = EntryIndex;
 
 		if(bIsRootEntry)
 		{
@@ -1677,23 +1709,20 @@ bool URigVM::Execute(TArrayView<URigVMMemoryStorage*> Memory, TArrayView<void*>
 	}
 	else
 	{
-		int32 FirstEntryIndex = INDEX_NONE;
 		for(int32 EntryIndex = 0; EntryIndex < ByteCode.NumEntries(); EntryIndex++)
 		{
 			if(ByteCode.GetEntry(EntryIndex).InstructionIndex == 0)
 			{
-				FirstEntryIndex = EntryIndex;
+				EntryIndexToPush = EntryIndex;
 				break;
 			}
 		}
 		
-		EntriesBeingExecuted.Push(FirstEntryIndex);
-
 		if(bIsRootEntry)
 		{
-			if(ByteCode.Entries.IsValidIndex(FirstEntryIndex))
+			if(ByteCode.Entries.IsValidIndex(EntryIndexToPush))
 			{
-				Context.PublicData.EventName = ByteCode.GetEntry(FirstEntryIndex).Name;
+				Context.PublicData.EventName = ByteCode.GetEntry(EntryIndexToPush).Name;
 			}
 			else
 			{
@@ -1702,6 +1731,8 @@ bool URigVM::Execute(TArrayView<URigVMMemoryStorage*> Memory, TArrayView<void*>
 		}
 	}
 
+	FEntryExecuteGuard EntryExecuteGuard(EntriesBeingExecuted, EntryIndexToPush);
+
 #if WITH_EDITOR
 	if (DebugInfo && bIsRootEntry)
 	{
@@ -1725,18 +1756,11 @@ bool URigVM::Execute(TArrayView<URigVMMemoryStorage*> Memory, TArrayView<void*>
 
 	while (Instructions.IsValidIndex(Context.PublicData.InstructionIndex))
 	{
-		if(bHitErrorDuringExecution)
-		{
-			EntriesBeingExecuted.Pop();
-			return false;
-		}
-		
 #if WITH_EDITOR
 		if (DebugInfo && ShouldHaltAtInstruction(InEntryName, Context.PublicData.InstructionIndex))
 		{
 			// we'll recursively exit all invoked
 			// entries here.
-			EntriesBeingExecuted.Pop();
 			return true;
 		}
 
@@ -2004,7 +2028,6 @@ bool URigVM::Execute(TArrayView<URigVMMemoryStorage*> Memory, TArrayView<void*>
 					}
 #endif
 				}
-				EntriesBeingExecuted.Pop();
 				return true;
 			}
 			case ERigVMOpCode::BeginBlock:
@@ -2563,38 +2586,24 @@ bool URigVM::Execute(TArrayView<URigVMMemoryStorage*> Memory, TArrayView<void*>
 			case ERigVMOpCode::InvokeEntry:
 			{
 				const FRigVMInvokeEntryOp& Op = ByteCode.GetOpAt<FRigVMInvokeEntryOp>(Instruction);
-				const int32 EntryIndex = ByteCode.FindEntryIndex(Op.EntryName);
 
-				if(EntryIndex == INDEX_NONE)
+				if(!CanExecuteEntry(Op.EntryName))
 				{
-					static constexpr TCHAR MissingEntry[] = TEXT("Entry('%s') cannot be found.");
-					Context.PublicData.Logf(EMessageSeverity::Error, MissingEntry, *Op.EntryName.ToString());
-					bHitErrorDuringExecution = true;
-				}
-				else if(EntriesBeingExecuted.Contains(EntryIndex)) // avoid cycles
-				{
-					TArray<FString> EntryNames;
-					for(int32 EntryBeingExecuted : EntriesBeingExecuted)
-					{
-						EntryNames.Add(ByteCode.GetEntry(EntryBeingExecuted).Name.ToString());
-					}
-					EntryNames.Add(Op.EntryName.ToString());
-
-					static constexpr TCHAR RecursiveEntry[] = TEXT("Entry('%s') is being invoked recursively (%s).");
-					Context.PublicData.Logf(EMessageSeverity::Error, RecursiveEntry, *Op.EntryName.ToString(), *FString::Join(EntryNames, TEXT(" -> ")));
-					bHitErrorDuringExecution = true;
+					return false;
 				}
 				else
 				{
 					// this will restore the public data after invoking the entry
 					TGuardValue<FRigVMExecuteContext> PublicDataGuard(Context.PublicData, Context.PublicData);
-					Execute(Memory, AdditionalArguments, Op.EntryName);
+					if(!Execute(Memory, AdditionalArguments, Op.EntryName))
+					{
+						return false;
+					}
 
 #if WITH_EDITOR
 					// if we are halted at a break point we need to exit here
 					if (DebugInfo && ShouldHaltAtInstruction(InEntryName, Context.PublicData.InstructionIndex))
 					{
-						EntriesBeingExecuted.Pop();
 						return true;
 					}
 #endif
@@ -2606,7 +2615,6 @@ bool URigVM::Execute(TArrayView<URigVMMemoryStorage*> Memory, TArrayView<void*>
 			case ERigVMOpCode::Invalid:
 			{
 				ensure(false);
-				EntriesBeingExecuted.Pop();
 				return false;
 			}
 		}
@@ -2643,7 +2651,6 @@ bool URigVM::Execute(TArrayView<URigVMMemoryStorage*> Memory, TArrayView<void*>
 	}
 #endif
 
-	EntriesBeingExecuted.Pop();
 	return true;
 }
 
diff --git a/Engine/Source/Runtime/RigVM/Private/RigVMCore/RigVMByteCode.cpp b/Engine/Source/Runtime/RigVM/Private/RigVMCore/RigVMByteCode.cpp
--- a/Engine/Source/Runtime/RigVM/Private/RigVMCore/RigVMByteCode.cpp
+++ b/Engine/Source/Runtime/RigVM/Private/RigVMCore/RigVMByteCode.cpp
@@ -194,6 +194,14 @@ void FRigVMInstructionArray::Empty()
 
 TArray<int32> FRigVMByteCode::EmptyInstructionIndices;
 
+FString FRigVMByteCodeEntry::GetSanitizedName() const
+{
+	FString SanitizedName = Name.ToString();
+	SanitizedName.ReplaceInline(TEXT(" "), TEXT("_"));
+	SanitizedName.TrimStartAndEndInline();
+	return SanitizedName;
+}
+
 FRigVMByteCode::FRigVMByteCode()
 	: NumInstructions(0)
 	, bByteCodeIsAligned(false)
diff --git a/Engine/Source/Runtime/RigVM/Public/RigVMCore/RigVM.h b/Engine/Source/Runtime/RigVM/Public/RigVMCore/RigVM.h
--- a/Engine/Source/Runtime/RigVM/Public/RigVMCore/RigVM.h
+++ b/Engine/Source/Runtime/RigVM/Public/RigVMCore/RigVM.h
@@ -249,8 +249,14 @@ public:
 	// Returns true if this VM's bytecode contains a given entry
 	virtual bool ContainsEntry(const FName& InEntryName) const;
 
+	// Returns the index of an entry
+	virtual int32 FindEntry(const FName& InEntryName) const;
+
 	// Returns a list of all valid entry names for this VM's bytecode
-	virtual TArray<FName> GetEntryNames() const;
+	virtual const TArray<FName>& GetEntryNames() const;
+
+	// returns false if an entry can not be executed
+	bool CanExecuteEntry(const FName& InEntryName) const;
 
 #if WITH_EDITOR
 	
@@ -653,8 +659,6 @@ private:
 	UPROPERTY(transient)
 	uint32 NumExecutions;
 
-	bool bHitErrorDuringExecution;
-
 #if WITH_EDITOR
 	FRigVMDebugInfo* DebugInfo;
 	TSharedPtr<FRigVMBreakpoint> HaltedAtBreakpoint;
@@ -745,7 +749,32 @@ private:
 	TMap<FRigVMOperand, TArray<FRigVMOperand>> OperandToDebugRegisters;
 
 	int32 ExecutingThreadId;
+	
+protected:
+
+	struct FEntryExecuteGuard
+	{
+	public:
+		FEntryExecuteGuard(TArray<int32>& InOutStack, int32 InEntryIndex)
+		: Stack(InOutStack)
+		{
+			Stack.Push(InEntryIndex);
+		}
+
+		~FEntryExecuteGuard()
+		{
+			Stack.Pop();
+		}
+	private:
+
+		TArray<int32>& Stack;
+	};
+	
 	TArray<int32> EntriesBeingExecuted;
+	
+private:
+	
+	mutable TArray<FName> EntryNames;
 
 	UPROPERTY(transient)
 	TObjectPtr<URigVM> DeferredVMToCopy;
diff --git a/Engine/Source/Runtime/RigVM/Public/RigVMCore/RigVMByteCode.h b/Engine/Source/Runtime/RigVM/Public/RigVMCore/RigVMByteCode.h
--- a/Engine/Source/Runtime/RigVM/Public/RigVMCore/RigVMByteCode.h
+++ b/Engine/Source/Runtime/RigVM/Public/RigVMCore/RigVMByteCode.h
@@ -817,6 +817,8 @@ struct RIGVM_API FRigVMByteCodeEntry
 
 	UPROPERTY()
 	int32 InstructionIndex;
+
+	FString GetSanitizedName() const;
 };
 
 /**
diff --git a/Engine/Source/Runtime/RigVM/Public/RigVMCore/RigVMNativized.h b/Engine/Source/Runtime/RigVM/Public/RigVMCore/RigVMNativized.h
--- a/Engine/Source/Runtime/RigVM/Public/RigVMCore/RigVMNativized.h
+++ b/Engine/Source/Runtime/RigVM/Public/RigVMCore/RigVMNativized.h
@@ -30,7 +30,8 @@ public:
 	virtual void ClearMemory() override { return; }
 	virtual const FRigVMInstructionArray& GetInstructions() override;
 	virtual bool ContainsEntry(const FName& InEntryName) const override { return GetEntryNames().Contains(InEntryName); }
-	virtual TArray<FName> GetEntryNames() const override { return TArray<FName>(); }
+	virtual int32 FindEntry(const FName& InEntryName) const override { return GetEntryNames().Find(InEntryName); }
+	virtual const TArray<FName>& GetEntryNames() const override { static const TArray<FName> EmptyEntries; return EmptyEntries; }
 	
 protected:
 
@@ -111,7 +112,10 @@ protected:
 
 	FORCEINLINE void BroadcastExecutionReachedExit()
 	{
-		ExecutionReachedExit().Broadcast(Context.PublicData.GetEventName());
+		if(EntriesBeingExecuted.Num() == 1)
+		{
+			ExecutionReachedExit().Broadcast(Context.PublicData.GetEventName());
+		}
 	}
 
 	template<typename T>
