RigVM: Invoke Entry functionality for the model

diff --git a/Engine/Source/Developer/RigVMDeveloper/Private/RigVMCompiler/RigVMAST.cpp b/Engine/Source/Developer/RigVMDeveloper/Private/RigVMCompiler/RigVMAST.cpp
--- a/Engine/Source/Developer/RigVMDeveloper/Private/RigVMCompiler/RigVMAST.cpp
+++ b/Engine/Source/Developer/RigVMDeveloper/Private/RigVMCompiler/RigVMAST.cpp
@@ -14,6 +14,7 @@
 #include "RigVMModel/Nodes/RigVMFunctionReturnNode.h"
 #include "RigVMModel/Nodes/RigVMFunctionEntryNode.h"
 #include "RigVMModel/Nodes/RigVMArrayNode.h"
+#include "RigVMModel/Nodes/RigVMInvokeEntryNode.h"
 #include "RigVMModel/RigVMGraph.h"
 #include "RigVMModel/RigVMController.h"
 #include "RigVMCore/RigVMExecuteContext.h"
@@ -976,6 +977,10 @@ FRigVMExprAST* FRigVMParserAST::CreateExpressionForNode(const FRigVMASTProxy& In
 		{
 			NodeExpr = MakeExpr<FRigVMArrayExprAST>(InNodeProxy);
 		}
+		else if (InNodeProxy.IsA<URigVMInvokeEntryNode>())
+		{
+			NodeExpr = MakeExpr<FRigVMInvokeEntryExprAST>(InNodeProxy);
+		}
 		else
 		{
 			NodeExpr = MakeExpr<FRigVMCallExternExprAST>(InNodeProxy);
diff --git a/Engine/Source/Developer/RigVMDeveloper/Private/RigVMCompiler/RigVMCompiler.cpp b/Engine/Source/Developer/RigVMDeveloper/Private/RigVMCompiler/RigVMCompiler.cpp
--- a/Engine/Source/Developer/RigVMDeveloper/Private/RigVMCompiler/RigVMCompiler.cpp
+++ b/Engine/Source/Developer/RigVMDeveloper/Private/RigVMCompiler/RigVMCompiler.cpp
@@ -681,6 +681,11 @@ void URigVMCompiler::TraverseExpression(const FRigVMExprAST* InExpr, FRigVMCompi
 			TraverseArray(InExpr->To<FRigVMArrayExprAST>(), WorkData);
 			break;
 		}
+		case FRigVMExprAST::EType::InvokeEntry:
+		{
+			TraverseInvokeEntry(InExpr->To<FRigVMInvokeEntryExprAST>(), WorkData);
+			break;
+		}
 		default:
 		{
 			ensure(false);
@@ -723,7 +728,6 @@ void URigVMCompiler::TraverseEntry(const FRigVMEntryExprAST* InExpr, FRigVMCompi
 	}
 	else
 	{
-		// todo: define the entry in the VM
 		TArray<FRigVMOperand> Operands;
 		for (FRigVMExprAST* ChildExpr : *InExpr)
 		{
@@ -1735,9 +1739,27 @@ void URigVMCompiler::TraverseArray(const FRigVMArrayExprAST* InExpr, FRigVMCompi
 	}
 }
 
+void URigVMCompiler::TraverseInvokeEntry(const FRigVMInvokeEntryExprAST* InExpr, FRigVMCompilerWorkData& WorkData)
+{
+	URigVMInvokeEntryNode* InvokeEntryNode = Cast<URigVMInvokeEntryNode>(InExpr->GetNode());
+	if(!ValidateNode(InvokeEntryNode))
+	{
+		return;
+	}
+
+	if (WorkData.bSetupMemory)
+	{
+		return;
+	}
+	else
+	{
+		WorkData.VM->GetByteCode().AddInvokeEntryOp(InvokeEntryNode->GetEntryName());
+	}
+}
+
 void URigVMCompiler::AddCopyOperator(const FRigVMCopyOp& InOp, const FRigVMAssignExprAST* InAssignExpr,
-	const FRigVMVarExprAST* InSourceExpr, const FRigVMVarExprAST* InTargetExpr,  FRigVMCompilerWorkData& WorkData,
-	bool bDelayCopyOperations)
+                                     const FRigVMVarExprAST* InSourceExpr, const FRigVMVarExprAST* InTargetExpr,  FRigVMCompilerWorkData& WorkData,
+                                     bool bDelayCopyOperations)
 {
 	if(bDelayCopyOperations)
 	{
diff --git a/Engine/Source/Developer/RigVMDeveloper/Private/RigVMModel/Nodes/RigVMInvokeEntryNode.cpp b/Engine/Source/Developer/RigVMDeveloper/Private/RigVMModel/Nodes/RigVMInvokeEntryNode.cpp
new file mode 100644
--- /dev/null
+++ b/Engine/Source/Developer/RigVMDeveloper/Private/RigVMModel/Nodes/RigVMInvokeEntryNode.cpp
@@ -0,0 +1,31 @@
+// Copyright Epic Games, Inc. All Rights Reserved.
+
+#include "RigVMModel/Nodes/RigVMInvokeEntryNode.h"
+
+#include "RigVMModel/RigVMGraph.h"
+
+const FString URigVMInvokeEntryNode::EntryName = TEXT("Entry");
+
+URigVMInvokeEntryNode::URigVMInvokeEntryNode()
+{
+}
+
+FString URigVMInvokeEntryNode::GetNodeTitle() const
+{
+	return FString::Printf(TEXT("Invoke %s"), *GetEntryName().ToString());
+}
+
+FName URigVMInvokeEntryNode::GetEntryName() const
+{
+	URigVMPin* EntryPin = GetEntryNamePin();
+	if (EntryPin == nullptr)
+	{
+		return NAME_None;
+	}
+	return *EntryPin->GetDefaultValue();
+}
+
+URigVMPin* URigVMInvokeEntryNode::GetEntryNamePin() const
+{
+	return FindPin(EntryName);
+}
diff --git a/Engine/Source/Developer/RigVMDeveloper/Private/RigVMModel/RigVMController.cpp b/Engine/Source/Developer/RigVMDeveloper/Private/RigVMModel/RigVMController.cpp
--- a/Engine/Source/Developer/RigVMDeveloper/Private/RigVMModel/RigVMController.cpp
+++ b/Engine/Source/Developer/RigVMDeveloper/Private/RigVMModel/RigVMController.cpp
@@ -654,6 +654,15 @@ TArray<FString> URigVMController::GetAddNodePythonCommands(URigVMNode* Node) con
 					*ContainedGraphName));
 		}
 	}
+	else if (const URigVMInvokeEntryNode* InvokeEntryNode = Cast<URigVMInvokeEntryNode>(Node))
+	{
+		// add_invoke_entry_node(entry_name, position=[0.0, 0.0], node_name='', undo=True)
+		Commands.Add(FString::Printf(TEXT("blueprint.get_controller_by_name('%s').add_invoke_entry_node('%s', %s, '%s')"),
+						*GraphName,
+						*InvokeEntryNode->GetEntryName().ToString(),
+						*RigVMPythonUtils::Vector2DToPythonString(EnumNode->GetPosition()),
+						*NodeName));
+	}
 	else if (Node->IsA<URigVMFunctionEntryNode>() || Node->IsA<URigVMFunctionReturnNode>())
 	{
 		
@@ -12561,6 +12570,80 @@ URigVMArrayNode* URigVMController::AddArrayNodeFromObjectPath(ERigVMOpCode InOpC
 	return AddArrayNode(InOpCode, InCPPType, CPPTypeObject, InPosition, InNodeName, bSetupUndoRedo, bPrintPythonCommand);
 }
 
+URigVMInvokeEntryNode* URigVMController::AddInvokeEntryNode(const FName& InEntryName, const FVector2D& InPosition,
+	const FString& InNodeName, bool bSetupUndoRedo, bool bPrintPythonCommand)
+{
+	if (!IsValidGraph())
+	{
+		return nullptr;
+	}
+
+	if (!bIsTransacting && !IsGraphEditable())
+	{
+		return nullptr;
+	}
+
+	URigVMGraph* Graph = GetGraph();
+	check(Graph);
+
+	if (Graph->IsA<URigVMFunctionLibrary>())
+	{
+		ReportError(TEXT("Cannot add invoke entry nodes to function library graphs."));
+		return nullptr;
+	}
+
+	FString Name = GetValidNodeName(InNodeName.IsEmpty() ? FString(TEXT("InvokeEntryNode")) : InNodeName);
+	URigVMInvokeEntryNode* Node = NewObject<URigVMInvokeEntryNode>(Graph, *Name);
+	Node->Position = InPosition;
+
+	UScriptStruct* ExecuteContextStruct = Graph->GetExecuteContextStruct();
+	URigVMPin* ExecutePin = NewObject<URigVMPin>(Node, FRigVMStruct::ExecuteContextName);
+	ExecutePin->CPPType = ExecuteContextStruct->GetStructCPPName();
+	ExecutePin->CPPTypeObject = ExecuteContextStruct;
+	ExecutePin->CPPTypeObjectPath = *ExecutePin->CPPTypeObject->GetPathName();
+	ExecutePin->Direction = ERigVMPinDirection::IO;
+	AddNodePin(Node, ExecutePin);
+
+	URigVMPin* EntryNamePin = NewObject<URigVMPin>(Node, *URigVMInvokeEntryNode::EntryName);
+	EntryNamePin->CPPType = RigVMTypeUtils::FNameType;
+	EntryNamePin->Direction = ERigVMPinDirection::Input;
+	EntryNamePin->bIsConstant = true;
+	EntryNamePin->DefaultValue = InEntryName.ToString();
+	EntryNamePin->CustomWidgetName = TEXT("EntryName");
+	AddNodePin(Node, EntryNamePin);
+
+	Graph->Nodes.Add(Node);
+
+	if (!bSuspendNotifications)
+	{
+		Graph->MarkPackageDirty();
+	}
+
+	FRigVMControllerCompileBracketScope CompileScope(this);
+
+	Notify(ERigVMGraphNotifType::NodeAdded, Node);
+	Notify(ERigVMGraphNotifType::VariableAdded, Node);
+
+	if (bSetupUndoRedo)
+	{
+		FRigVMAddInvokeEntryNodeAction Action(Node);
+		Action.Title = FString::Printf(TEXT("Add Invoke %s Entry"), *InEntryName.ToString());
+		ActionStack->AddAction(Action);
+	}
+
+	if (bPrintPythonCommand)
+	{
+		TArray<FString> Commands = GetAddNodePythonCommands(Node);
+		for (const FString& Command : Commands)
+		{
+			RigVMPythonUtils::Print(GetGraphOuterName(), 
+								FString::Printf(TEXT("%s"), *Command));
+		}
+	}
+
+	return Node;
+}
+
 void URigVMController::ForEveryPinRecursively(URigVMPin* InPin, TFunction<void(URigVMPin*)> OnEachPinFunction)
 {
 	OnEachPinFunction(InPin);
diff --git a/Engine/Source/Developer/RigVMDeveloper/Private/RigVMModel/RigVMControllerActions.cpp b/Engine/Source/Developer/RigVMDeveloper/Private/RigVMModel/RigVMControllerActions.cpp
--- a/Engine/Source/Developer/RigVMDeveloper/Private/RigVMModel/RigVMControllerActions.cpp
+++ b/Engine/Source/Developer/RigVMDeveloper/Private/RigVMModel/RigVMControllerActions.cpp
@@ -736,7 +736,7 @@ FRigVMRemoveNodeAction::FRigVMRemoveNodeAction(URigVMNode* InNode, URigVMControl
 	else if (URigVMVariableNode* VariableNode = Cast<URigVMVariableNode>(InNode))
 	{
 		InverseAction.AddAction(FRigVMAddVariableNodeAction(VariableNode), InController);
-		URigVMPin* ValuePin = VariableNode->FindPin(TEXT("Value"));
+		URigVMPin* ValuePin = VariableNode->GetValuePin();
 		InverseAction.AddAction(FRigVMSetPinDefaultValueAction(ValuePin, ValuePin->GetDefaultValue()), InController);
 	}
 	else if (URigVMCommentNode* CommentNode = Cast<URigVMCommentNode>(InNode))
@@ -783,6 +783,12 @@ FRigVMRemoveNodeAction::FRigVMRemoveNodeAction(URigVMNode* InNode, URigVMControl
 	{
 		InverseAction.AddAction(FRigVMImportNodeFromTextAction(LibraryNode, InController), InController);
 	}
+	else if (URigVMInvokeEntryNode* InvokeEntryNode = Cast<URigVMInvokeEntryNode>(InNode))
+	{
+		InverseAction.AddAction(FRigVMAddInvokeEntryNodeAction(InvokeEntryNode), InController);
+		URigVMPin* EntryNamePin = InvokeEntryNode->GetEntryNamePin();
+		InverseAction.AddAction(FRigVMSetPinDefaultValueAction(EntryNamePin, EntryNamePin->GetDefaultValue()), InController);
+	}
 	else if (InNode->IsA<URigVMFunctionEntryNode>() || InNode->IsA<URigVMFunctionReturnNode>())
 	{
 		// Do nothing
@@ -2222,3 +2228,37 @@ bool FRigVMPromoteNodeAction::Redo(URigVMController* InController)
 	const FName FunctionRefNodeName = InController->PromoteCollapseNodeToFunctionReferenceNode(*LibraryNodePath, false, false);
 	return FunctionRefNodeName.ToString() == LibraryNodePath;
 }
+
+FRigVMAddInvokeEntryNodeAction::FRigVMAddInvokeEntryNodeAction()
+	: EntryName(NAME_None)
+	, Position(FVector2D::ZeroVector)
+	, NodePath()
+{
+}
+
+FRigVMAddInvokeEntryNodeAction::FRigVMAddInvokeEntryNodeAction(URigVMInvokeEntryNode* InNode)
+	: EntryName(InNode->GetEntryName())
+	, Position(InNode->GetPosition())
+	, NodePath(InNode->GetNodePath())
+{
+}
+
+bool FRigVMAddInvokeEntryNodeAction::Undo(URigVMController* InController)
+{
+	if (!FRigVMBaseAction::Undo(InController))
+	{
+		return false;
+	}
+	return InController->RemoveNodeByName(*NodePath, false);
+}
+
+bool FRigVMAddInvokeEntryNodeAction::Redo(URigVMController* InController)
+{
+#if WITH_EDITOR
+	if (URigVMInvokeEntryNode* Node = InController->AddInvokeEntryNode(EntryName, Position, NodePath, false))
+	{
+		return FRigVMBaseAction::Redo(InController);
+	}
+#endif
+	return false;
+}
diff --git a/Engine/Source/Developer/RigVMDeveloper/Private/RigVMModel/RigVMPin.cpp b/Engine/Source/Developer/RigVMDeveloper/Private/RigVMModel/RigVMPin.cpp
--- a/Engine/Source/Developer/RigVMDeveloper/Private/RigVMModel/RigVMPin.cpp
+++ b/Engine/Source/Developer/RigVMDeveloper/Private/RigVMModel/RigVMPin.cpp
@@ -20,6 +20,7 @@
 #include "RigVMModel/Nodes/RigVMFunctionReferenceNode.h"
 #include "RigVMModel/Nodes/RigVMFunctionEntryNode.h"
 #include "RigVMModel/Nodes/RigVMFunctionReturnNode.h"
+#include "RigVMModel/Nodes/RigVMInvokeEntryNode.h"
 
 #if WITH_EDITOR
 #include "UObject/CoreRedirects.h"
@@ -1218,6 +1219,13 @@ bool URigVMPin::ShowInDetailsPanelOnly() const
 			}
 		}
 	}
+	else if(URigVMInvokeEntryNode* InvokeEntryNode = Cast<URigVMInvokeEntryNode>(GetNode()))
+	{
+		if(GetName() == URigVMInvokeEntryNode::EntryName)
+		{
+			return true;
+		}
+	}
 #endif
 	return false;
 }
diff --git a/Engine/Source/Developer/RigVMDeveloper/Public/RigVMCompiler/RigVMAST.h b/Engine/Source/Developer/RigVMDeveloper/Public/RigVMCompiler/RigVMAST.h
--- a/Engine/Source/Developer/RigVMDeveloper/Public/RigVMCompiler/RigVMAST.h
+++ b/Engine/Source/Developer/RigVMDeveloper/Public/RigVMCompiler/RigVMAST.h
@@ -14,6 +14,7 @@
 class FRigVMParserAST;
 class FRigVMBlockExprAST;
 class FRigVMEntryExprAST;
+class FRigVMInvokeEntryExprAST;
 class FRigVMCallExternExprAST;
 class FRigVMNoOpExprAST;
 class FRigVMVarExprAST;
@@ -79,6 +80,7 @@ public:
 	{
 		Block,
 		Entry,
+		InvokeEntry,
 		CallExtern,
 		NoOp,
 		Var,
@@ -313,6 +315,17 @@ FORCEINLINE const FRigVMEntryExprAST* FRigVMExprAST::To() const
 	return (const FRigVMEntryExprAST*)this;
 }
 
+// specialized cast for type checking
+// for a InvokeEntry / FRigVMInvokeEntryExprAST expression
+// will raise if types are not compatible
+// @return this expression cast to FRigVMInvokeEntryExprAST
+template<>
+FORCEINLINE const FRigVMInvokeEntryExprAST* FRigVMExprAST::To() const
+{
+	ensure(IsA(EType::InvokeEntry));
+	return (const FRigVMInvokeEntryExprAST*)this;
+}
+
 // specialized cast for type checking
 // for a CallExtern / FRigVMCallExternExprAST expression
 // will raise if types are not compatible
@@ -591,6 +604,46 @@ private:
 	friend class FRigVMParserAST;
 };
 
+/*
+ * An abstract syntax tree entry expression represents an invocation
+ * of an entry point.
+ */
+class RIGVMDEVELOPER_API FRigVMInvokeEntryExprAST : public FRigVMNodeExprAST
+{
+public:
+
+	// virtual destructor
+	virtual ~FRigVMInvokeEntryExprAST() {}
+
+	// disable copy constructor
+	FRigVMInvokeEntryExprAST(const FRigVMInvokeEntryExprAST&) = delete;
+
+	// overload of the type checking mechanism
+	virtual bool IsA(EType InType) const override
+	{
+		if(FRigVMNodeExprAST::IsA(InType))
+		{
+			return true;
+		}
+		return InType == EType::InvokeEntry;
+	};
+
+	// returns the name of the entry / event
+	// @return the name of the entry / event
+	FName GetEventName() const;
+
+protected:
+
+	// default constructor (protected so that only parser can access it)
+	FRigVMInvokeEntryExprAST(const FRigVMASTProxy& InNodeProxy)
+		: FRigVMNodeExprAST(EType::InvokeEntry, InNodeProxy)
+	{}
+
+private:
+
+	friend class FRigVMParserAST;
+};
+
 /*
  * An abstract syntax tree call extern expression represents the invocation
  * of an extern function.
diff --git a/Engine/Source/Developer/RigVMDeveloper/Public/RigVMCompiler/RigVMCompiler.h b/Engine/Source/Developer/RigVMDeveloper/Public/RigVMCompiler/RigVMCompiler.h
--- a/Engine/Source/Developer/RigVMDeveloper/Public/RigVMCompiler/RigVMCompiler.h
+++ b/Engine/Source/Developer/RigVMDeveloper/Public/RigVMCompiler/RigVMCompiler.h
@@ -179,6 +179,7 @@ private:
 	void TraverseIf(const FRigVMIfExprAST* InExpr, FRigVMCompilerWorkData& WorkData);
 	void TraverseSelect(const FRigVMSelectExprAST* InExpr, FRigVMCompilerWorkData& WorkData);
 	void TraverseArray(const FRigVMArrayExprAST* InExpr, FRigVMCompilerWorkData& WorkData);
+	void TraverseInvokeEntry(const FRigVMInvokeEntryExprAST* InExpr, FRigVMCompilerWorkData& WorkData);
 
 	void AddCopyOperator(
 		const FRigVMCopyOp& InOp,
diff --git a/Engine/Source/Developer/RigVMDeveloper/Public/RigVMModel/Nodes/RigVMInvokeEntryNode.h b/Engine/Source/Developer/RigVMDeveloper/Public/RigVMModel/Nodes/RigVMInvokeEntryNode.h
new file mode 100644
--- /dev/null
+++ b/Engine/Source/Developer/RigVMDeveloper/Public/RigVMModel/Nodes/RigVMInvokeEntryNode.h
@@ -0,0 +1,47 @@
+// Copyright Epic Games, Inc. All Rights Reserved.
+
+#pragma once
+
+#include "RigVMModel/RigVMNode.h"
+#include "RigVMInvokeEntryNode.generated.h"
+
+/**
+ * The Invoke Entry Node is used to invoke another entry from the 
+ * the currently running entry.
+ */
+UCLASS(BlueprintType)
+class RIGVMDEVELOPER_API URigVMInvokeEntryNode : public URigVMNode
+{
+	GENERATED_BODY()
+
+public:
+
+	// Default constructor
+	URigVMInvokeEntryNode();
+
+	// Override of node title
+	virtual FString GetNodeTitle() const;
+
+	// Returns the name of the entry to run
+	UFUNCTION(BlueprintCallable, Category = RigVMInvokeEntryNode)
+	FName GetEntryName() const;
+
+	// Override of node title
+	virtual FLinearColor GetNodeColor() const override { return FLinearColor::Blue; }
+	virtual bool IsDefinedAsVarying() const override { return true; }
+
+private:
+
+	static const FString EntryName;
+
+	URigVMPin* GetEntryNamePin() const;
+	
+
+	friend class URigVMController;
+	friend class UControlRigBlueprint;
+	friend struct FRigVMRemoveNodeAction;
+	friend class URigVMCompiler;
+	friend class FRigVMParserAST;
+	friend class URigVMPin;
+};
+
diff --git a/Engine/Source/Developer/RigVMDeveloper/Public/RigVMModel/RigVMController.h b/Engine/Source/Developer/RigVMDeveloper/Public/RigVMModel/RigVMController.h
--- a/Engine/Source/Developer/RigVMDeveloper/Public/RigVMModel/RigVMController.h
+++ b/Engine/Source/Developer/RigVMDeveloper/Public/RigVMModel/RigVMController.h
@@ -17,6 +17,7 @@
 #include "RigVMModel/Nodes/RigVMCollapseNode.h"
 #include "RigVMModel/Nodes/RigVMFunctionReferenceNode.h"
 #include "RigVMModel/Nodes/RigVMArrayNode.h"
+#include "RigVMModel/Nodes/RigVMInvokeEntryNode.h"
 #include "RigVMModel/RigVMBuildData.h"
 #include "RigVMCore/RigVMUserWorkflow.h"
 #include "UObject/Interface.h"
@@ -415,6 +416,11 @@ public:
 	UFUNCTION(BlueprintCallable, Category = RigVMController)
 	URigVMArrayNode* AddArrayNodeFromObjectPath(ERigVMOpCode InOpCode, const FString& InCPPType, const FString& InCPPTypeObjectPath, const FVector2D& InPosition = FVector2D::ZeroVector, const FString& InNodeName = TEXT(""), bool bSetupUndoRedo = true, bool bPrintPythonCommand = false);
 
+	// Adds an entry invocation node
+	// This causes a NodeAdded modified event.
+	UFUNCTION(BlueprintCallable, Category = RigVMController)
+	URigVMInvokeEntryNode* AddInvokeEntryNode(const FName& InEntryName, const FVector2D& InPosition = FVector2D::ZeroVector, const FString& InNodeName = TEXT(""), bool bSetupUndoRedo = true, bool bPrintPythonCommand = false);
+
 	// Un-does the last action on the stack.
 	// Note: This should really only be used for unit tests,
 	// use the GEditor's main Undo method instead.
diff --git a/Engine/Source/Developer/RigVMDeveloper/Public/RigVMModel/RigVMControllerActions.h b/Engine/Source/Developer/RigVMDeveloper/Public/RigVMModel/RigVMControllerActions.h
--- a/Engine/Source/Developer/RigVMDeveloper/Public/RigVMModel/RigVMControllerActions.h
+++ b/Engine/Source/Developer/RigVMDeveloper/Public/RigVMModel/RigVMControllerActions.h
@@ -1806,4 +1806,31 @@ public:
 
 	UPROPERTY()
 	bool bFromFunctionToCollapseNode;
-};
\ No newline at end of file
+};
+
+/**
+ * An action adding an invoke entry node to the graph.
+ */
+USTRUCT()
+struct FRigVMAddInvokeEntryNodeAction : public FRigVMBaseAction
+{
+	GENERATED_BODY()
+
+public:
+
+	FRigVMAddInvokeEntryNodeAction();
+	FRigVMAddInvokeEntryNodeAction(URigVMInvokeEntryNode* InNode);
+	virtual ~FRigVMAddInvokeEntryNodeAction() {};
+	virtual UScriptStruct* GetScriptStruct() const override { return FRigVMAddInvokeEntryNodeAction::StaticStruct(); }
+	virtual bool Undo(URigVMController* InController) override;
+	virtual bool Redo(URigVMController* InController) override;
+
+	UPROPERTY()
+	FName EntryName;
+	
+	UPROPERTY()
+	FVector2D Position;
+
+	UPROPERTY()
+	FString NodePath;
+};
