fixed discrepancy between struct subpins's value within the full struct default value and the subpin's real default value

Here is a simple example we could encounter (using FRigUnit_AimItem_Target)

the struct default's value looked like this "(Weight=1.000000,Axis=(X=1.000000,Y=0.000000,Z=0.000000),Target=(X=1.000000,Y=0.000000,Z=0.000000),Kind=Location,Space="")"
setting the Space default value to "" whereas the Space default value is an empty string.
Using "" as a default value would fail later in URigVMController::SetPinDefaultValue

diff --git a/Engine/Source/Developer/RigVMDeveloper/Private/RigVMModel/RigVMPin.cpp b/Engine/Source/Developer/RigVMDeveloper/Private/RigVMModel/RigVMPin.cpp
--- a/Engine/Source/Developer/RigVMDeveloper/Private/RigVMModel/RigVMPin.cpp
+++ b/Engine/Source/Developer/RigVMDeveloper/Private/RigVMModel/RigVMPin.cpp
@@ -704,10 +704,10 @@ FString URigVMPin::GetDefaultValue(const URigVMPin::FPinOverride& InOverride) co
 		if (SubPins.Num() > 0)
 		{
 			TArray<FString> MemberDefaultValues;
-			for (URigVMPin* SubPin : SubPins)
+			for (const URigVMPin* SubPin : SubPins)
 			{
 				FString MemberDefaultValue = SubPin->GetDefaultValue(InOverride);
-				if (SubPin->IsStringType())
+				if (SubPin->IsStringType() && !MemberDefaultValue.IsEmpty())
 				{
 					MemberDefaultValue = TEXT("\"") + MemberDefaultValue + TEXT("\"");
 				}
