Control Rig: Fix some templates not resolving

diff --git a/Engine/Source/Developer/RigVMDeveloper/Private/RigVMModel/Nodes/RigVMTemplateNode.cpp b/Engine/Source/Developer/RigVMDeveloper/Private/RigVMModel/Nodes/RigVMTemplateNode.cpp
--- a/Engine/Source/Developer/RigVMDeveloper/Private/RigVMModel/Nodes/RigVMTemplateNode.cpp
+++ b/Engine/Source/Developer/RigVMDeveloper/Private/RigVMModel/Nodes/RigVMTemplateNode.cpp
@@ -549,7 +549,12 @@ TArray<int32> URigVMTemplateNode::GetNewFilteredPermutations(URigVMPin* InPin, U
 		}
 	}
 
-	TArray<int32> PermutationsToTry = PreferredPermutationTypes.IsEmpty()? FilteredPermutations : TArray<int32>(FindPermutationsForTypes(PreferredPermutationTypes));
+	TArray<int32> PermutationsToTry = FilteredPermutations;
+
+	// Reduce permutations to the ones respecting the preferred types
+	TArray<int32> PreferredPermutations = FindPermutationsForTypes(PreferredPermutationTypes, true);
+	PermutationsToTry = PermutationsToTry.FilterByPredicate([&](const int32& OtherPermutation) { return PreferredPermutations.Contains(OtherPermutation); });
+	
 	bool bLinkedIsTemplate = false;
 	if (URigVMTemplateNode* LinkedTemplate = Cast<URigVMTemplateNode>(LinkedPin->GetNode()))
 	{
@@ -612,7 +617,11 @@ TArray<int32> URigVMTemplateNode::GetNewFilteredPermutations(URigVMPin* InPin, c
 		}
 	}
 
-	TArray<int32> PermutationsToTry = PreferredPermutationTypes.IsEmpty()? FilteredPermutations : TArray<int32>(FindPermutationsForTypes(PreferredPermutationTypes));
+	TArray<int32> PermutationsToTry = FilteredPermutations;
+
+	// Reduce permutations to the ones respecting the preferred types
+	TArray<int32> PreferredPermutations = FindPermutationsForTypes(PreferredPermutationTypes, true);
+	PermutationsToTry = PermutationsToTry.FilterByPredicate([&](const int32& OtherPermutation) { return PreferredPermutations.Contains(OtherPermutation); });
 	
 	if (const FRigVMTemplateArgument* Argument = GetTemplate()->FindArgument(RootPin->GetFName()))
 	{
diff --git a/Engine/Source/Developer/RigVMDeveloper/Private/RigVMModel/RigVMController.cpp b/Engine/Source/Developer/RigVMDeveloper/Private/RigVMModel/RigVMController.cpp
--- a/Engine/Source/Developer/RigVMDeveloper/Private/RigVMModel/RigVMController.cpp
+++ b/Engine/Source/Developer/RigVMDeveloper/Private/RigVMModel/RigVMController.cpp
@@ -3589,9 +3589,10 @@ TArray<FName> URigVMController::ImportNodesFromText(const FString& InText, bool
 						TGuardValue<TArray<FString>> GuardPreferredPermutation(TemplateNode->PreferredPermutationTypes, {});
 						TemplateNode->InitializeFilteredPermutations();
 					}
+					TemplateNode->InitializeFilteredPermutationsFromTypes();
 					if (!TemplateNode->PreferredPermutationTypes.IsEmpty())
 					{
-						TemplateNode->FilteredPermutations = TemplateNode->FindPermutationsForTypes(TemplateNode->PreferredPermutationTypes);
+						TemplateNode->FilteredPermutations = TemplateNode->FindPermutationsForTypes(TemplateNode->PreferredPermutationTypes, true);
 					}
 					UpdateTemplateNodePinTypes(TemplateNode, bSetupUndoRedo);
 				}
@@ -3700,6 +3701,29 @@ TArray<FName> URigVMController::ImportNodesFromText(const FString& InText, bool
 						}
 						else
 						{
+							if (URigVMTemplateNode* FirstTemplateNode = Cast<URigVMTemplateNode>(TargetPin->GetNode()))
+							{
+								if (!FirstTemplateNode->IsSingleton())
+								{
+									TArray<FRigVMTemplateArgument::FType> InputTypes = GetWildcardFilteredTypes(SourcePin);
+									if (InputTypes.Num() > 0)
+									{
+										PrepareTemplatePinForType(TargetPin, InputTypes, bSetupUndoRedo);										
+									}
+								}
+							}
+							if (URigVMTemplateNode* SecondTemplateNode = Cast<URigVMTemplateNode>(SourcePin->GetNode()))
+							{
+								if (!SecondTemplateNode->IsSingleton())
+								{
+									TArray<FRigVMTemplateArgument::FType> OutTypes = GetWildcardFilteredTypes(TargetPin);
+									if (OutTypes.Num() > 0)
+									{
+										PrepareTemplatePinForType(SourcePin, OutTypes, bSetupUndoRedo);										
+									}
+								}
+							}
+							
 							Graph->Links.Add(CreatedLink);
 							SourcePin->Links.Add(CreatedLink);
 							TargetPin->Links.Add(CreatedLink);
@@ -12063,6 +12087,8 @@ URigVMTemplateNode* URigVMController::AddTemplateNode(const FName& InNotation, c
 		}
 	}
 
+	UpdateTemplateNodePinTypes(Node, false);
+
 	Graph->Nodes.Add(Node);
 
 	Notify(ERigVMGraphNotifType::NodeAdded, Node);
