Control Rig: Fix crash when  there are links between subpins and resolving to a different type

diff --git a/Engine/Source/Developer/RigVMDeveloper/Private/RigVMModel/RigVMController.cpp b/Engine/Source/Developer/RigVMDeveloper/Private/RigVMModel/RigVMController.cpp
--- a/Engine/Source/Developer/RigVMDeveloper/Private/RigVMModel/RigVMController.cpp
+++ b/Engine/Source/Developer/RigVMDeveloper/Private/RigVMModel/RigVMController.cpp
@@ -2709,7 +2709,6 @@ URigVMInjectionInfo* URigVMController::AddInjectedNode(const FString& InPinPath,
 			InputPin->IsArray() != OutputPin->IsArray())
 		{
 			ReportErrorf(TEXT("Injected node %s is using incompatible input and output pins."), *InScriptStruct->GetName());
-			RemoveNode(UnitNode, false);
 			if (bSetupUndoRedo)
 			{
 				ActionStack->CancelAction(Action, this);
@@ -2721,7 +2720,6 @@ URigVMInjectionInfo* URigVMController::AddInjectedNode(const FString& InPinPath,
 			InputPin->IsArray() != Pin->IsArray())
 		{
 			ReportErrorf(TEXT("Injected node %s is using incompatible pin."), *InScriptStruct->GetName());
-			RemoveNode(UnitNode, false);
 			if (bSetupUndoRedo)
 			{
 				ActionStack->CancelAction(Action, this);
@@ -15228,6 +15226,9 @@ bool URigVMController::PrepareTemplatePinForType(URigVMPin* InPin, const TArray<
 
 				// Update our node with the requested type
 				UpdateFilteredPermutations(InPin, InTypes, true);
+				
+				// Update types in case some pins on this node disappear, and should no longer propagate
+				UpdateTemplateNodePinTypes(TemplateNode, true);
 
 				// Propagate the filtered types and collect break link actions 
 				bBrokenLinks = !PropagateTemplateFilteredTypes(TemplateNode, true);
@@ -16021,7 +16022,6 @@ bool URigVMController::ChangePinType(URigVMPin* InPin, const FString& InCPPType,
 			BreakAllLinksRecursive(InPin, true, false, true);
 			BreakAllLinksRecursive(InPin, false, false, true);
 		}
-		ActionStack->AddAction(FRigVMChangePinTypeAction(InPin, CPPType, CPPTypeObjectPath, bSetupOrphanPins, bBreakLinks, bRemoveSubPins));
 	}
 	
 	if(bSetupOrphanPins)
@@ -16051,12 +16051,17 @@ bool URigVMController::ChangePinType(URigVMPin* InPin, const FString& InCPPType,
 		TArray<URigVMPin*> Pins = InPin->SubPins;
 		for (URigVMPin* Pin : Pins)
 		{
-			RemovePin(Pin, false, true);
+			RemovePin(Pin, bSetupUndoRedo, true);
 		}
 		
 		InPin->SubPins.Reset();
 	}
 
+	if (bSetupUndoRedo)
+	{
+		ActionStack->AddAction(FRigVMChangePinTypeAction(InPin, CPPType, CPPTypeObjectPath, bSetupOrphanPins, bBreakLinks, bRemoveSubPins));
+	}
+	
 	// compute the number of remaining wildcard pins
 	auto WildCardPinCountPredicate = [](const URigVMPin* Pin) { return Pin->IsWildCard(); };
 	TArray<URigVMPin*> AllPins = InPin->GetNode()->GetAllPinsRecursively();
