Control Rig: Fix undo deletion of template node not working

diff --git a/Engine/Source/Developer/RigVMDeveloper/Private/RigVMModel/RigVMControllerActions.cpp b/Engine/Source/Developer/RigVMDeveloper/Private/RigVMModel/RigVMControllerActions.cpp
--- a/Engine/Source/Developer/RigVMDeveloper/Private/RigVMModel/RigVMControllerActions.cpp
+++ b/Engine/Source/Developer/RigVMDeveloper/Private/RigVMModel/RigVMControllerActions.cpp
@@ -720,20 +720,7 @@ FRigVMRemoveNodeAction::FRigVMRemoveNodeAction(URigVMNode* InNode, URigVMControl
 {
 	FRigVMInverseAction InverseAction;
 
-	if (URigVMUnitNode* UnitNode = Cast<URigVMUnitNode>(InNode))
-	{
-		InverseAction.AddAction(FRigVMAddUnitNodeAction(UnitNode), InController);
-		for (URigVMPin* Pin : UnitNode->GetPins())
-		{
-			if (Pin->GetDirection() == ERigVMPinDirection::Input ||
-				Pin->GetDirection() == ERigVMPinDirection::IO ||
-				Pin->GetDirection() == ERigVMPinDirection::Visible)
-			{
-				InverseAction.AddAction(FRigVMSetPinDefaultValueAction(Pin, Pin->GetDefaultValue()), InController);
-			}
-		}
-	}
-	else if (URigVMVariableNode* VariableNode = Cast<URigVMVariableNode>(InNode))
+	if (URigVMVariableNode* VariableNode = Cast<URigVMVariableNode>(InNode))
 	{
 		InverseAction.AddAction(FRigVMAddVariableNodeAction(VariableNode), InController);
 		URigVMPin* ValuePin = VariableNode->GetValuePin();
@@ -759,10 +746,6 @@ FRigVMRemoveNodeAction::FRigVMRemoveNodeAction(URigVMNode* InNode, URigVMControl
 	{
 		InverseAction.AddAction(FRigVMAddSelectNodeAction(SelectNode), InController);
 	}
-	else if (URigVMTemplateNode* TemplateNode = Cast<URigVMTemplateNode>(InNode))
-	{
-		InverseAction.AddAction(FRigVMAddTemplateNodeAction(TemplateNode), InController);
-	}
 	else if (URigVMEnumNode* EnumNode = Cast<URigVMEnumNode>(InNode))
 	{
 		InverseAction.AddAction(FRigVMAddEnumNodeAction(EnumNode), InController);
@@ -793,6 +776,19 @@ FRigVMRemoveNodeAction::FRigVMRemoveNodeAction(URigVMNode* InNode, URigVMControl
 	{
 		// Do nothing
 	}
+	else if (URigVMTemplateNode* TemplateNode = Cast<URigVMTemplateNode>(InNode))
+	{
+		InverseAction.AddAction(FRigVMAddTemplateNodeAction(TemplateNode), InController);
+		for (URigVMPin* Pin : TemplateNode->GetPins())
+		{
+			if (Pin->GetDirection() == ERigVMPinDirection::Input ||
+				Pin->GetDirection() == ERigVMPinDirection::IO ||
+				Pin->GetDirection() == ERigVMPinDirection::Visible)
+			{
+				InverseAction.AddAction(FRigVMSetPinDefaultValueAction(Pin, Pin->GetDefaultValue()), InController);
+			}
+		}
+	}
 	else
 	{
 		ensure(false);
