Control Rig: Fix unit test with undo/redo not passing

diff --git a/Engine/Source/Developer/RigVMDeveloper/Private/RigVMModel/RigVMControllerActions.cpp b/Engine/Source/Developer/RigVMDeveloper/Private/RigVMModel/RigVMControllerActions.cpp
--- a/Engine/Source/Developer/RigVMDeveloper/Private/RigVMModel/RigVMControllerActions.cpp
+++ b/Engine/Source/Developer/RigVMDeveloper/Private/RigVMModel/RigVMControllerActions.cpp
@@ -778,14 +778,47 @@ FRigVMRemoveNodeAction::FRigVMRemoveNodeAction(URigVMNode* InNode, URigVMControl
 	}
 	else if (URigVMTemplateNode* TemplateNode = Cast<URigVMTemplateNode>(InNode))
 	{
-		InverseAction.AddAction(FRigVMAddTemplateNodeAction(TemplateNode), InController);
+		if (TemplateNode->IsSingleton())
+		{
+			if (URigVMUnitNode* UnitNode = Cast<URigVMUnitNode>(InNode))
+			{
+				InverseAction.AddAction(FRigVMAddUnitNodeAction(UnitNode), InController);
+			}
+		}
+		else
+		{
+			InverseAction.AddAction(FRigVMAddTemplateNodeAction(TemplateNode), InController);
+			InverseAction.AddAction(FRigVMSetPreferredTemplatePermutationsAction(TemplateNode, TemplateNode->PreferredPermutationTypes), InController);
+			InverseAction.AddAction(FRigVMSetTemplateFilteredPermutationsAction(TemplateNode, nullptr, {}), InController);
+
+			for (URigVMPin* Pin : TemplateNode->GetPins())
+			{
+				if (!Pin->IsWildCard())
+				{
+					if (const FRigVMTemplate* Template = TemplateNode->GetTemplate())
+					{
+						if (const FRigVMTemplateArgument* Argument = Template->FindArgument(Pin->GetFName()))
+						{
+							if (!Argument->IsSingleton())
+							{
+								InverseAction.AddAction(FRigVMChangePinTypeAction(Pin, Pin->GetCPPType(), *Pin->GetCPPTypeObject()->GetPathName(), false, false, false), InController);
+							}
+						}
+					}			
+				}
+			}
+		}
+
 		for (URigVMPin* Pin : TemplateNode->GetPins())
 		{
-			if (Pin->GetDirection() == ERigVMPinDirection::Input ||
-				Pin->GetDirection() == ERigVMPinDirection::IO ||
-				Pin->GetDirection() == ERigVMPinDirection::Visible)
+			if (!Pin->IsWildCard())
 			{
-				InverseAction.AddAction(FRigVMSetPinDefaultValueAction(Pin, Pin->GetDefaultValue()), InController);
+				if (Pin->GetDirection() == ERigVMPinDirection::Input ||
+				    Pin->GetDirection() == ERigVMPinDirection::IO ||
+				    Pin->GetDirection() == ERigVMPinDirection::Visible)
+				{
+					InverseAction.AddAction(FRigVMSetPinDefaultValueAction(Pin, Pin->GetDefaultValue()), InController);
+				}
 			}
 		}
 	}
diff --git a/Engine/Source/Developer/RigVMDeveloper/Public/RigVMModel/Nodes/RigVMTemplateNode.h b/Engine/Source/Developer/RigVMDeveloper/Public/RigVMModel/Nodes/RigVMTemplateNode.h
--- a/Engine/Source/Developer/RigVMDeveloper/Public/RigVMModel/Nodes/RigVMTemplateNode.h
+++ b/Engine/Source/Developer/RigVMDeveloper/Public/RigVMModel/Nodes/RigVMTemplateNode.h
@@ -123,5 +123,6 @@ protected:
 	friend class UControlRigBlueprint;
 	friend struct FRigVMSetTemplateFilteredPermutationsAction;
 	friend struct FRigVMSetPreferredTemplatePermutationsAction;
+	friend struct FRigVMRemoveNodeAction;
 };
 
