Control Rig: Fix wildcard pin errors when opening rig

diff --git a/Engine/Source/Developer/RigVMDeveloper/Private/RigVMModel/RigVMController.cpp b/Engine/Source/Developer/RigVMDeveloper/Private/RigVMModel/RigVMController.cpp
--- a/Engine/Source/Developer/RigVMDeveloper/Private/RigVMModel/RigVMController.cpp
+++ b/Engine/Source/Developer/RigVMDeveloper/Private/RigVMModel/RigVMController.cpp
@@ -15874,11 +15874,6 @@ void URigVMController::InitializeAllTemplateFiltersInGraph(bool bSetupUndoRedo,
 		{
 			if (!TemplateNode->IsSingleton())
 			{
-				if (!TemplateNode->PreferredPermutationTypes.IsEmpty())
-				{
-					continue;
-				}
-				
 				TArray<int32> OldPermutations = TemplateNode->FilteredPermutations;
 				TemplateNode->InitializeFilteredPermutations();
 				if (bChangePinTypes)
