RigVM: Change order of arguments / pins to be based on super -> child struct

diff --git a/Engine/Source/Developer/RigVMDeveloper/Private/RigVMCompiler/RigVMAST.cpp b/Engine/Source/Developer/RigVMDeveloper/Private/RigVMCompiler/RigVMAST.cpp
--- a/Engine/Source/Developer/RigVMDeveloper/Private/RigVMCompiler/RigVMAST.cpp
+++ b/Engine/Source/Developer/RigVMDeveloper/Private/RigVMCompiler/RigVMAST.cpp
@@ -260,6 +260,23 @@ void FRigVMExprAST::RemoveParent(FRigVMExprAST* InParent)
 	
 	if (Parents.Remove(InParent) > 0)
 	{
+		const int32 ExistingIndex = InParent->Children.Find(this);
+		if(ExistingIndex != INDEX_NONE)
+		{
+			FName NameToRemove(NAME_None);
+			for(TPair<FName, int32>& Pair: InParent->PinNameToChildIndex)
+			{
+				if(Pair.Value > ExistingIndex)
+				{
+					Pair.Value--;
+				}
+				else if(Pair.Value == ExistingIndex)
+				{
+					NameToRemove = Pair.Key;
+				}
+			}
+			InParent->PinNameToChildIndex.Remove(NameToRemove);
+		}
 		InParent->Children.Remove(this);
 	}
 }
@@ -426,29 +443,37 @@ bool FRigVMNodeExprAST::IsConstant() const
 	return FRigVMExprAST::IsConstant();
 }
 
-const FRigVMVarExprAST* FRigVMNodeExprAST::FindVarWithPinName(const FName& InPinName) const
+const FRigVMExprAST* FRigVMNodeExprAST::FindExprWithPinName(const FName& InPinName) const
 {
+	if(const int32* PinIndex = PinNameToChildIndex.Find(InPinName))
+	{
+		return ChildAt(*PinIndex);
+	}
+	
 	if(URigVMNode* CurrentNode = GetNode())
 	{
-		for(URigVMPin* Pin : CurrentNode->GetPins())
+		if(const URigVMPin* Pin = CurrentNode->FindPin(InPinName.ToString()))
 		{
-			if(Pin->GetFName() == InPinName)
+			const int32 PinIndex = Pin->GetPinIndex();
+			if(PinIndex <= NumChildren())
 			{
-				const int32 PinIndex = Pin->GetPinIndex();
-				if(PinIndex <= NumChildren())
-				{
-					const FRigVMExprAST* Child = ChildAt(PinIndex);
-					if (Child->IsA(FRigVMExprAST::Var))
-					{
-						return Child->To<FRigVMVarExprAST>();
-					}
-				}
+				return ChildAt(PinIndex);
 			}
 		}
 	}
 	return nullptr;
 }
 
+const FRigVMVarExprAST* FRigVMNodeExprAST::FindVarWithPinName(const FName& InPinName) const
+{
+	const FRigVMExprAST* Child = FindExprWithPinName(InPinName);
+	if (Child->IsA(FRigVMExprAST::Var))
+	{
+		return Child->To<FRigVMVarExprAST>();
+	}
+	return nullptr;
+}
+
 FRigVMNodeExprAST::FRigVMNodeExprAST(EType InType, const FRigVMASTProxy& InNodeProxy)
 	: FRigVMBlockExprAST(InType)
 	, Proxy(InNodeProxy)
@@ -1014,11 +1039,15 @@ TArray<FRigVMExprAST*> FRigVMParserAST::TraversePins(const FRigVMASTProxy& InNod
 	{
 		if(UScriptStruct* ScriptStruct = UnitNode->GetScriptStruct())
 		{
-			for (TFieldIterator<FProperty> PropertyIt(ScriptStruct); PropertyIt; ++PropertyIt)
+			TArray<UStruct*> Structs = FRigVMTemplate::GetSuperStructs(ScriptStruct, true);
+			for(const UStruct* Struct : Structs)
 			{
-				if(URigVMPin* Pin = UnitNode->FindPin(PropertyIt->GetName()))
+				for (TFieldIterator<FProperty> PropertyIt(Struct, EFieldIterationFlags::None); PropertyIt; ++PropertyIt)
 				{
-					Pins.Add(Pin);
+					if(URigVMPin* Pin = UnitNode->FindPin(PropertyIt->GetName()))
+					{
+						Pins.Add(Pin);
+					}
 				}
 			}
 		}
@@ -1042,13 +1071,26 @@ TArray<FRigVMExprAST*> FRigVMParserAST::TraversePins(const FRigVMASTProxy& InNod
 				for (URigVMPin* CasePin : CasePins)
 				{
 					FRigVMASTProxy CasePinProxy = PinProxy.GetSibling(CasePin);
+
 					PinExpressions.Add(TraversePin(CasePinProxy, InParentExpr));
+
+					if(InParentExpr)
+					{
+						const FString PinName = FString::Printf(TEXT("%s_%s"), *Pin->GetName(), *CasePin->GetName());
+						const int32 ChildIndex = InParentExpr->Children.Find(PinExpressions.Last());
+						InParentExpr->PinNameToChildIndex.FindOrAdd(*PinName) = ChildIndex;
+					}
 				}
 				continue;
 			}
 		}
 
 		PinExpressions.Add(TraversePin(PinProxy, InParentExpr));
+		if(InParentExpr)
+		{
+			const int32 ChildIndex = InParentExpr->Children.Find(PinExpressions.Last());
+			InParentExpr->PinNameToChildIndex.FindOrAdd(Pin->GetFName()) = ChildIndex;
+		}
 	}
 
 	return PinExpressions;
diff --git a/Engine/Source/Developer/RigVMDeveloper/Private/RigVMCompiler/RigVMCompiler.cpp b/Engine/Source/Developer/RigVMDeveloper/Private/RigVMCompiler/RigVMCompiler.cpp
--- a/Engine/Source/Developer/RigVMDeveloper/Private/RigVMCompiler/RigVMCompiler.cpp
+++ b/Engine/Source/Developer/RigVMDeveloper/Private/RigVMCompiler/RigVMCompiler.cpp
@@ -791,8 +791,20 @@ int32 URigVMCompiler::TraverseCallExtern(const FRigVMCallExternExprAST* InExpr,
 	else
 	{
 		TArray<FRigVMOperand> Operands;
-		for (FRigVMExprAST* ChildExpr : *InExpr)
+
+		// iterate over the child expressions in the order of the arguments on the function
+		const FRigVMFunction* Function = FRigVMRegistry::Get().FindFunction(UnitNode->GetScriptStruct(), *UnitNode->GetMethodName().ToString());
+		check(Function);
+		
+		for(const FRigVMFunctionArgument& Argument : Function->GetArguments())
 		{
+			const FRigVMExprAST* ChildExpr = InExpr->FindExprWithPinName(Argument.Name);
+			if(ChildExpr == nullptr)
+			{
+				// opaque arguments don't have a matching child expression
+				continue;
+			}
+			
 			if (ChildExpr->GetType() == FRigVMExprAST::EType::CachedValue)
 			{
 				Operands.Add(WorkData.ExprToOperand.FindChecked(GetSourceVarExpr(ChildExpr->To<FRigVMCachedValueExprAST>()->GetVarExpr())));
diff --git a/Engine/Source/Developer/RigVMDeveloper/Private/RigVMModel/RigVMController.cpp b/Engine/Source/Developer/RigVMDeveloper/Private/RigVMModel/RigVMController.cpp
--- a/Engine/Source/Developer/RigVMDeveloper/Private/RigVMModel/RigVMController.cpp
+++ b/Engine/Source/Developer/RigVMDeveloper/Private/RigVMModel/RigVMController.cpp
@@ -13114,75 +13114,81 @@ void URigVMController::AddPinsForStruct(UStruct* InStruct, URigVMNode* InNode, U
 		}
 	}
 
-	for (TFieldIterator<FProperty> It(InStruct); It; ++It)
+	TArray<UStruct*> StructsToVisit = FRigVMTemplate::GetSuperStructs(InStruct, true);
+	for(UStruct* StructToVisit : StructsToVisit)
 	{
-		FName PropertyName = It->GetFName();
-
-		URigVMPin* Pin = NewObject<URigVMPin>(InParentPin == nullptr ? Cast<UObject>(InNode) : Cast<UObject>(InParentPin), PropertyName);
-		ConfigurePinFromProperty(*It, Pin, InPinDirection);
-
-		if (InParentPin)
-		{
-			AddSubPin(InParentPin, Pin);
-		}
-		else
+		// using EFieldIterationFlags::None excludes the
+		// properties of the super struct in this iterator.
+		for (TFieldIterator<FProperty> It(StructToVisit, EFieldIterationFlags::None); It; ++It)
 		{
-			AddNodePin(InNode, Pin);
-		}
+			FName PropertyName = It->GetFName();
 
-		FString* DefaultValuePtr = MemberValues.Find(Pin->GetFName());
+			URigVMPin* Pin = NewObject<URigVMPin>(InParentPin == nullptr ? Cast<UObject>(InNode) : Cast<UObject>(InParentPin), PropertyName);
+			ConfigurePinFromProperty(*It, Pin, InPinDirection);
 
-		FStructProperty* StructProperty = CastField<FStructProperty>(*It);
-		if (StructProperty)
-		{
-			if (ShouldStructBeUnfolded(StructProperty->Struct))
+			if (InParentPin)
 			{
-				FString DefaultValue;
-				if (DefaultValuePtr != nullptr)
-				{
-					DefaultValue = *DefaultValuePtr;
-				}
-				CreateDefaultValueForStructIfRequired(StructProperty->Struct, DefaultValue);
-
-				AddPinsForStruct(StructProperty->Struct, InNode, Pin, Pin->GetDirection(), DefaultValue, bAutoExpandArrays);
+				AddSubPin(InParentPin, Pin);
 			}
-			else if(DefaultValuePtr != nullptr)
+			else
 			{
-				Pin->DefaultValue = *DefaultValuePtr;
+				AddNodePin(InNode, Pin);
 			}
-		}
 
-		FArrayProperty* ArrayProperty = CastField<FArrayProperty>(*It);
-		if (ArrayProperty)
-		{
-			ensure(Pin->IsArray());
+			FString* DefaultValuePtr = MemberValues.Find(Pin->GetFName());
 
-			if (DefaultValuePtr)
+			FStructProperty* StructProperty = CastField<FStructProperty>(*It);
+			if (StructProperty)
 			{
-				if (ShouldPinBeUnfolded(Pin))
+				if (ShouldStructBeUnfolded(StructProperty->Struct))
 				{
-					TArray<FString> ElementDefaultValues = URigVMPin::SplitDefaultValue(*DefaultValuePtr);
-					AddPinsForArray(ArrayProperty, InNode, Pin, Pin->Direction, ElementDefaultValues, bAutoExpandArrays);
+					FString DefaultValue;
+					if (DefaultValuePtr != nullptr)
+					{
+						DefaultValue = *DefaultValuePtr;
+					}
+					CreateDefaultValueForStructIfRequired(StructProperty->Struct, DefaultValue);
+
+					AddPinsForStruct(StructProperty->Struct, InNode, Pin, Pin->GetDirection(), DefaultValue, bAutoExpandArrays);
 				}
-				else
+				else if(DefaultValuePtr != nullptr)
 				{
-					FString DefaultValue = *DefaultValuePtr;
-					PostProcessDefaultValue(Pin, DefaultValue);
 					Pin->DefaultValue = *DefaultValuePtr;
 				}
 			}
-		}
-		
-		if (!Pin->IsArray() && !Pin->IsStruct() && DefaultValuePtr != nullptr)
-		{
-			FString DefaultValue = *DefaultValuePtr;
-			PostProcessDefaultValue(Pin, DefaultValue);
-			Pin->DefaultValue = DefaultValue;
-		}
 
-		if (bNotify)
-		{
-			Notify(ERigVMGraphNotifType::PinAdded, Pin);
+			FArrayProperty* ArrayProperty = CastField<FArrayProperty>(*It);
+			if (ArrayProperty)
+			{
+				ensure(Pin->IsArray());
+
+				if (DefaultValuePtr)
+				{
+					if (ShouldPinBeUnfolded(Pin))
+					{
+						TArray<FString> ElementDefaultValues = URigVMPin::SplitDefaultValue(*DefaultValuePtr);
+						AddPinsForArray(ArrayProperty, InNode, Pin, Pin->Direction, ElementDefaultValues, bAutoExpandArrays);
+					}
+					else
+					{
+						FString DefaultValue = *DefaultValuePtr;
+						PostProcessDefaultValue(Pin, DefaultValue);
+						Pin->DefaultValue = *DefaultValuePtr;
+					}
+				}
+			}
+			
+			if (!Pin->IsArray() && !Pin->IsStruct() && DefaultValuePtr != nullptr)
+			{
+				FString DefaultValue = *DefaultValuePtr;
+				PostProcessDefaultValue(Pin, DefaultValue);
+				Pin->DefaultValue = DefaultValue;
+			}
+
+			if (bNotify)
+			{
+				Notify(ERigVMGraphNotifType::PinAdded, Pin);
+			}
 		}
 	}
 }
@@ -14978,22 +14984,26 @@ bool URigVMController::FullyResolveTemplateNode(URigVMTemplateNode* InNode, int3
 	// find all missing pins which are not arguments for the template
 	if(ResolvedFunction)
 	{
-		for (TFieldIterator<FProperty> It(ResolvedFunction->Struct); It; ++It)
+		TArray<UStruct*> StructsToVisit = FRigVMTemplate::GetSuperStructs(ResolvedFunction->Struct, true);
+		for(UStruct* StructToVisit : StructsToVisit)
 		{
-			const FRigVMTemplateArgument ExpectedArgument(*It);
-			const FRigVMTemplateArgumentType ExpectedType = ExpectedArgument.GetSupportedTypes()[0];
-
-			if(URigVMPin* Pin = InNode->FindPin(It->GetFName().ToString()))
+			for (TFieldIterator<FProperty> It(StructToVisit, EFieldIterationFlags::None); It; ++It)
 			{
-				if(Pin->GetCPPType() != ExpectedType.CPPType)
+				const FRigVMTemplateArgument ExpectedArgument(*It);
+				const FRigVMTemplateArgumentType ExpectedType = ExpectedArgument.GetSupportedTypes()[0];
+
+				if(URigVMPin* Pin = InNode->FindPin(It->GetFName().ToString()))
+				{
+					if(Pin->GetCPPType() != ExpectedType.CPPType)
+					{
+						PinTypesToChange.Add(Pin, ExpectedType);
+					}
+				}
+				else
 				{
-					PinTypesToChange.Add(Pin, ExpectedType);
+					MissingPins.Add(ExpectedArgument);
 				}
 			}
-			else
-			{
-				MissingPins.Add(ExpectedArgument);
-			}
 		}
 	}
 
diff --git a/Engine/Source/Developer/RigVMDeveloper/Public/RigVMCompiler/RigVMAST.h b/Engine/Source/Developer/RigVMDeveloper/Public/RigVMCompiler/RigVMAST.h
--- a/Engine/Source/Developer/RigVMDeveloper/Public/RigVMCompiler/RigVMAST.h
+++ b/Engine/Source/Developer/RigVMDeveloper/Public/RigVMCompiler/RigVMAST.h
@@ -274,7 +274,7 @@ protected:
 	// @param InReplacement the expression to replace this one
 	void ReplaceBy(FRigVMExprAST* InReplacement);
 
-private:
+protected:
 
 	// returns a string containing an indented tree structure
 	// for debugging purposes. this is only used by the parser
@@ -288,6 +288,7 @@ private:
 	const FRigVMParserAST* ParserPtr;
 	TArray<FRigVMExprAST*> Parents;
 	TArray<FRigVMExprAST*> Children;
+	TMap<FName, int32> PinNameToChildIndex;
 
 	friend class FRigVMParserAST;
 	friend class URigVMCompiler;
@@ -544,6 +545,7 @@ public:
 	URigVMNode* GetNode() const { return GetProxy().GetSubjectChecked<URigVMNode>(); }
 
 	virtual bool IsConstant() const override;
+	const FRigVMExprAST* FindExprWithPinName(const FName& InPinName) const;
 	const FRigVMVarExprAST* FindVarWithPinName(const FName& InPinName) const;
 
 protected:
diff --git a/Engine/Source/Runtime/RigVM/Private/RigVMCore/RigVMTemplate.cpp b/Engine/Source/Runtime/RigVM/Private/RigVMCore/RigVMTemplate.cpp
--- a/Engine/Source/Runtime/RigVM/Private/RigVMCore/RigVMTemplate.cpp
+++ b/Engine/Source/Runtime/RigVM/Private/RigVMCore/RigVMTemplate.cpp
@@ -486,15 +486,31 @@ FRigVMTemplate::FRigVMTemplate(UScriptStruct* InStruct, const FString& InTemplat
 	, Notation(NAME_None)
 {
 	TArray<FString> ArgumentNotations;
-	for (TFieldIterator<FProperty> It(InStruct); It; ++It)
+
+	// create the arguments sorted by super -> child struct.
+	TArray<UStruct*> Structs = GetSuperStructs(InStruct, true);
+	for(UStruct* Struct : Structs)
 	{
-		FRigVMTemplateArgument Argument(*It);
-		Argument.Index = Arguments.Num();
+		// only iterate on this struct's fields, not the super structs'
+		for (TFieldIterator<FProperty> It(Struct, EFieldIterationFlags::None); It; ++It)
+		{
+			FRigVMTemplateArgument Argument(*It);
+			Argument.Index = Arguments.Num();
 
-		if(IsValidArgumentForTemplate(Argument))
+			if(IsValidArgumentForTemplate(Argument))
+			{
+				Arguments.Add(Argument);
+			}
+		}
+	}
+
+	// the template notation needs to be in the same order as the C++ implementation,
+	// which is the order of child -> super class members
+	for (TFieldIterator<FProperty> It(InStruct); It; ++It)
+	{
+		if(const FRigVMTemplateArgument* Argument = FindArgument(It->GetFName()))
 		{
-			Arguments.Add(Argument);
-			ArgumentNotations.Add(GetArgumentNotation(Argument));
+			ArgumentNotations.Add(GetArgumentNotation(*Argument));
 		}
 	}
 
@@ -586,6 +602,30 @@ FString FRigVMTemplate::GetArgumentNotation(const FRigVMTemplateArgument& InArgu
 		*InArgument.GetName().ToString());
 }
 
+TArray<UStruct*> FRigVMTemplate::GetSuperStructs(UStruct* InStruct, bool bIncludeLeaf)
+{
+	// Create an array of structs, ordered super -> child struct
+	TArray<UStruct*> SuperStructs = {InStruct};
+	while(true)
+	{
+		if(UStruct* SuperStruct = SuperStructs[0]->GetSuperStruct())
+		{
+			SuperStructs.Insert(SuperStruct, 0);
+		}
+		else
+		{
+			break;
+		}
+	}
+
+	if(!bIncludeLeaf)
+	{
+		SuperStructs.Remove(SuperStructs.Last());
+	}
+
+	return SuperStructs;
+}
+
 bool FRigVMTemplate::IsValid() const
 {
 	return !Notation.IsNone();
diff --git a/Engine/Source/Runtime/RigVM/Public/RigVMCore/RigVMTemplate.h b/Engine/Source/Runtime/RigVM/Public/RigVMCore/RigVMTemplate.h
--- a/Engine/Source/Runtime/RigVM/Public/RigVMCore/RigVMTemplate.h
+++ b/Engine/Source/Runtime/RigVM/Public/RigVMCore/RigVMTemplate.h
@@ -317,6 +317,9 @@ public:
 	// returns the notation of an argument
 	static FString GetArgumentNotation(const FRigVMTemplateArgument& InArgument);
 
+	// returns an array of structs in the inheritance order of a given struct
+	static TArray<UStruct*> GetSuperStructs(UStruct* InStruct, bool bIncludeLeaf = true);
+
 #if WITH_EDITOR
 
 	// Returns the color based on the permutation's metadata
