RigVM: Maintain pin state during unresolve for pins that didn't change type

diff --git a/Engine/Source/Developer/RigVMDeveloper/Private/RigVMModel/RigVMController.cpp b/Engine/Source/Developer/RigVMDeveloper/Private/RigVMModel/RigVMController.cpp
--- a/Engine/Source/Developer/RigVMDeveloper/Private/RigVMModel/RigVMController.cpp
+++ b/Engine/Source/Developer/RigVMDeveloper/Private/RigVMModel/RigVMController.cpp
@@ -16064,6 +16064,9 @@ bool URigVMController::ChangePinType(URigVMPin* InPin, const FString& InCPPType,
 	TArray<URigVMPin*> AllPins = InPin->GetNode()->GetAllPinsRecursively();
 	int32 RemainingWildCardPins = Algo::CountIf(AllPins, WildCardPinCountPredicate);
 	const bool bPinWasWildCard = InPin->IsWildCard();
+
+	const FPinState PreviousPinState = GetPinState(InPin);
+	const FString PreviousCPPType = InPin->CPPType;
 	
 	InPin->CPPType = CPPType;
 	InPin->CPPTypeObjectPath = CPPTypeObjectPath;
@@ -16106,6 +16109,12 @@ bool URigVMController::ChangePinType(URigVMPin* InPin, const FString& InCPPType,
 		}
 	}
 
+	// if the pin didn't change type - let's maintain the pin state
+	if(PreviousCPPType == InPin->CPPType && !InPin->IsWildCard())
+	{
+		ApplyPinState(InPin, PreviousPinState, false);
+	}
+
 	// if this is a template clear its caches
 	if(URigVMTemplateNode* TemplateNode = Cast<URigVMTemplateNode>(InPin->GetNode()))
 	{
