Control Rig: Fix to run backwards compat code

diff --git a/Engine/Plugins/Animation/ControlRig/Source/ControlRigDeveloper/Private/ControlRigBlueprint.cpp b/Engine/Plugins/Animation/ControlRig/Source/ControlRigDeveloper/Private/ControlRigBlueprint.cpp
--- a/Engine/Plugins/Animation/ControlRig/Source/ControlRigDeveloper/Private/ControlRigBlueprint.cpp
+++ b/Engine/Plugins/Animation/ControlRig/Source/ControlRigDeveloper/Private/ControlRigBlueprint.cpp
@@ -124,8 +124,8 @@ UControlRigBlueprint::UControlRigBlueprint(const FObjectInitializer& ObjectIniti
 	RigVMClient.SetOuterClientHost(this, GET_MEMBER_NAME_CHECKED(UControlRigBlueprint, RigVMClient));
 	{
 		TGuardValue<bool> DisableClientNotifs(RigVMClient.bSuspendNotifications, true);
-		RigVMClient.AddModel(RigVMModelPrefix, false, &ObjectInitializer);
-		RigVMClient.GetOrCreateFunctionLibrary(false, &ObjectInitializer);
+		RigVMClient.AddModel(RigVMModelPrefix, false, &ObjectInitializer, false);
+		RigVMClient.GetOrCreateFunctionLibrary(false, &ObjectInitializer, false);
 	}
 	
 	FunctionLibraryEdGraph = ObjectInitializer.CreateDefaultSubobject<UControlRigGraph>(this, TEXT("RigVMFunctionLibraryEdGraph"));
diff --git a/Engine/Source/Developer/RigVMDeveloper/Private/RigVMModel/RigVMClient.cpp b/Engine/Source/Developer/RigVMDeveloper/Private/RigVMModel/RigVMClient.cpp
--- a/Engine/Source/Developer/RigVMDeveloper/Private/RigVMModel/RigVMClient.cpp
+++ b/Engine/Source/Developer/RigVMDeveloper/Private/RigVMModel/RigVMClient.cpp
@@ -213,7 +213,7 @@ URigVMController* FRigVMClient::GetOrCreateController(const UObject* InEditorSid
 	return nullptr;
 }
 
-URigVMGraph* FRigVMClient::AddModel(const FName& InName, bool bSetupUndoRedo, const FObjectInitializer* ObjectInitializer)
+URigVMGraph* FRigVMClient::AddModel(const FName& InName, bool bSetupUndoRedo, const FObjectInitializer* ObjectInitializer, bool bCreateController)
 {
 #if WITH_EDITOR
 	TSharedPtr<FScopedTransaction> Transaction;
@@ -234,7 +234,7 @@ URigVMGraph* FRigVMClient::AddModel(const FName& InName, bool bSetupUndoRedo, co
 	{
 		Model = NewObject<URigVMGraph>(GetOuter(), SafeGraphName);
 	}
-	AddModel(Model, true);
+	AddModel(Model, bCreateController);
 
 	if(bSetupUndoRedo)
 	{
@@ -242,7 +242,7 @@ URigVMGraph* FRigVMClient::AddModel(const FName& InName, bool bSetupUndoRedo, co
 		UndoRedoIndex++;
 
 		FRigVMClientAction Action;
-		Action.Type = ERigVMClientAction_AddModel;
+		Action.Type = ERigVMClientAction::ERigVMClientAction_AddModel;
 		Action.NodePath = Model->GetNodePath();
 		UndoStack.Push(Action);
 		RedoStack.Reset();
@@ -291,7 +291,7 @@ void FRigVMClient::AddModel(URigVMGraph* InModel, bool bCreateController)
 
 }
 
-URigVMFunctionLibrary* FRigVMClient::GetOrCreateFunctionLibrary(bool bSetupUndoRedo, const FObjectInitializer* ObjectInitializer)
+URigVMFunctionLibrary* FRigVMClient::GetOrCreateFunctionLibrary(bool bSetupUndoRedo, const FObjectInitializer* ObjectInitializer, bool bCreateController)
 {
 	if(FunctionLibrary)
 	{
@@ -318,7 +318,7 @@ URigVMFunctionLibrary* FRigVMClient::GetOrCreateFunctionLibrary(bool bSetupUndoR
 	{
 		NewFunctionLibrary = NewObject<URigVMFunctionLibrary>(GetOuter(), SafeGraphName);
 	}
-	AddModel(NewFunctionLibrary, true);
+	AddModel(NewFunctionLibrary, bCreateController);
 	return NewFunctionLibrary;
 }
 
diff --git a/Engine/Source/Developer/RigVMDeveloper/Public/RigVMModel/RigVMClient.h b/Engine/Source/Developer/RigVMDeveloper/Public/RigVMModel/RigVMClient.h
--- a/Engine/Source/Developer/RigVMDeveloper/Public/RigVMModel/RigVMClient.h
+++ b/Engine/Source/Developer/RigVMDeveloper/Public/RigVMModel/RigVMClient.h
@@ -102,10 +102,10 @@ public:
 	URigVMController* GetOrCreateController(const URigVMGraph* InModel);
 	URigVMController* GetOrCreateController(const UObject* InEditorSideObject);
 	URigVMFunctionLibrary* GetFunctionLibrary() const { return FunctionLibrary; }
-	URigVMFunctionLibrary* GetOrCreateFunctionLibrary(bool bSetupUndoRedo, const FObjectInitializer* ObjectInitializer = nullptr);
+	URigVMFunctionLibrary* GetOrCreateFunctionLibrary(bool bSetupUndoRedo, const FObjectInitializer* ObjectInitializer = nullptr, bool bCreateController = true);
 	TArray<FName> GetEntryNames() const;
 
-	URigVMGraph* AddModel(const FName& InName, bool bSetupUndoRedo, const FObjectInitializer* ObjectInitializer = nullptr);
+	URigVMGraph* AddModel(const FName& InName, bool bSetupUndoRedo, const FObjectInitializer* ObjectInitializer = nullptr, bool bCreateController = true);
 	void AddModel(URigVMGraph* InModel, bool bCreateController);
 	bool RemoveModel(const FString& InNodePathOrName, bool bSetupUndoRedo);
 	FName RenameModel(const FString& InNodePathOrName, const FName& InNewName, bool bSetupUndoRedo);
