Fix static analysis warnings
#rnx

diff --git a/Engine/Source/Developer/RigVMDeveloper/Private/RigVMCompiler/RigVMAST.cpp b/Engine/Source/Developer/RigVMDeveloper/Private/RigVMCompiler/RigVMAST.cpp
--- a/Engine/Source/Developer/RigVMDeveloper/Private/RigVMCompiler/RigVMAST.cpp
+++ b/Engine/Source/Developer/RigVMDeveloper/Private/RigVMCompiler/RigVMAST.cpp
@@ -1062,8 +1062,9 @@ TArray<FRigVMExprAST*> FRigVMParserAST::TraversePins(const FRigVMASTProxy& InNod
 	{
 		FRigVMASTProxy PinProxy = InNodeProxy.GetSibling(Pin);
 
-		if (Pin->GetDirection() == ERigVMPinDirection::Input &&
-			InParentExpr->IsA(FRigVMExprAST::EType::Select))
+		if (InParentExpr &&
+		    Pin->GetDirection() == ERigVMPinDirection::Input &&
+		    InParentExpr->IsA(FRigVMExprAST::EType::Select))
 		{
 			if (Pin->GetName() == URigVMSelectNode::ValueName)
 			{
@@ -1074,12 +1075,9 @@ TArray<FRigVMExprAST*> FRigVMParserAST::TraversePins(const FRigVMASTProxy& InNod
 
 					PinExpressions.Add(TraversePin(CasePinProxy, InParentExpr));
 
-					if(InParentExpr)
-					{
-						const FString PinName = FString::Printf(TEXT("%s_%s"), *Pin->GetName(), *CasePin->GetName());
-						const int32 ChildIndex = InParentExpr->Children.Find(PinExpressions.Last());
-						InParentExpr->PinNameToChildIndex.FindOrAdd(*PinName) = ChildIndex;
-					}
+					const FString PinName = FString::Printf(TEXT("%s_%s"), *Pin->GetName(), *CasePin->GetName());
+					const int32 ChildIndex = InParentExpr->Children.Find(PinExpressions.Last());
+					InParentExpr->PinNameToChildIndex.FindOrAdd(*PinName) = ChildIndex;
 				}
 				continue;
 			}
diff --git a/Engine/Source/Editor/GraphEditor/Public/EdGraphHandleTypes.h b/Engine/Source/Editor/GraphEditor/Public/EdGraphHandleTypes.h
--- a/Engine/Source/Editor/GraphEditor/Public/EdGraphHandleTypes.h
+++ b/Engine/Source/Editor/GraphEditor/Public/EdGraphHandleTypes.h
@@ -62,14 +62,22 @@ public:
 
 	FORCEINLINE FEdGraphPinHandle(const UEdGraphPin* InPin)
 		: FEdGraphNodeHandle(InPin ? InPin->GetOwningNode() : nullptr)
-		, PinName(InPin ? InPin->GetFName() : NAME_None)
-		, PinDirection(InPin ? InPin->Direction.GetValue() : EEdGraphPinDirection::EGPD_Input)
-		, PersistentPinGuid(InPin ? InPin->PersistentGuid : FGuid())
 		, PinIndex(INDEX_NONE)
 	{
-		if(const UEdGraphNode* Node = InPin->GetOwningNode())
+		if (InPin)
 		{
-			PinIndex = Node->Pins.Find((UEdGraphPin*)InPin);
+			PinName = InPin->GetFName();
+			PinDirection = InPin->Direction.GetValue();
+			PersistentPinGuid = InPin->PersistentGuid;
+			
+			if (const UEdGraphNode* Node = InPin->GetOwningNode())
+			{
+				PinIndex = Node->Pins.Find((UEdGraphPin*)InPin);
+			}
+		}
+		else
+		{
+			PinDirection = EEdGraphPinDirection::EGPD_Input;
 		}
 	}
 		
