Control Rig: Fix crash when connecting to template node to other type

diff --git a/Engine/Source/Developer/RigVMDeveloper/Private/RigVMModel/RigVMController.cpp b/Engine/Source/Developer/RigVMDeveloper/Private/RigVMModel/RigVMController.cpp
--- a/Engine/Source/Developer/RigVMDeveloper/Private/RigVMModel/RigVMController.cpp
+++ b/Engine/Source/Developer/RigVMDeveloper/Private/RigVMModel/RigVMController.cpp
@@ -9162,10 +9162,10 @@ bool URigVMController::AddLink(URigVMPin* OutputPin, URigVMPin* InputPin, bool b
 	TGuardValue<ERigVMPinDirection> UserLinkDirectionGuard(UserLinkDirection,
 	InUserDirection == ERigVMPinDirection::Invalid ? UserLinkDirection : InUserDirection);
 
+	if (!bIsTransacting)
 	{
 		FString FailureReason;
-		const bool bAllowWildcard = bIsTransacting; 
-		if (!Graph->CanLink(OutputPin, InputPin, &FailureReason, GetCurrentByteCode(), UserLinkDirection, bAllowWildcard))
+		if (!Graph->CanLink(OutputPin, InputPin, &FailureReason, GetCurrentByteCode(), UserLinkDirection))
 		{
 			if(OutputPin->IsExecuteContext() && InputPin->IsExecuteContext())
 			{
diff --git a/Engine/Source/Developer/RigVMDeveloper/Private/RigVMModel/RigVMGraph.cpp b/Engine/Source/Developer/RigVMDeveloper/Private/RigVMModel/RigVMGraph.cpp
--- a/Engine/Source/Developer/RigVMDeveloper/Private/RigVMModel/RigVMGraph.cpp
+++ b/Engine/Source/Developer/RigVMDeveloper/Private/RigVMModel/RigVMGraph.cpp
@@ -421,9 +421,9 @@ void URigVMGraph::PrepareCycleChecking(URigVMPin* InPin, bool bAsInput)
 	GetDiagnosticsAST(false, LinksToSkip)->PrepareCycleChecking(InPin);
 }
 
-bool URigVMGraph::CanLink(URigVMPin* InSourcePin, URigVMPin* InTargetPin, FString* OutFailureReason, const FRigVMByteCode* InByteCode, ERigVMPinDirection InUserLinkDirection, bool bInAllowWildcard)
+bool URigVMGraph::CanLink(URigVMPin* InSourcePin, URigVMPin* InTargetPin, FString* OutFailureReason, const FRigVMByteCode* InByteCode, ERigVMPinDirection InUserLinkDirection)
 {
-	if (!URigVMPin::CanLink(InSourcePin, InTargetPin, OutFailureReason, InByteCode, InUserLinkDirection, bInAllowWildcard))
+	if (!URigVMPin::CanLink(InSourcePin, InTargetPin, OutFailureReason, InByteCode, InUserLinkDirection))
 	{
 		return false;
 	}
diff --git a/Engine/Source/Developer/RigVMDeveloper/Private/RigVMModel/RigVMPin.cpp b/Engine/Source/Developer/RigVMDeveloper/Private/RigVMModel/RigVMPin.cpp
--- a/Engine/Source/Developer/RigVMDeveloper/Private/RigVMModel/RigVMPin.cpp
+++ b/Engine/Source/Developer/RigVMDeveloper/Private/RigVMModel/RigVMPin.cpp
@@ -1603,7 +1603,7 @@ URigVMGraph* URigVMPin::GetGraph() const
 	return nullptr;
 }
 
-bool URigVMPin::CanLink(URigVMPin* InSourcePin, URigVMPin* InTargetPin, FString* OutFailureReason, const FRigVMByteCode* InByteCode, ERigVMPinDirection InUserLinkDirection, bool bInAllowWildcard)
+bool URigVMPin::CanLink(URigVMPin* InSourcePin, URigVMPin* InTargetPin, FString* OutFailureReason, const FRigVMByteCode* InByteCode, ERigVMPinDirection InUserLinkDirection)
 {
 	if (InSourcePin == nullptr || InTargetPin == nullptr)
 	{
@@ -1683,11 +1683,6 @@ bool URigVMPin::CanLink(URigVMPin* InSourcePin, URigVMPin* InTargetPin, FString*
 			bCPPTypesDiffer = false;
 		}
 
-		if (bInAllowWildcard && (InSourcePin->CPPType == RigVMTypeUtils::GetWildCardCPPType() || InTargetPin->CPPType == RigVMTypeUtils::GetWildCardCPPType()))
-		{
-			bCPPTypesDiffer = false;
-		}
-
 		if (bCPPTypesDiffer)
 		{
 			{
diff --git a/Engine/Source/Developer/RigVMDeveloper/Public/RigVMModel/RigVMGraph.h b/Engine/Source/Developer/RigVMDeveloper/Public/RigVMModel/RigVMGraph.h
--- a/Engine/Source/Developer/RigVMDeveloper/Public/RigVMModel/RigVMGraph.h
+++ b/Engine/Source/Developer/RigVMDeveloper/Public/RigVMModel/RigVMGraph.h
@@ -140,7 +140,7 @@ public:
 	UScriptStruct* GetExecuteContextStruct() const;
 
 	void PrepareCycleChecking(URigVMPin* InPin, bool bAsInput);
-	virtual bool CanLink(URigVMPin* InSourcePin, URigVMPin* InTargetPin, FString* OutFailureReason, const FRigVMByteCode* InByteCode, ERigVMPinDirection InUserLinkDirection = ERigVMPinDirection::IO, bool bInAllowWildcard = false);
+	virtual bool CanLink(URigVMPin* InSourcePin, URigVMPin* InTargetPin, FString* OutFailureReason, const FRigVMByteCode* InByteCode, ERigVMPinDirection InUserLinkDirection = ERigVMPinDirection::IO);
 	
 	TSharedPtr<FRigVMParserAST> GetDiagnosticsAST(bool bForceRefresh = false, TArray<URigVMLink*> InLinksToSkip = TArray<URigVMLink*>());
 	TSharedPtr<FRigVMParserAST> GetRuntimeAST(const FRigVMParserASTSettings& InSettings = FRigVMParserASTSettings::Optimized(), bool bForceRefresh = false);
diff --git a/Engine/Source/Developer/RigVMDeveloper/Public/RigVMModel/RigVMPin.h b/Engine/Source/Developer/RigVMDeveloper/Public/RigVMModel/RigVMPin.h
--- a/Engine/Source/Developer/RigVMDeveloper/Public/RigVMModel/RigVMPin.h
+++ b/Engine/Source/Developer/RigVMDeveloper/Public/RigVMModel/RigVMPin.h
@@ -375,7 +375,7 @@ public:
 
 	// Returns true is the two provided source and target Pins
 	// can be linked to one another.
-	static bool CanLink(URigVMPin* InSourcePin, URigVMPin* InTargetPin, FString* OutFailureReason, const FRigVMByteCode* InByteCode, ERigVMPinDirection InUserLinkDirection = ERigVMPinDirection::IO, bool bInAllowWildcard = false);
+	static bool CanLink(URigVMPin* InSourcePin, URigVMPin* InTargetPin, FString* OutFailureReason, const FRigVMByteCode* InByteCode, ERigVMPinDirection InUserLinkDirection = ERigVMPinDirection::IO);
 
 	// Returns true if this pin has injected nodes
 	bool HasInjectedNodes() const { return InjectionInfos.Num() > 0; }
