Control Rig: Fix crash when connecting subpins of template nodes

diff --git a/Engine/Source/Developer/RigVMDeveloper/Private/RigVMModel/RigVMController.cpp b/Engine/Source/Developer/RigVMDeveloper/Private/RigVMModel/RigVMController.cpp
--- a/Engine/Source/Developer/RigVMDeveloper/Private/RigVMModel/RigVMController.cpp
+++ b/Engine/Source/Developer/RigVMDeveloper/Private/RigVMModel/RigVMController.cpp
@@ -3344,7 +3344,7 @@ FString URigVMController::ExportNodesToText(const TArray<FName>& InNodeNames)
 	{
 		URigVMPin* SourcePin = Link->GetSourcePin();
 		URigVMPin* TargetPin = Link->GetTargetPin();
-		if (SourcePin && TargetPin)
+		if (SourcePin && TargetPin && IsValid(SourcePin) && IsValid(TargetPin))
 		{
 			if (!AllNodeNames.Contains(SourcePin->GetNode()->GetFName()))
 			{
@@ -15732,10 +15732,19 @@ bool URigVMController::RecomputeAllTemplateFilteredTypes(bool bSetupUndoRedo)
 	// links like (reroute-reroute) before their filtered permutations are reduced
 	Algo::Sort(SortedLinks, [](const URigVMLink* A, const URigVMLink* B)
 	{
-		const bool bASourceIsUnitNode = A->GetSourcePin()->GetNode()->IsA<URigVMUnitNode>();
-		const bool bATargetIsUnitNode = A->GetTargetPin()->GetNode()->IsA<URigVMUnitNode>();
-		const bool bBSourceIsUnitNode = B->GetSourcePin()->GetNode()->IsA<URigVMUnitNode>();
-		const bool bBTargetIsUnitNode = B->GetTargetPin()->GetNode()->IsA<URigVMUnitNode>();
+		const URigVMPin* APinSource = A->GetSourcePin();
+		const URigVMPin* APinTarget = A->GetTargetPin();
+		const URigVMPin* BPinSource = B->GetSourcePin();
+		const URigVMPin* BPinTarget = B->GetTargetPin();
+		if (!IsValid(APinSource) || !IsValid(APinTarget) || !IsValid(BPinSource) || !IsValid(BPinTarget))
+		{
+			return false;
+		}
+	
+		const bool bASourceIsUnitNode = APinSource->GetNode()->IsA<URigVMUnitNode>();
+		const bool bATargetIsUnitNode = APinTarget->GetNode()->IsA<URigVMUnitNode>();
+		const bool bBSourceIsUnitNode = BPinSource->GetNode()->IsA<URigVMUnitNode>();
+		const bool bBTargetIsUnitNode = BPinTarget->GetNode()->IsA<URigVMUnitNode>();
 		
 		if (bASourceIsUnitNode && bATargetIsUnitNode && (!bBSourceIsUnitNode || !bBTargetIsUnitNode))
 		{
@@ -15757,6 +15766,11 @@ bool URigVMController::RecomputeAllTemplateFilteredTypes(bool bSetupUndoRedo)
 		URigVMPin* OutputPin = Link->GetSourcePin();
 		URigVMPin* InputPin = Link->GetTargetPin();
 
+		if (!IsValid(OutputPin) || !IsValid(InputPin))
+		{
+			continue;
+		}
+		
 		// If pin is a struct member, we should resolve to that type
 		if (URigVMTemplateNode* OutputNode = Cast<URigVMTemplateNode>(OutputPin->GetNode()))
 		{
