Fix use of IsChildOf with possible nullptr

diff --git a/Engine/Source/Developer/RigVMDeveloper/Private/RigVMModel/RigVMNode.cpp b/Engine/Source/Developer/RigVMDeveloper/Private/RigVMModel/RigVMNode.cpp
--- a/Engine/Source/Developer/RigVMDeveloper/Private/RigVMModel/RigVMNode.cpp
+++ b/Engine/Source/Developer/RigVMDeveloper/Private/RigVMModel/RigVMNode.cpp
@@ -283,12 +283,11 @@ bool URigVMNode::IsPure() const
 
 bool URigVMNode::IsMutable() const
 {
-	URigVMPin* ExecutePin = FindPin(FRigVMStruct::ExecuteContextName.ToString());
-	if (ExecutePin)
+	if (const URigVMPin* ExecutePin = FindPin(FRigVMStruct::ExecuteContextName.ToString()))
 	{
-		if (ExecutePin->GetScriptStruct()->IsChildOf(FRigVMExecuteContext::StaticStruct()))
+		if(const UScriptStruct* ScriptStruct = ExecutePin->GetScriptStruct())
 		{
-			return true;
+			return ScriptStruct->IsChildOf(FRigVMExecuteContext::StaticStruct());
 		}
 	}
 	return false;
