Control Rig / RigVM: Improve asset loading performance

 - Overall Loading time went from 21.439 seconds to 7.521 seconds
 - Improved Postload from 6.643 seconds to 0.915 seconds
 - Improved InitControlRigEditor from 13.422 seconds to 4.842 seconds

diff --git a/Engine/Plugins/Animation/ControlRig/Source/ControlRigDeveloper/Private/ControlRigBlueprint.cpp b/Engine/Plugins/Animation/ControlRig/Source/ControlRigDeveloper/Private/ControlRigBlueprint.cpp
--- a/Engine/Plugins/Animation/ControlRig/Source/ControlRigDeveloper/Private/ControlRigBlueprint.cpp
+++ b/Engine/Plugins/Animation/ControlRig/Source/ControlRigDeveloper/Private/ControlRigBlueprint.cpp
@@ -1323,8 +1323,10 @@ void UControlRigBlueprint::DecrementVMRecompileBracket()
 	}
 }
 
-void UControlRigBlueprint::RefreshAllModels()
+void UControlRigBlueprint::RefreshAllModels(EControlRigBlueprintLoadType InLoadType)
 {
+	const bool bIsPostLoad = InLoadType == EControlRigBlueprintLoadType::PostLoad;
+	
 	TGuardValue<bool> IsCompilingGuard(bIsCompiling, true);
 	
 	TArray<URigVMGraph*> GraphsToDetach = RigVMClient.GetAllModels(true, false);
@@ -1352,7 +1354,10 @@ void UControlRigBlueprint::RefreshAllModels()
 		Controller->ReattachLinksToPinObjects(true /* follow redirectors */, nullptr, false, true);
 	}
 
-	PatchTemplateNodesWithPreferredPermutation();
+	if(bIsPostLoad)
+	{
+		PatchTemplateNodesWithPreferredPermutation();
+	}
 	
 	TArray<URigVMGraph*> GraphsToClean = GetAllModels();
 	
@@ -1365,18 +1370,25 @@ void UControlRigBlueprint::RefreshAllModels()
 			Controller->RemoveUnusedOrphanedPins(ModelNode, false);
 		}
 
-		for(URigVMNode* ModelNode : GraphToClean->GetNodes())
+		if(bIsPostLoad)
 		{
-			if (URigVMTemplateNode* TemplateNode = Cast<URigVMTemplateNode>(ModelNode))
+			for(URigVMNode* ModelNode : GraphToClean->GetNodes())
 			{
-				TemplateNode->InvalidateCache();
-				TemplateNode->PostLoad();
+				if (URigVMTemplateNode* TemplateNode = Cast<URigVMTemplateNode>(ModelNode))
+				{
+					TemplateNode->InvalidateCache();
+					TemplateNode->PostLoad();
+				}
 			}
 		}
 #if WITH_EDITOR
-		if (Controller->RecomputeAllTemplateFilteredTypes(false))
+
+		if(bIsPostLoad)
 		{
-			ensureMsgf(false, TEXT("Pin type changed during load %s"), *GetPackage()->GetPathName());
+			if (Controller->RecomputeAllTemplateFilteredTypes(false))
+			{
+				ensureMsgf(false, TEXT("Pin type changed during load %s"), *GetPackage()->GetPathName());
+			}
 		}
 #endif
 	}
diff --git a/Engine/Plugins/Animation/ControlRig/Source/ControlRigDeveloper/Private/Graph/ControlRigGraph.cpp b/Engine/Plugins/Animation/ControlRig/Source/ControlRigDeveloper/Private/Graph/ControlRigGraph.cpp
--- a/Engine/Plugins/Animation/ControlRig/Source/ControlRigDeveloper/Private/Graph/ControlRigGraph.cpp
+++ b/Engine/Plugins/Animation/ControlRig/Source/ControlRigDeveloper/Private/Graph/ControlRigGraph.cpp
@@ -1041,9 +1041,16 @@ UControlRigBlueprint* UControlRigGraph::GetBlueprint() const
 
 URigVMGraph* UControlRigGraph::GetModel() const
 {
-	if (FRigVMClient* Client = GetRigVMClient())
+	if(CachedModelGraph.IsValid())
+	{
+		return CachedModelGraph.Get();
+	}
+	
+	if (const FRigVMClient* Client = GetRigVMClient())
 	{
-		return Client->GetModel(this);
+		URigVMGraph* Model = Client->GetModel(this);
+		CachedModelGraph = Model;
+		return Model;
 	}
 	return nullptr;
 }
@@ -1052,7 +1059,7 @@ URigVMController* UControlRigGraph::GetController() const
 {
 	if (FRigVMClient* Client = GetRigVMClient())
 	{
-		return Client->GetOrCreateController(this);
+		return Client->GetOrCreateController(GetModel());
 	}
 	return nullptr;
 }
diff --git a/Engine/Plugins/Animation/ControlRig/Source/ControlRigDeveloper/Public/ControlRigBlueprint.h b/Engine/Plugins/Animation/ControlRig/Source/ControlRigDeveloper/Public/ControlRigBlueprint.h
--- a/Engine/Plugins/Animation/ControlRig/Source/ControlRigDeveloper/Public/ControlRigBlueprint.h
+++ b/Engine/Plugins/Animation/ControlRig/Source/ControlRigDeveloper/Public/ControlRigBlueprint.h
@@ -195,6 +195,12 @@ struct CONTROLRIGDEVELOPER_API FControlRigPythonSettings
 	}
 };
 
+enum class EControlRigBlueprintLoadType : uint8
+{
+	PostLoad,
+	CheckUserDefinedStructs
+};
+
 UCLASS(BlueprintType, meta=(IgnoreClassThumbnail))
 class CONTROLRIGDEVELOPER_API UControlRigBlueprint : public UBlueprint, public IInterface_PreviewMeshProvider, public IRigVMClientHost
 {
@@ -291,7 +297,7 @@ public:
 	// model data can change while the Control Rig BP is not opened
 	// for example, if a user defined struct changed after BP load,
 	// any pin that references the struct needs to be regenerated
-	void RefreshAllModels();
+	void RefreshAllModels(EControlRigBlueprintLoadType InLoadType = EControlRigBlueprintLoadType::PostLoad);
 
 	UFUNCTION(BlueprintCallable, Category = "Control Rig Blueprint")
 	void RequestControlRigInit();
diff --git a/Engine/Plugins/Animation/ControlRig/Source/ControlRigDeveloper/Public/Graph/ControlRigGraph.h b/Engine/Plugins/Animation/ControlRig/Source/ControlRigDeveloper/Public/Graph/ControlRigGraph.h
--- a/Engine/Plugins/Animation/ControlRig/Source/ControlRigDeveloper/Public/Graph/ControlRigGraph.h
+++ b/Engine/Plugins/Animation/ControlRig/Source/ControlRigDeveloper/Public/Graph/ControlRigGraph.h
@@ -174,7 +174,8 @@ protected:
 private:
 	static TArray<TSharedPtr<FString>> EmptyElementNameList;
 	TMap<FName, UEdGraphNode*> ModelNodePathToEdNode;
-	
+	mutable TWeakObjectPtr<URigVMGraph> CachedModelGraph;
+
 	friend class UControlRigUnitNodeSpawner;
 	friend class UControlRigVariableNodeSpawner;
 	friend class UControlRigParameterNodeSpawner;
diff --git a/Engine/Plugins/Animation/ControlRig/Source/ControlRigEditor/Private/Editor/ControlRigEditor.cpp b/Engine/Plugins/Animation/ControlRig/Source/ControlRigEditor/Private/Editor/ControlRigEditor.cpp
--- a/Engine/Plugins/Animation/ControlRig/Source/ControlRigEditor/Private/Editor/ControlRigEditor.cpp
+++ b/Engine/Plugins/Animation/ControlRig/Source/ControlRigEditor/Private/Editor/ControlRigEditor.cpp
@@ -375,7 +375,7 @@ void FControlRigEditor::InitControlRigEditor(const EToolkitMode::Type Mode, cons
 	// user-defined-struct can change even after load
 	// refresh the models such that pins are updated to match
 	// the latest struct member layout
-	InControlRigBlueprint->RefreshAllModels();
+	InControlRigBlueprint->RefreshAllModels(EControlRigBlueprintLoadType::CheckUserDefinedStructs);
 
 	{
 		TArray<UEdGraph*> EdGraphs;
diff --git a/Engine/Source/Developer/RigVMDeveloper/Private/RigVMModel/Nodes/RigVMArrayNode.cpp b/Engine/Source/Developer/RigVMDeveloper/Private/RigVMModel/Nodes/RigVMArrayNode.cpp
--- a/Engine/Source/Developer/RigVMDeveloper/Private/RigVMModel/Nodes/RigVMArrayNode.cpp
+++ b/Engine/Source/Developer/RigVMDeveloper/Private/RigVMModel/Nodes/RigVMArrayNode.cpp
@@ -251,125 +251,150 @@ const FRigVMTemplate* URigVMArrayNode::GetTemplate() const
 	}
 	
 	if(CachedTemplate == nullptr)
-	{		
-		const TArray<FRigVMTemplateArgumentType>& SingleTypes = FRigVMTemplateArgument::GetCompatibleTypes(FRigVMTemplateArgument::ETypeCategory_SingleAnyValue);
-		const TArray<FRigVMTemplateArgumentType>& ArrayTypes = FRigVMTemplateArgument::GetCompatibleTypes(FRigVMTemplateArgument::ETypeCategory_ArrayAnyValue);
-		const TArray<FRigVMTemplateArgumentType>& ArrayArrayTypes = FRigVMTemplateArgument::GetCompatibleTypes(FRigVMTemplateArgument::ETypeCategory_ArrayArrayAnyValue);
+	{
+		static TArray<const FRigVMTemplate*> TemplatesPerOpCode;
+		static TArray<TArray<FRigVMTemplateArgument>> ArgumentsPerOpCode;
+		if(TemplatesPerOpCode.IsEmpty() || ArgumentsPerOpCode.IsEmpty())
+		{
+			constexpr int32 NumArrayOpCodes = int32(ERigVMOpCode::LastArrayOpCode) - int32(ERigVMOpCode::FirstArrayOpCode) + 1;
+			TemplatesPerOpCode.SetNumZeroed(NumArrayOpCodes);
+			ArgumentsPerOpCode.SetNum(NumArrayOpCodes);
+		}
 
-		TArray<FRigVMTemplateArgumentType> ArrayValueTypes = ArrayTypes;
-		ArrayValueTypes.Append(ArrayArrayTypes);
-		TArray<FRigVMTemplateArgumentType> ElementTypes = SingleTypes;
-		ElementTypes.Append(ArrayTypes);
+		const int32 OpCodeIndex = int32(OpCode) - int32(ERigVMOpCode::FirstArrayOpCode);
 
-		UScriptStruct* ExecuteStruct = GetGraph()->GetExecuteContextStruct();
-		FRigVMTemplateArgumentType ExecuteType(ExecuteStruct->GetStructCPPName(), ExecuteStruct);
+		if(TemplatesPerOpCode[OpCodeIndex])
+		{
+			return TemplatesPerOpCode[OpCodeIndex];
+		}
 		
-		TArray<FRigVMTemplateArgument> Arguments;
+		TArray<FRigVMTemplateArgument>& Arguments = ArgumentsPerOpCode[OpCodeIndex];
 		
-		switch(OpCode)
-	    {
-    		case ERigVMOpCode::ArrayReset:
-			case ERigVMOpCode::ArrayReverse:
-    		{
-    			Arguments.Emplace(*FRigVMStruct::ExecuteContextName.ToString(), ERigVMPinDirection::IO, ExecuteType);
-    			Arguments.Emplace(*ArrayName, ERigVMPinDirection::IO, ArrayValueTypes);
-    			break;
-    		}
-    		case ERigVMOpCode::ArrayGetNum:
-    		{
-    			Arguments.Emplace(*ArrayName, ERigVMPinDirection::Input, ArrayValueTypes);
-    			Arguments.Emplace(*NumName, ERigVMPinDirection::Output, FRigVMTemplateArgumentType(RigVMTypeUtils::Int32Type, nullptr));
-    			break;
-    		} 
-    		case ERigVMOpCode::ArraySetNum:
-    		{
-    			Arguments.Emplace(*FRigVMStruct::ExecuteContextName.ToString(), ERigVMPinDirection::IO, ExecuteType);
-    			Arguments.Emplace(*ArrayName, ERigVMPinDirection::IO, ArrayValueTypes);
-    			Arguments.Emplace(*NumName, ERigVMPinDirection::Input, FRigVMTemplateArgumentType(RigVMTypeUtils::Int32Type, nullptr));
-    			break;
-    		}
-    		case ERigVMOpCode::ArrayGetAtIndex:
-    		{
-    			Arguments.Emplace(*ArrayName, ERigVMPinDirection::Input, ArrayValueTypes);
-    			Arguments.Emplace(*IndexName, ERigVMPinDirection::Input, FRigVMTemplateArgumentType(RigVMTypeUtils::Int32Type, nullptr));
-    			Arguments.Emplace(*ElementName, ERigVMPinDirection::Output, ElementTypes);
-    			break;
-    		}  
-    		case ERigVMOpCode::ArraySetAtIndex:
-			case ERigVMOpCode::ArrayInsert:
-    		{
-    			Arguments.Emplace(*FRigVMStruct::ExecuteContextName.ToString(), ERigVMPinDirection::IO, ExecuteType);
-    			Arguments.Emplace(*ArrayName, ERigVMPinDirection::IO, ArrayValueTypes);
-    			Arguments.Emplace(*IndexName, ERigVMPinDirection::Input, FRigVMTemplateArgumentType(RigVMTypeUtils::Int32Type, nullptr));
-    			Arguments.Emplace(*ElementName, ERigVMPinDirection::Input, ElementTypes);
-    			break;
-    		}
-    		case ERigVMOpCode::ArrayAdd:
-    		{
-    			Arguments.Emplace(*FRigVMStruct::ExecuteContextName.ToString(), ERigVMPinDirection::IO, ExecuteType);
-    			Arguments.Emplace(*ArrayName, ERigVMPinDirection::IO, ArrayValueTypes);
-    			Arguments.Emplace(*ElementName, ERigVMPinDirection::Input, ElementTypes);
-    			Arguments.Emplace(*IndexName, ERigVMPinDirection::Output, FRigVMTemplateArgumentType(RigVMTypeUtils::Int32Type, nullptr));
-    			break;
-    		}
-    		case ERigVMOpCode::ArrayRemove:
-    		{
-    			Arguments.Emplace(*FRigVMStruct::ExecuteContextName.ToString(), ERigVMPinDirection::IO, ExecuteType);
-    			Arguments.Emplace(*ArrayName, ERigVMPinDirection::IO, ArrayValueTypes);
-    			Arguments.Emplace(*IndexName, ERigVMPinDirection::Input, FRigVMTemplateArgumentType(RigVMTypeUtils::Int32Type, nullptr));
-    			break;
-    		}
-    		case ERigVMOpCode::ArrayFind:
-    		{
-    			Arguments.Emplace(*ArrayName, ERigVMPinDirection::Input, ArrayValueTypes);
-    			Arguments.Emplace(*ElementName, ERigVMPinDirection::Input, ElementTypes);
-    			Arguments.Emplace(*IndexName, ERigVMPinDirection::Output, FRigVMTemplateArgumentType(RigVMTypeUtils::Int32Type, nullptr));
-    			Arguments.Emplace(*SuccessName, ERigVMPinDirection::Output, FRigVMTemplateArgumentType(RigVMTypeUtils::BoolType, nullptr));
-    			break;
-    		}
-    		case ERigVMOpCode::ArrayAppend:
-			case ERigVMOpCode::ArrayUnion:
-    		{
-    			Arguments.Emplace(*FRigVMStruct::ExecuteContextName.ToString(), ERigVMPinDirection::IO, ExecuteType);
-    			Arguments.Emplace(*ArrayName, ERigVMPinDirection::IO, ArrayValueTypes);
-    			Arguments.Emplace(*OtherName, ERigVMPinDirection::Input, ArrayValueTypes);
-    			break;
-    		}
-    		case ERigVMOpCode::ArrayClone:
-    		{
-    			Arguments.Emplace(*ArrayName, ERigVMPinDirection::Input, ArrayValueTypes);
-    			Arguments.Emplace(*CloneName, ERigVMPinDirection::Output, ArrayValueTypes);
-    			break;    				
-    		}
-    		case ERigVMOpCode::ArrayIterator:
-    		{
-    			Arguments.Emplace(*FRigVMStruct::ExecuteContextName.ToString(), ERigVMPinDirection::IO, ExecuteType);
-    			Arguments.Emplace(*ArrayName, ERigVMPinDirection::Input, ArrayValueTypes);
-    			Arguments.Emplace(*ElementName, ERigVMPinDirection::Output, ElementTypes);
-    			Arguments.Emplace(*IndexName, ERigVMPinDirection::Output, FRigVMTemplateArgumentType(RigVMTypeUtils::Int32Type, nullptr));
-    			Arguments.Emplace(*CountName, ERigVMPinDirection::Output, FRigVMTemplateArgumentType(RigVMTypeUtils::Int32Type, nullptr));
-    			Arguments.Emplace(*RatioName, ERigVMPinDirection::Output, FRigVMTemplateArgumentType(RigVMTypeUtils::FloatType, nullptr));
-    			Arguments.Emplace(*ContinueName, ERigVMPinDirection::Hidden, ExecuteType);
-    			Arguments.Emplace(*CompletedName, ERigVMPinDirection::Output, ExecuteType);
-    			break;
-    		}
-    		case ERigVMOpCode::ArrayDifference:
-    		case ERigVMOpCode::ArrayIntersection:
-    		{
-    			Arguments.Emplace(*ArrayName, ERigVMPinDirection::Input, ArrayValueTypes);
-    			Arguments.Emplace(*OtherName, ERigVMPinDirection::Input, ArrayValueTypes);
-    			Arguments.Emplace(*ResultName, ERigVMPinDirection::Output, ArrayValueTypes);
-    			break;
-    		}
-    		default:
-    		{
-    			ensure(false);
-    			break;
-    		}
-	    }
+		if(Arguments.IsEmpty())
+		{
+			Arguments.Reserve(8);
+			
+			static const TArray<FRigVMTemplateArgumentType>& SingleTypes = FRigVMTemplateArgument::GetCompatibleTypes(FRigVMTemplateArgument::ETypeCategory_SingleAnyValue);
+			static const TArray<FRigVMTemplateArgumentType>& ArrayTypes = FRigVMTemplateArgument::GetCompatibleTypes(FRigVMTemplateArgument::ETypeCategory_ArrayAnyValue);
+			static const TArray<FRigVMTemplateArgumentType>& ArrayArrayTypes = FRigVMTemplateArgument::GetCompatibleTypes(FRigVMTemplateArgument::ETypeCategory_ArrayArrayAnyValue);
 
+			static TArray<FRigVMTemplateArgumentType> ArrayValueTypes, ElementTypes;
+			if(ArrayValueTypes.IsEmpty())
+			{
+				ArrayValueTypes.Append(ArrayTypes);
+				ArrayValueTypes.Append(ArrayArrayTypes);
+				ElementTypes.Append(SingleTypes);
+				ElementTypes.Append(ArrayTypes);
+			}
+
+			UScriptStruct* ExecuteStruct = GetGraph()->GetExecuteContextStruct();
+			FRigVMTemplateArgumentType ExecuteType(ExecuteStruct->GetStructCPPName(), ExecuteStruct);
+		
+			switch(OpCode)
+			{
+				case ERigVMOpCode::ArrayReset:
+				case ERigVMOpCode::ArrayReverse:
+				{
+					Arguments.Emplace(*FRigVMStruct::ExecuteContextName.ToString(), ERigVMPinDirection::IO, ExecuteType);
+					Arguments.Emplace(*ArrayName, ERigVMPinDirection::IO, ArrayValueTypes);
+					break;
+				}
+				case ERigVMOpCode::ArrayGetNum:
+				{
+					Arguments.Emplace(*ArrayName, ERigVMPinDirection::Input, ArrayValueTypes);
+					Arguments.Emplace(*NumName, ERigVMPinDirection::Output, FRigVMTemplateArgumentType(RigVMTypeUtils::Int32Type, nullptr));
+					break;
+				} 
+				case ERigVMOpCode::ArraySetNum:
+				{
+					Arguments.Emplace(*FRigVMStruct::ExecuteContextName.ToString(), ERigVMPinDirection::IO, ExecuteType);
+					Arguments.Emplace(*ArrayName, ERigVMPinDirection::IO, ArrayValueTypes);
+					Arguments.Emplace(*NumName, ERigVMPinDirection::Input, FRigVMTemplateArgumentType(RigVMTypeUtils::Int32Type, nullptr));
+					break;
+				}
+				case ERigVMOpCode::ArrayGetAtIndex:
+				{
+					Arguments.Emplace(*ArrayName, ERigVMPinDirection::Input, ArrayValueTypes);
+					Arguments.Emplace(*IndexName, ERigVMPinDirection::Input, FRigVMTemplateArgumentType(RigVMTypeUtils::Int32Type, nullptr));
+					Arguments.Emplace(*ElementName, ERigVMPinDirection::Output, ElementTypes);
+					break;
+				}  
+				case ERigVMOpCode::ArraySetAtIndex:
+				case ERigVMOpCode::ArrayInsert:
+				{
+					Arguments.Emplace(*FRigVMStruct::ExecuteContextName.ToString(), ERigVMPinDirection::IO, ExecuteType);
+					Arguments.Emplace(*ArrayName, ERigVMPinDirection::IO, ArrayValueTypes);
+					Arguments.Emplace(*IndexName, ERigVMPinDirection::Input, FRigVMTemplateArgumentType(RigVMTypeUtils::Int32Type, nullptr));
+					Arguments.Emplace(*ElementName, ERigVMPinDirection::Input, ElementTypes);
+					break;
+				}
+				case ERigVMOpCode::ArrayAdd:
+				{
+					Arguments.Emplace(*FRigVMStruct::ExecuteContextName.ToString(), ERigVMPinDirection::IO, ExecuteType);
+					Arguments.Emplace(*ArrayName, ERigVMPinDirection::IO, ArrayValueTypes);
+					Arguments.Emplace(*ElementName, ERigVMPinDirection::Input, ElementTypes);
+					Arguments.Emplace(*IndexName, ERigVMPinDirection::Output, FRigVMTemplateArgumentType(RigVMTypeUtils::Int32Type, nullptr));
+					break;
+				}
+				case ERigVMOpCode::ArrayRemove:
+				{
+					Arguments.Emplace(*FRigVMStruct::ExecuteContextName.ToString(), ERigVMPinDirection::IO, ExecuteType);
+					Arguments.Emplace(*ArrayName, ERigVMPinDirection::IO, ArrayValueTypes);
+					Arguments.Emplace(*IndexName, ERigVMPinDirection::Input, FRigVMTemplateArgumentType(RigVMTypeUtils::Int32Type, nullptr));
+					break;
+				}
+				case ERigVMOpCode::ArrayFind:
+				{
+					Arguments.Emplace(*ArrayName, ERigVMPinDirection::Input, ArrayValueTypes);
+					Arguments.Emplace(*ElementName, ERigVMPinDirection::Input, ElementTypes);
+					Arguments.Emplace(*IndexName, ERigVMPinDirection::Output, FRigVMTemplateArgumentType(RigVMTypeUtils::Int32Type, nullptr));
+					Arguments.Emplace(*SuccessName, ERigVMPinDirection::Output, FRigVMTemplateArgumentType(RigVMTypeUtils::BoolType, nullptr));
+					break;
+				}
+				case ERigVMOpCode::ArrayAppend:
+				case ERigVMOpCode::ArrayUnion:
+				{
+					Arguments.Emplace(*FRigVMStruct::ExecuteContextName.ToString(), ERigVMPinDirection::IO, ExecuteType);
+					Arguments.Emplace(*ArrayName, ERigVMPinDirection::IO, ArrayValueTypes);
+					Arguments.Emplace(*OtherName, ERigVMPinDirection::Input, ArrayValueTypes);
+					break;
+				}
+				case ERigVMOpCode::ArrayClone:
+				{
+					Arguments.Emplace(*ArrayName, ERigVMPinDirection::Input, ArrayValueTypes);
+					Arguments.Emplace(*CloneName, ERigVMPinDirection::Output, ArrayValueTypes);
+					break;    				
+				}
+				case ERigVMOpCode::ArrayIterator:
+				{
+					Arguments.Emplace(*FRigVMStruct::ExecuteContextName.ToString(), ERigVMPinDirection::IO, ExecuteType);
+					Arguments.Emplace(*ArrayName, ERigVMPinDirection::Input, ArrayValueTypes);
+					Arguments.Emplace(*ElementName, ERigVMPinDirection::Output, ElementTypes);
+					Arguments.Emplace(*IndexName, ERigVMPinDirection::Output, FRigVMTemplateArgumentType(RigVMTypeUtils::Int32Type, nullptr));
+					Arguments.Emplace(*CountName, ERigVMPinDirection::Output, FRigVMTemplateArgumentType(RigVMTypeUtils::Int32Type, nullptr));
+					Arguments.Emplace(*RatioName, ERigVMPinDirection::Output, FRigVMTemplateArgumentType(RigVMTypeUtils::FloatType, nullptr));
+					Arguments.Emplace(*ContinueName, ERigVMPinDirection::Hidden, ExecuteType);
+					Arguments.Emplace(*CompletedName, ERigVMPinDirection::Output, ExecuteType);
+					break;
+				}
+				case ERigVMOpCode::ArrayDifference:
+				case ERigVMOpCode::ArrayIntersection:
+				{
+					Arguments.Emplace(*ArrayName, ERigVMPinDirection::Input, ArrayValueTypes);
+					Arguments.Emplace(*OtherName, ERigVMPinDirection::Input, ArrayValueTypes);
+					Arguments.Emplace(*ResultName, ERigVMPinDirection::Output, ArrayValueTypes);
+					break;
+				}
+				default:
+				{
+					ensure(false);
+					break;
+				}
+			}
+		}
+		
 		const FString TemplateName = StaticEnum<ERigVMOpCode>()->GetDisplayNameTextByValue((int64)OpCode).ToString();
-		CachedTemplate = FRigVMRegistry::Get().GetOrAddTemplateFromArguments(*TemplateName, Arguments);
+		TemplatesPerOpCode[OpCodeIndex] = CachedTemplate = FRigVMRegistry::Get().GetOrAddTemplateFromArguments(*TemplateName, Arguments);
 	}
 	return CachedTemplate;
 }
diff --git a/Engine/Source/Developer/RigVMDeveloper/Private/RigVMModel/Nodes/RigVMIfNode.cpp b/Engine/Source/Developer/RigVMDeveloper/Private/RigVMModel/Nodes/RigVMIfNode.cpp
--- a/Engine/Source/Developer/RigVMDeveloper/Private/RigVMModel/Nodes/RigVMIfNode.cpp
+++ b/Engine/Source/Developer/RigVMDeveloper/Private/RigVMModel/Nodes/RigVMIfNode.cpp
@@ -25,18 +25,32 @@ const FRigVMTemplate* URigVMIfNode::GetTemplate() const
 	
 	if(CachedTemplate == nullptr)
 	{
-		TArray<FRigVMTemplateArgumentType> Types;
-		Types.Append(FRigVMTemplateArgument::GetCompatibleTypes(FRigVMTemplateArgument::ETypeCategory_SingleAnyValue));
-		Types.Append(FRigVMTemplateArgument::GetCompatibleTypes(FRigVMTemplateArgument::ETypeCategory_ArrayAnyValue));
-		Types.Append(FRigVMTemplateArgument::GetCompatibleTypes(FRigVMTemplateArgument::ETypeCategory_ArrayArrayAnyValue));
+		static const FRigVMTemplate* IfNodeTemplate = nullptr;
+		if(IfNodeTemplate)
+		{
+			return IfNodeTemplate;
+		}
 		
-		TArray<FRigVMTemplateArgument> Arguments;
-		Arguments.Emplace(*ConditionName, ERigVMPinDirection::Input, FRigVMTemplateArgumentType(RigVMTypeUtils::BoolType, nullptr));
-		Arguments.Emplace(*TrueName, ERigVMPinDirection::Input, Types);
-		Arguments.Emplace(*FalseName, ERigVMPinDirection::Input, Types);
-		Arguments.Emplace(*ResultName, ERigVMPinDirection::Output, Types);
+		static TArray<FRigVMTemplateArgument> Arguments;
+		if(Arguments.IsEmpty())
+		{
+			static TArray<FRigVMTemplateArgumentType> Types;
+			if(Types.IsEmpty())
+			{
+				Types.Reserve(3);
+				Types.Append(FRigVMTemplateArgument::GetCompatibleTypes(FRigVMTemplateArgument::ETypeCategory_SingleAnyValue));
+				Types.Append(FRigVMTemplateArgument::GetCompatibleTypes(FRigVMTemplateArgument::ETypeCategory_ArrayAnyValue));
+				Types.Append(FRigVMTemplateArgument::GetCompatibleTypes(FRigVMTemplateArgument::ETypeCategory_ArrayArrayAnyValue));
+			}
+
+			Arguments.Reserve(4);
+			Arguments.Emplace(*ConditionName, ERigVMPinDirection::Input, FRigVMTemplateArgumentType(RigVMTypeUtils::BoolType, nullptr));
+			Arguments.Emplace(*TrueName, ERigVMPinDirection::Input, Types);
+			Arguments.Emplace(*FalseName, ERigVMPinDirection::Input, Types);
+			Arguments.Emplace(*ResultName, ERigVMPinDirection::Output, Types);
+		}
 		
-		CachedTemplate = FRigVMRegistry::Get().GetOrAddTemplateFromArguments(*IfName, Arguments);
+		IfNodeTemplate = CachedTemplate = FRigVMRegistry::Get().GetOrAddTemplateFromArguments(*IfName, Arguments);
 	}
 	return CachedTemplate;
 }
diff --git a/Engine/Source/Developer/RigVMDeveloper/Private/RigVMModel/Nodes/RigVMRerouteNode.cpp b/Engine/Source/Developer/RigVMDeveloper/Private/RigVMModel/Nodes/RigVMRerouteNode.cpp
--- a/Engine/Source/Developer/RigVMDeveloper/Private/RigVMModel/Nodes/RigVMRerouteNode.cpp
+++ b/Engine/Source/Developer/RigVMDeveloper/Private/RigVMModel/Nodes/RigVMRerouteNode.cpp
@@ -80,18 +80,29 @@ const FRigVMTemplate* URigVMRerouteNode::GetTemplate() const
 	
 	if(CachedTemplate == nullptr)
 	{
-		TArray<FRigVMTemplateArgumentType> Types;
-		Types.Append(FRigVMTemplateArgument::GetCompatibleTypes(FRigVMTemplateArgument::ETypeCategory_SingleAnyValue));
-		Types.Append(FRigVMTemplateArgument::GetCompatibleTypes(FRigVMTemplateArgument::ETypeCategory_ArrayAnyValue));
-		Types.Append(FRigVMTemplateArgument::GetCompatibleTypes(FRigVMTemplateArgument::ETypeCategory_ArrayArrayAnyValue));
+		static const FRigVMTemplate* RerouteNodeTemplate = nullptr;
+		if(RerouteNodeTemplate)
+		{
+			return RerouteNodeTemplate;
+		}
 
-		UScriptStruct* ExecuteStruct = GetGraph()->GetExecuteContextStruct();
-		Types.Add(FRigVMTemplateArgumentType(ExecuteStruct->GetStructCPPName(), ExecuteStruct));
-		
-		TArray<FRigVMTemplateArgument> Arguments;
-		Arguments.Emplace(TEXT("Value"), ERigVMPinDirection::IO, Types);
+		static TArray<FRigVMTemplateArgument> Arguments;
+		if(Arguments.IsEmpty())
+		{
+			static TArray<FRigVMTemplateArgumentType> Types;
+			if(Types.IsEmpty())
+			{
+				Types.Append(FRigVMTemplateArgument::GetCompatibleTypes(FRigVMTemplateArgument::ETypeCategory_SingleAnyValue));
+				Types.Append(FRigVMTemplateArgument::GetCompatibleTypes(FRigVMTemplateArgument::ETypeCategory_ArrayAnyValue));
+				Types.Append(FRigVMTemplateArgument::GetCompatibleTypes(FRigVMTemplateArgument::ETypeCategory_ArrayArrayAnyValue));
+
+				UScriptStruct* ExecuteStruct = GetGraph()->GetExecuteContextStruct();
+				Types.Add(FRigVMTemplateArgumentType(ExecuteStruct->GetStructCPPName(), ExecuteStruct));
+			}
 		
-		CachedTemplate = FRigVMRegistry::Get().GetOrAddTemplateFromArguments(*RerouteName, Arguments);
+			Arguments.Emplace(TEXT("Value"), ERigVMPinDirection::IO, Types);
+		}
+		RerouteNodeTemplate = CachedTemplate = FRigVMRegistry::Get().GetOrAddTemplateFromArguments(*RerouteName, Arguments);
 	}
 	return CachedTemplate;
 }
diff --git a/Engine/Source/Developer/RigVMDeveloper/Private/RigVMModel/Nodes/RigVMSelectNode.cpp b/Engine/Source/Developer/RigVMDeveloper/Private/RigVMModel/Nodes/RigVMSelectNode.cpp
--- a/Engine/Source/Developer/RigVMDeveloper/Private/RigVMModel/Nodes/RigVMSelectNode.cpp
+++ b/Engine/Source/Developer/RigVMDeveloper/Private/RigVMModel/Nodes/RigVMSelectNode.cpp
@@ -36,21 +36,34 @@ const FRigVMTemplate* URigVMSelectNode::GetTemplate() const
 	
 	if(CachedTemplate == nullptr)
 	{
-		const TArray<FRigVMTemplateArgumentType>& SingleTypes = FRigVMTemplateArgument::GetCompatibleTypes(FRigVMTemplateArgument::ETypeCategory_SingleAnyValue);
-		const TArray<FRigVMTemplateArgumentType>& ArrayTypes = FRigVMTemplateArgument::GetCompatibleTypes(FRigVMTemplateArgument::ETypeCategory_ArrayAnyValue);
-		const TArray<FRigVMTemplateArgumentType>& ArrayArrayTypes = FRigVMTemplateArgument::GetCompatibleTypes(FRigVMTemplateArgument::ETypeCategory_ArrayArrayAnyValue);
-
-		TArray<FRigVMTemplateArgumentType> ResultTypes = SingleTypes;
-		ResultTypes.Append(ArrayTypes);
-		TArray<FRigVMTemplateArgumentType> ValueTypes = ArrayTypes;
-		ValueTypes.Append(ArrayArrayTypes);		
-		
-		TArray<FRigVMTemplateArgument> Arguments;
-		Arguments.Emplace(*IndexName, ERigVMPinDirection::Input, FRigVMTemplateArgumentType(RigVMTypeUtils::Int32Type, nullptr));
-		Arguments.Emplace(*ValueName, ERigVMPinDirection::Input, ValueTypes);
-		Arguments.Emplace(*ResultName, ERigVMPinDirection::Output, ResultTypes);
-		
-		CachedTemplate = FRigVMRegistry::Get().GetOrAddTemplateFromArguments(*SelectName, Arguments);
+		static const FRigVMTemplate* SelectNodeTemplate = nullptr;
+		if(SelectNodeTemplate)
+		{
+			return SelectNodeTemplate;
+		}
+
+		static TArray<FRigVMTemplateArgument> Arguments;
+		if(Arguments.IsEmpty())
+		{
+			static const TArray<FRigVMTemplateArgumentType>& SingleTypes = FRigVMTemplateArgument::GetCompatibleTypes(FRigVMTemplateArgument::ETypeCategory_SingleAnyValue);
+			static const TArray<FRigVMTemplateArgumentType>& ArrayTypes = FRigVMTemplateArgument::GetCompatibleTypes(FRigVMTemplateArgument::ETypeCategory_ArrayAnyValue);
+			static const TArray<FRigVMTemplateArgumentType>& ArrayArrayTypes = FRigVMTemplateArgument::GetCompatibleTypes(FRigVMTemplateArgument::ETypeCategory_ArrayArrayAnyValue);
+
+			static TArray<FRigVMTemplateArgumentType> ResultTypes, ValueTypes;
+			if(ResultTypes.IsEmpty())
+			{
+				ResultTypes.Append(SingleTypes);
+				ResultTypes.Append(ArrayTypes);
+				ValueTypes.Append(ArrayTypes);
+				ValueTypes.Append(ArrayArrayTypes);
+			}
+
+			Arguments.Reserve(3);
+			Arguments.Emplace(*IndexName, ERigVMPinDirection::Input, FRigVMTemplateArgumentType(RigVMTypeUtils::Int32Type, nullptr));
+			Arguments.Emplace(*ValueName, ERigVMPinDirection::Input, ValueTypes);
+			Arguments.Emplace(*ResultName, ERigVMPinDirection::Output, ResultTypes);
+		}
+		SelectNodeTemplate = CachedTemplate = FRigVMRegistry::Get().GetOrAddTemplateFromArguments(*SelectName, Arguments);
 	}
 	return CachedTemplate;
 }
\ No newline at end of file
diff --git a/Engine/Source/Developer/RigVMDeveloper/Private/RigVMModel/Nodes/RigVMTemplateNode.cpp b/Engine/Source/Developer/RigVMDeveloper/Private/RigVMModel/Nodes/RigVMTemplateNode.cpp
--- a/Engine/Source/Developer/RigVMDeveloper/Private/RigVMModel/Nodes/RigVMTemplateNode.cpp
+++ b/Engine/Source/Developer/RigVMDeveloper/Private/RigVMModel/Nodes/RigVMTemplateNode.cpp
@@ -777,6 +777,8 @@ bool URigVMTemplateNode::UpdateFilteredPermutations(URigVMPin* InPin, const TArr
 
 void URigVMTemplateNode::InvalidateCache()
 {
+	Super::InvalidateCache();
+	
 	SupportedTypesCache.Reset();
 	CachedFunction = nullptr;
 	CachedTemplate = nullptr;
