RigVM: Remove the support for arrays of arrays

diff --git a/Engine/Source/Developer/RigVMDeveloper/Private/RigVMModel/Nodes/RigVMArrayNode.cpp b/Engine/Source/Developer/RigVMDeveloper/Private/RigVMModel/Nodes/RigVMArrayNode.cpp
--- a/Engine/Source/Developer/RigVMDeveloper/Private/RigVMModel/Nodes/RigVMArrayNode.cpp
+++ b/Engine/Source/Developer/RigVMDeveloper/Private/RigVMModel/Nodes/RigVMArrayNode.cpp
@@ -276,15 +276,12 @@ const FRigVMTemplate* URigVMArrayNode::GetTemplate() const
 			
 			static const TArray<FRigVMTemplateArgumentType>& SingleTypes = FRigVMTemplateArgument::GetCompatibleTypes(FRigVMTemplateArgument::ETypeCategory_SingleAnyValue);
 			static const TArray<FRigVMTemplateArgumentType>& ArrayTypes = FRigVMTemplateArgument::GetCompatibleTypes(FRigVMTemplateArgument::ETypeCategory_ArrayAnyValue);
-			static const TArray<FRigVMTemplateArgumentType>& ArrayArrayTypes = FRigVMTemplateArgument::GetCompatibleTypes(FRigVMTemplateArgument::ETypeCategory_ArrayArrayAnyValue);
 
 			static TArray<FRigVMTemplateArgumentType> ArrayValueTypes, ElementTypes;
 			if(ArrayValueTypes.IsEmpty())
 			{
 				ArrayValueTypes.Append(ArrayTypes);
-				ArrayValueTypes.Append(ArrayArrayTypes);
 				ElementTypes.Append(SingleTypes);
-				ElementTypes.Append(ArrayTypes);
 			}
 
 			UScriptStruct* ExecuteStruct = GetGraph()->GetExecuteContextStruct();
diff --git a/Engine/Source/Developer/RigVMDeveloper/Private/RigVMModel/Nodes/RigVMIfNode.cpp b/Engine/Source/Developer/RigVMDeveloper/Private/RigVMModel/Nodes/RigVMIfNode.cpp
--- a/Engine/Source/Developer/RigVMDeveloper/Private/RigVMModel/Nodes/RigVMIfNode.cpp
+++ b/Engine/Source/Developer/RigVMDeveloper/Private/RigVMModel/Nodes/RigVMIfNode.cpp
@@ -40,7 +40,6 @@ const FRigVMTemplate* URigVMIfNode::GetTemplate() const
 				Types.Reserve(3);
 				Types.Append(FRigVMTemplateArgument::GetCompatibleTypes(FRigVMTemplateArgument::ETypeCategory_SingleAnyValue));
 				Types.Append(FRigVMTemplateArgument::GetCompatibleTypes(FRigVMTemplateArgument::ETypeCategory_ArrayAnyValue));
-				Types.Append(FRigVMTemplateArgument::GetCompatibleTypes(FRigVMTemplateArgument::ETypeCategory_ArrayArrayAnyValue));
 			}
 
 			Arguments.Reserve(4);
diff --git a/Engine/Source/Developer/RigVMDeveloper/Private/RigVMModel/Nodes/RigVMRerouteNode.cpp b/Engine/Source/Developer/RigVMDeveloper/Private/RigVMModel/Nodes/RigVMRerouteNode.cpp
--- a/Engine/Source/Developer/RigVMDeveloper/Private/RigVMModel/Nodes/RigVMRerouteNode.cpp
+++ b/Engine/Source/Developer/RigVMDeveloper/Private/RigVMModel/Nodes/RigVMRerouteNode.cpp
@@ -94,7 +94,6 @@ const FRigVMTemplate* URigVMRerouteNode::GetTemplate() const
 			{
 				Types.Append(FRigVMTemplateArgument::GetCompatibleTypes(FRigVMTemplateArgument::ETypeCategory_SingleAnyValue));
 				Types.Append(FRigVMTemplateArgument::GetCompatibleTypes(FRigVMTemplateArgument::ETypeCategory_ArrayAnyValue));
-				Types.Append(FRigVMTemplateArgument::GetCompatibleTypes(FRigVMTemplateArgument::ETypeCategory_ArrayArrayAnyValue));
 
 				UScriptStruct* ExecuteStruct = GetGraph()->GetExecuteContextStruct();
 				Types.Add(FRigVMTemplateArgumentType(ExecuteStruct->GetStructCPPName(), ExecuteStruct));
diff --git a/Engine/Source/Runtime/RigVM/Private/RigVMCore/RigVMTemplate.cpp b/Engine/Source/Runtime/RigVM/Private/RigVMCore/RigVMTemplate.cpp
--- a/Engine/Source/Runtime/RigVM/Private/RigVMCore/RigVMTemplate.cpp
+++ b/Engine/Source/Runtime/RigVM/Private/RigVMCore/RigVMTemplate.cpp
@@ -158,7 +158,7 @@ FRigVMTemplateArgument::EArrayType FRigVMTemplateArgument::GetArrayType() const
 	return EArrayType_Invalid;
 }
 
-const TArray<FRigVMTemplateArgumentType> FRigVMTemplateArgument::GetCompatibleTypes(ETypeCategory InCategory)
+const TArray<FRigVMTemplateArgumentType>& FRigVMTemplateArgument::GetCompatibleTypes(ETypeCategory InCategory)
 {
 	check(InCategory != ETypeCategory_Invalid);
 	
diff --git a/Engine/Source/Runtime/RigVM/Public/RigVMCore/RigVMTemplate.h b/Engine/Source/Runtime/RigVM/Public/RigVMCore/RigVMTemplate.h
--- a/Engine/Source/Runtime/RigVM/Public/RigVMCore/RigVMTemplate.h
+++ b/Engine/Source/Runtime/RigVM/Public/RigVMCore/RigVMTemplate.h
@@ -216,7 +216,7 @@ struct RIGVM_API FRigVMTemplateArgument
 	EArrayType GetArrayType() const;
 	
 	// Returns all compatible types given a category
-	static const TArray<FRigVMTemplateArgumentType> GetCompatibleTypes(ETypeCategory InCategory);
+	static const TArray<FRigVMTemplateArgumentType>& GetCompatibleTypes(ETypeCategory InCategory);
 
 protected:
 
