Follow up to CL20773015

improved RefreshAllModels
- from 431.621 ms to  318.637 ms when first opening CR_Mannequin_Body
- from 6.930 sec to 3.390 sec for next openings

diff --git a/Engine/Source/Developer/RigVMDeveloper/Private/RigVMModel/RigVMController.cpp b/Engine/Source/Developer/RigVMDeveloper/Private/RigVMModel/RigVMController.cpp
--- a/Engine/Source/Developer/RigVMDeveloper/Private/RigVMModel/RigVMController.cpp
+++ b/Engine/Source/Developer/RigVMDeveloper/Private/RigVMModel/RigVMController.cpp
@@ -7382,7 +7382,8 @@ bool URigVMController::SetPinDefaultValue(URigVMPin* InPin, const FString& InDef
 	
 	if(!bSetPinDefaultValueSucceeded)
 	{
-		if (InPin->GetSubPins().Num() == 0)
+		// no need to send notifications if not changing the value
+		if (InPin->GetSubPins().IsEmpty() && (InPin->DefaultValue != ClampedDefaultValue))
 		{
 			InPin->DefaultValue = ClampedDefaultValue;
 			Notify(ERigVMGraphNotifType::PinDefaultValueChanged, InPin);
