RigVM: Default values on array element pins don't propagate correctly.

diff --git a/Engine/Source/Developer/RigVMDeveloper/Private/RigVMCompiler/RigVMAST.cpp b/Engine/Source/Developer/RigVMDeveloper/Private/RigVMCompiler/RigVMAST.cpp
--- a/Engine/Source/Developer/RigVMDeveloper/Private/RigVMCompiler/RigVMAST.cpp
+++ b/Engine/Source/Developer/RigVMDeveloper/Private/RigVMCompiler/RigVMAST.cpp
@@ -3189,8 +3189,6 @@ void FRigVMParserAST::Inline(TArray<URigVMGraph*> InGraphs, const TArray<FRigVMA
 					SourceNode->IsA<URigVMLibraryNode>() ||
 					SourceNode->IsA<URigVMFunctionReturnNode>())
 				{
-					OutTraversalInfo.PinOverrides->FindOrAdd(InPinProxy) = URigVMPin::FPinOverrideValue(SourcePin);
-
 					// for arrays - if there are sub-pins on the determined source, we need to walk those as well
 					if(SourcePin->IsArray())
 					{
@@ -3204,6 +3202,13 @@ void FRigVMParserAST::Inline(TArray<URigVMGraph*> InGraphs, const TArray<FRigVMA
 							SourceSubPins.Append(SourceSubPin->GetSubPins());
 						}
 					}
+
+					// when asking for the default value of the array - we may need to get the previously stored override
+					const URigVMPin::FPinOverride SourceOverride(SourcePinProxy, *OutTraversalInfo.PinOverrides);
+						
+					// we query the default value now since the potential array elements have been visited and their
+					// potential default value override has been stored to the map.
+					OutTraversalInfo.PinOverrides->FindOrAdd(InPinProxy) = URigVMPin::FPinOverrideValue(SourcePin, SourceOverride);
 				}
 				else
 				{
diff --git a/Engine/Source/Developer/RigVMDeveloper/Public/RigVMModel/RigVMPin.h b/Engine/Source/Developer/RigVMDeveloper/Public/RigVMModel/RigVMPin.h
--- a/Engine/Source/Developer/RigVMDeveloper/Public/RigVMModel/RigVMPin.h
+++ b/Engine/Source/Developer/RigVMDeveloper/Public/RigVMModel/RigVMPin.h
@@ -101,6 +101,12 @@ public:
 		{
 		}
 
+		FPinOverrideValue(URigVMPin* InPin, const TPair<FRigVMASTProxy, const TMap<FRigVMASTProxy, FPinOverrideValue>&>& InOverride)
+			: DefaultValue(InPin->GetDefaultValue(InOverride))
+			, BoundVariablePath(InPin->GetBoundVariablePath())
+		{
+		}
+
 		FString DefaultValue;
 		FString BoundVariablePath;
 	};
