RigVM: Refactor template type system to rely on integers rather than string comparisons

* Switching the FRigVMTemplateArgumentType to be backed up by FName instead of FString for faster comparison
* Adding a "Types" storage under the RigVMRegistry of all of the known supported types
* Have the FRigVMTemplateArgument refer to type indices rather than full type structures
* Have the resolval comparison code rely on comparing type indices rather than type structures
* Store relationship between the bool and the TArray<bool> type as an array / element lookup within the registry - to avoid string operations when doing lookups from element to array to 2 dimensional array for a type.
* The preferred types list on the template node is now stored as pairs of argument name and type index <FName, int32> rather than FString.
* Also we no longer support all types for the untyped nodes (such as reroute, select, if etc) - but instead we support all user defined types (non-native structs and enums) as well as all types ever used as a template argument. The reasoning here is that supporting a type that's not used by any node is pointless. This reduced the type set from 9876 to 237.

diff --git a/Engine/Plugins/Animation/ControlRig/Source/ControlRigDeveloper/Private/ControlRigBlueprint.cpp b/Engine/Plugins/Animation/ControlRig/Source/ControlRigDeveloper/Private/ControlRigBlueprint.cpp
--- a/Engine/Plugins/Animation/ControlRig/Source/ControlRigDeveloper/Private/ControlRigBlueprint.cpp
+++ b/Engine/Plugins/Animation/ControlRig/Source/ControlRigDeveloper/Private/ControlRigBlueprint.cpp
@@ -4224,7 +4224,7 @@ void UControlRigBlueprint::PatchParameterNodesOnLoad()
 				{
 					FRigVMGraphParameterDescription Description = ParameterNode->GetParameterDescription();
 
-					const FEdGraphPinType PinType = RigVMTypeUtils::PinTypeFromCPPType(Description.CPPType, Description.CPPTypeObject);
+					const FEdGraphPinType PinType = RigVMTypeUtils::PinTypeFromCPPType(*Description.CPPType, Description.CPPTypeObject);
 
 					int32 VariableIndex = INDEX_NONE;
 					bool bFoundName = false;
@@ -4678,6 +4678,10 @@ void UControlRigBlueprint::OnVariableAdded(const FName& InVarName)
 			CPPType = CPPTypeObject->GetPathName();
 		}
 	}
+
+	// register the type in the registry
+	FRigVMRegistry::Get().FindOrAddType(FRigVMTemplateArgumentType(*CPPType, CPPTypeObject));
+	
     RigVMPythonUtils::Print(GetFName().ToString(),
 		FString::Printf(TEXT("blueprint.add_member_variable('%s', '%s', %s, %s, '%s')"),
 			*InVarName.ToString(),
@@ -4785,6 +4789,9 @@ void UControlRigBlueprint::OnVariableTypeChanged(const FName& InVarName, FEdGrap
 		}
 	}
 
+	// register the type in the registry
+	FRigVMRegistry::Get().FindOrAddType(FRigVMTemplateArgumentType(*CPPType, CPPTypeObject));
+
 	RigVMPythonUtils::Print(GetFName().ToString(),
 		FString::Printf(TEXT("blueprint.change_member_variable_type('%s', '%s')"),
 		*InVarName.ToString(),
diff --git a/Engine/Plugins/Animation/ControlRig/Source/ControlRigDeveloper/Private/Graph/ControlRigGraphNode.cpp b/Engine/Plugins/Animation/ControlRig/Source/ControlRigDeveloper/Private/Graph/ControlRigGraphNode.cpp
--- a/Engine/Plugins/Animation/ControlRig/Source/ControlRigDeveloper/Private/Graph/ControlRigGraphNode.cpp
+++ b/Engine/Plugins/Animation/ControlRig/Source/ControlRigDeveloper/Private/Graph/ControlRigGraphNode.cpp
@@ -1209,7 +1209,7 @@ bool UControlRigGraphNode::ShouldDrawNodeAsControlPointOnly(int32& OutInputPinIn
 
 FEdGraphPinType UControlRigGraphNode::GetPinTypeForModelPin(URigVMPin* InModelPin)
 {
-	FEdGraphPinType PinType = RigVMTypeUtils::PinTypeFromCPPType(InModelPin->GetCPPType(), InModelPin->GetCPPTypeObject());
+	FEdGraphPinType PinType = RigVMTypeUtils::PinTypeFromCPPType(*InModelPin->GetCPPType(), InModelPin->GetCPPTypeObject());
 	PinType.bIsConst = InModelPin->IsDefinedAsConstant();
 	return PinType;
 }
diff --git a/Engine/Plugins/Animation/ControlRig/Source/ControlRigDeveloper/Private/Graph/ControlRigGraphSchema.cpp b/Engine/Plugins/Animation/ControlRig/Source/ControlRigDeveloper/Private/Graph/ControlRigGraphSchema.cpp
--- a/Engine/Plugins/Animation/ControlRig/Source/ControlRigDeveloper/Private/Graph/ControlRigGraphSchema.cpp
+++ b/Engine/Plugins/Animation/ControlRig/Source/ControlRigDeveloper/Private/Graph/ControlRigGraphSchema.cpp
@@ -1795,7 +1795,10 @@ bool UControlRigGraphSchema::ArePinsCompatible(const UEdGraphPin* PinA, const UE
 							URigVMPin::SplitPinPathAtEnd(InPin->ParentPin->GetName(), PinPath, PinName);
 							CPPType = RigVMTypeUtils::ArrayTypeFromBaseType(CPPType);
 						}
-						if(!Template->ArgumentSupportsType(*PinName, CPPType))
+
+						const FRigVMTemplateArgumentType Type(*CPPType, CPPTypeObject);
+						const int32 TypeIndex = FRigVMRegistry::Get().GetTypeIndex(Type);
+						if(!Template->ArgumentSupportsTypeIndex(*PinName, TypeIndex))
 						{
 							return false;
 						}
diff --git a/Engine/Plugins/Animation/ControlRig/Source/ControlRigEditor/Private/Editor/ControlRigEditor.cpp b/Engine/Plugins/Animation/ControlRig/Source/ControlRigEditor/Private/Editor/ControlRigEditor.cpp
--- a/Engine/Plugins/Animation/ControlRig/Source/ControlRigEditor/Private/Editor/ControlRigEditor.cpp
+++ b/Engine/Plugins/Animation/ControlRig/Source/ControlRigEditor/Private/Editor/ControlRigEditor.cpp
@@ -5757,7 +5757,7 @@ void FControlRigEditor::HandleMakeElementGetterSetter(ERigElementGetterSetterTyp
 		FName ChannelValue = Key.Name;
 		TArray<FName> NamePins;
 		TArray<FName> ChannelPins;
-		TMap<FName, FRigVMTemplateArgumentType> PinsToResolve; 
+		TMap<FName, int32> PinsToResolve; 
 
 		if(FRigControlElement* ControlElement = Hierarchy->Find<FRigControlElement>(Key))
 		{
@@ -5789,7 +5789,7 @@ void FControlRigEditor::HandleMakeElementGetterSetter(ERigElementGetterSetterTyp
 						{
 							StructTemplate = FRigUnit_SetBoolAnimationChannel::StaticStruct();
 						}
-						PinsToResolve.Add(ValueName, FRigVMTemplateArgumentType(RigVMTypeUtils::BoolType));
+						PinsToResolve.Add(ValueName, RigVMTypeUtils::TypeIndex::Bool);
 						break;
 					}
 					case ERigControlType::Float:
@@ -5802,7 +5802,7 @@ void FControlRigEditor::HandleMakeElementGetterSetter(ERigElementGetterSetterTyp
 						{
 							StructTemplate = FRigUnit_SetFloatAnimationChannel::StaticStruct();
 						}
-						PinsToResolve.Add(ValueName, FRigVMTemplateArgumentType(RigVMTypeUtils::FloatType));
+						PinsToResolve.Add(ValueName, RigVMTypeUtils::TypeIndex::Float);
 						break;
 					}
 					case ERigControlType::Integer:
@@ -5815,7 +5815,7 @@ void FControlRigEditor::HandleMakeElementGetterSetter(ERigElementGetterSetterTyp
 						{
 							StructTemplate = FRigUnit_SetIntAnimationChannel::StaticStruct();
 						}
-						PinsToResolve.Add(ValueName, FRigVMTemplateArgumentType(RigVMTypeUtils::Int32Type));
+						PinsToResolve.Add(ValueName, RigVMTypeUtils::TypeIndex::Int32);
 						break;
 					}
 					case ERigControlType::Vector2D:
@@ -5830,7 +5830,9 @@ void FControlRigEditor::HandleMakeElementGetterSetter(ERigElementGetterSetterTyp
 						}
 
 						UScriptStruct* ValueStruct = TBaseStructure<FVector2D>::Get();
-						PinsToResolve.Add(ValueName, FRigVMTemplateArgumentType(ValueStruct->GetStructCPPName(), ValueStruct));
+						const FRigVMTemplateArgumentType TypeForStruct(*ValueStruct->GetStructCPPName(), ValueStruct);
+						const int32 TypeIndex = FRigVMRegistry::Get().GetTypeIndex(TypeForStruct);
+						PinsToResolve.Add(ValueName, TypeIndex);
 						break;
 					}
 					case ERigControlType::Position:
@@ -5845,7 +5847,9 @@ void FControlRigEditor::HandleMakeElementGetterSetter(ERigElementGetterSetterTyp
 							StructTemplate = FRigUnit_SetVectorAnimationChannel::StaticStruct();
 						}
 						UScriptStruct* ValueStruct = TBaseStructure<FVector>::Get();
-						PinsToResolve.Add(ValueName, FRigVMTemplateArgumentType(ValueStruct->GetStructCPPName(), ValueStruct));
+						const FRigVMTemplateArgumentType TypeForStruct(*ValueStruct->GetStructCPPName(), ValueStruct);
+						const int32 TypeIndex = FRigVMRegistry::Get().GetTypeIndex(TypeForStruct);
+						PinsToResolve.Add(ValueName, TypeIndex);
 						break;
 					}
 					case ERigControlType::Rotator:
@@ -5859,7 +5863,9 @@ void FControlRigEditor::HandleMakeElementGetterSetter(ERigElementGetterSetterTyp
 							StructTemplate = FRigUnit_SetRotatorAnimationChannel::StaticStruct();
 						}
 						UScriptStruct* ValueStruct = TBaseStructure<FRotator>::Get();
-						PinsToResolve.Add(ValueName, FRigVMTemplateArgumentType(ValueStruct->GetStructCPPName(), ValueStruct));
+						const FRigVMTemplateArgumentType TypeForStruct(*ValueStruct->GetStructCPPName(), ValueStruct);
+						const int32 TypeIndex = FRigVMRegistry::Get().GetTypeIndex(TypeForStruct);
+						PinsToResolve.Add(ValueName, TypeIndex);
 						break;
 					}
 					case ERigControlType::Transform:
@@ -5875,7 +5881,9 @@ void FControlRigEditor::HandleMakeElementGetterSetter(ERigElementGetterSetterTyp
 							StructTemplate = FRigUnit_SetTransformAnimationChannel::StaticStruct();
 						}
 						UScriptStruct* ValueStruct = TBaseStructure<FTransform>::Get();
-						PinsToResolve.Add(ValueName, FRigVMTemplateArgumentType(ValueStruct->GetStructCPPName(), ValueStruct));
+						const FRigVMTemplateArgumentType TypeForStruct(*ValueStruct->GetStructCPPName(), ValueStruct);
+						const int32 TypeIndex = FRigVMRegistry::Get().GetTypeIndex(TypeForStruct);
+						PinsToResolve.Add(ValueName, TypeIndex);
 						break;
 					}
 					default:
@@ -6117,7 +6125,7 @@ void FControlRigEditor::HandleMakeElementGetterSetter(ERigElementGetterSetterTyp
 			NewNode.Name = ModelNode->GetFName();
 			NewNodes.Add(NewNode);
 
-			for (const TPair<FName, FRigVMTemplateArgumentType>& PinToResolve : PinsToResolve)
+			for (const TPair<FName, int32>& PinToResolve : PinsToResolve)
 			{
 				if(URigVMPin* Pin = ModelNode->FindPin(PinToResolve.Key.ToString()))
 				{
diff --git a/Engine/Plugins/Animation/ControlRig/Source/ControlRigEditor/Private/Editor/SRigHierarchyTreeView.cpp b/Engine/Plugins/Animation/ControlRig/Source/ControlRigEditor/Private/Editor/SRigHierarchyTreeView.cpp
--- a/Engine/Plugins/Animation/ControlRig/Source/ControlRigEditor/Private/Editor/SRigHierarchyTreeView.cpp
+++ b/Engine/Plugins/Animation/ControlRig/Source/ControlRigEditor/Private/Editor/SRigHierarchyTreeView.cpp
@@ -805,43 +805,43 @@ FLinearColor SRigHierarchyItem::GetColorForControlType(ERigControlType InControl
 	{
 		case ERigControlType::Bool:
 		{
-			PinType = RigVMTypeUtils::PinTypeFromCPPType(RigVMTypeUtils::BoolType, nullptr);
+			PinType = RigVMTypeUtils::PinTypeFromCPPType(RigVMTypeUtils::BoolTypeName, nullptr);
 			break;
 		}
 		case ERigControlType::Float:
 		{
-			PinType = RigVMTypeUtils::PinTypeFromCPPType(RigVMTypeUtils::FloatType, nullptr);
+			PinType = RigVMTypeUtils::PinTypeFromCPPType(RigVMTypeUtils::FloatTypeName, nullptr);
 			break;
 		}
 		case ERigControlType::Integer:
 		{
 			if(InControlEnum)
 			{
-				PinType = RigVMTypeUtils::PinTypeFromCPPType(FString(), InControlEnum);
+				PinType = RigVMTypeUtils::PinTypeFromCPPType(NAME_None, InControlEnum);
 			}
 			else
 			{
-				PinType = RigVMTypeUtils::PinTypeFromCPPType(RigVMTypeUtils::Int32Type, nullptr);
+				PinType = RigVMTypeUtils::PinTypeFromCPPType(RigVMTypeUtils::Int32TypeName, nullptr);
 			}
 			break;
 		}
 		case ERigControlType::Vector2D:
 		{
 			UScriptStruct* Struct = TBaseStructure<FVector2D>::Get(); 
-			PinType = RigVMTypeUtils::PinTypeFromCPPType(Struct->GetStructCPPName(), Struct);
+			PinType = RigVMTypeUtils::PinTypeFromCPPType(*Struct->GetStructCPPName(), Struct);
 			break;
 		}
 		case ERigControlType::Position:
 		case ERigControlType::Scale:
 		{
 			UScriptStruct* Struct = TBaseStructure<FVector>::Get(); 
-			PinType = RigVMTypeUtils::PinTypeFromCPPType(Struct->GetStructCPPName(), Struct);
+			PinType = RigVMTypeUtils::PinTypeFromCPPType(*Struct->GetStructCPPName(), Struct);
 			break;
 		}
 		case ERigControlType::Rotator:
 		{
 			UScriptStruct* Struct = TBaseStructure<FRotator>::Get(); 
-			PinType = RigVMTypeUtils::PinTypeFromCPPType(Struct->GetStructCPPName(), Struct);
+			PinType = RigVMTypeUtils::PinTypeFromCPPType(*Struct->GetStructCPPName(), Struct);
 			break;
 		}
 		case ERigControlType::Transform:
@@ -850,7 +850,7 @@ FLinearColor SRigHierarchyItem::GetColorForControlType(ERigControlType InControl
 		default:
 		{
 			UScriptStruct* Struct = TBaseStructure<FTransform>::Get(); 
-			PinType = RigVMTypeUtils::PinTypeFromCPPType(Struct->GetStructCPPName(), Struct);
+			PinType = RigVMTypeUtils::PinTypeFromCPPType(*Struct->GetStructCPPName(), Struct);
 			break;
 		}
 	}
diff --git a/Engine/Plugins/Animation/ControlRig/Source/ControlRigEditor/Private/Graph/NodeSpawners/ControlRigTemplateNodeSpawner.cpp b/Engine/Plugins/Animation/ControlRig/Source/ControlRigEditor/Private/Graph/NodeSpawners/ControlRigTemplateNodeSpawner.cpp
--- a/Engine/Plugins/Animation/ControlRig/Source/ControlRigEditor/Private/Graph/NodeSpawners/ControlRigTemplateNodeSpawner.cpp
+++ b/Engine/Plugins/Animation/ControlRig/Source/ControlRigEditor/Private/Graph/NodeSpawners/ControlRigTemplateNodeSpawner.cpp
@@ -96,11 +96,10 @@ bool UControlRigTemplateNodeSpawner::IsTemplateNodeFilteredOut(FBlueprintActionF
 							{
 								if (const FRigVMTemplate* Template = FRigVMRegistry::Get().FindTemplate(TemplateNotation))
 								{
-									FRigVMTemplateArgumentType Type(ModelPin->GetCPPType(), ModelPin->GetCPPTypeObject());
 									for (int32 i=0; i<Template->NumArguments(); ++i)
 									{
 										const FRigVMTemplateArgument* Argument = Template->GetArgument(i);
-										if (Template->ArgumentSupportsType(Argument->GetName(), Type.CPPType))
+										if (Template->ArgumentSupportsTypeIndex(Argument->GetName(), ModelPin->GetTypeIndex()))
 										{
 											return false;
 										}
diff --git a/Engine/Plugins/Animation/ControlRig/Source/ControlRigEditor/Private/Graph/SControlRigGraphChangePinType.cpp b/Engine/Plugins/Animation/ControlRig/Source/ControlRigEditor/Private/Graph/SControlRigGraphChangePinType.cpp
--- a/Engine/Plugins/Animation/ControlRig/Source/ControlRigEditor/Private/Graph/SControlRigGraphChangePinType.cpp
+++ b/Engine/Plugins/Animation/ControlRig/Source/ControlRigEditor/Private/Graph/SControlRigGraphChangePinType.cpp
@@ -64,7 +64,7 @@ FText SControlRigChangePinType::GetBindingText(const FRigVMTemplateArgumentType&
 	}
 	else
 	{
-		FString CPPType = InType.CPPType;
+		FString CPPType = InType.CPPType.ToString();
 		if(RigVMTypeUtils::IsArrayType(CPPType))
 		{
 			CPPType = RigVMTypeUtils::BaseTypeFromArrayType(CPPType);
@@ -192,8 +192,8 @@ void SControlRigChangePinType::FillPinTypeMenu(FMenuBuilder& MenuBuilder)
 		FRigVMTemplateArgumentType Type;
 	};
 
-	typedef TPair<FRigVMTemplateArgumentType, FArgumentInfo> FTypePair;
-	TMap<FRigVMTemplateArgumentType, FArgumentInfo> Types;
+	typedef TPair<int32, FArgumentInfo> FTypePair;
+	TMap<int32, FArgumentInfo> Types;
 	for(URigVMPin* ModelPin : ModelPins)
 	{
 		if(!ModelPin->IsRootPin())
@@ -208,29 +208,29 @@ void SControlRigChangePinType::FillPinTypeMenu(FMenuBuilder& MenuBuilder)
 				const FName ArgumentName = ModelPin->GetFName();
 				if(const FRigVMTemplateArgument* Argument = Template->FindArgument(ArgumentName))
 				{
-					const TArray<FRigVMTemplateArgumentType>& AllArgumentTypes = Argument->GetTypes();
-					const TArray<FRigVMTemplateArgumentType>& FilteredArgumentTypes = Argument->GetSupportedTypes(TemplateNode->GetFilteredPermutationsIndices());
+					const TArray<int32>& AllArgumentTypes = Argument->GetTypeIndices();
+					const TArray<int32>& FilteredArgumentTypes = Argument->GetSupportedTypeIndices(TemplateNode->GetFilteredPermutationsIndices());
 					
 					for(int32 PermutationIndex = 0; PermutationIndex < Template->NumPermutations(); PermutationIndex++)
 					{
-						FRigVMTemplateArgumentType ArgumentType = AllArgumentTypes[PermutationIndex];
+						int32 ArgumentTypeIndex = AllArgumentTypes[PermutationIndex];
 					
 						bool bIsFilteredOut = false;
-						if(!FilteredArgumentTypes.Contains(ArgumentType))
+						if(!FilteredArgumentTypes.Contains(ArgumentTypeIndex))
 						{
 							bIsFilteredOut = true;
 						}
 
-						if(ArgumentType.CPPType == RigVMTypeUtils::FloatType)
+						if(ArgumentTypeIndex == RigVMTypeUtils::TypeIndex::Float)
 						{
-							ArgumentType.CPPType = RigVMTypeUtils::DoubleType;
+							ArgumentTypeIndex = RigVMTypeUtils::TypeIndex::Double;
 						}
-						else if(ArgumentType.CPPType == RigVMTypeUtils::FloatArrayType)
+						else if(ArgumentTypeIndex == RigVMTypeUtils::TypeIndex::FloatArray)
 						{
-							ArgumentType.CPPType = RigVMTypeUtils::DoubleArrayType;
+							ArgumentTypeIndex = RigVMTypeUtils::TypeIndex::DoubleArray;
 						}
 						
-						if(Types.Contains(ArgumentType))
+						if(Types.Contains(ArgumentTypeIndex))
 						{
 							continue;
 						}
@@ -241,13 +241,14 @@ void SControlRigChangePinType::FillPinTypeMenu(FMenuBuilder& MenuBuilder)
 							{
 								if(const FProperty* Property = FunctionStruct->FindPropertyByName(ArgumentName))
 								{
-									Types.Add(ArgumentType, FArgumentInfo::Make(Property, bIsFilteredOut));
+									Types.Add(ArgumentTypeIndex, FArgumentInfo::Make(Property, bIsFilteredOut));
 								}
 							}
 						}
 						else
 						{
-							Types.Add(ArgumentType, FArgumentInfo::Make(ArgumentType, bIsFilteredOut));
+							const FRigVMTemplateArgumentType ArgumentType = FRigVMRegistry::Get().GetType(ArgumentTypeIndex);
+							Types.Add(ArgumentTypeIndex, FArgumentInfo::Make(ArgumentType, bIsFilteredOut));
 						}
 					}
 				}
@@ -277,8 +278,10 @@ void SControlRigChangePinType::FillPinTypeMenu(FMenuBuilder& MenuBuilder)
 	}
 	SortedTypes.Sort([](const FTypePair& A, const FTypePair& B) -> bool
 	{
-		const FString BaseTypeA = A.Key.GetBaseCPPType();
-		const FString BaseTypeB = B.Key.GetBaseCPPType();
+		const FRigVMTemplateArgumentType TypeA = FRigVMRegistry::Get().GetType(A.Key);
+		const FRigVMTemplateArgumentType TypeB = FRigVMRegistry::Get().GetType(B.Key);
+		const FString BaseTypeA = TypeA.GetBaseCPPType();
+		const FString BaseTypeB = TypeB.GetBaseCPPType();
 
 		const int32 IndexA = SortOrder.Find(BaseTypeA);  
 		const int32 IndexB = SortOrder.Find(BaseTypeB);
@@ -301,12 +304,13 @@ void SControlRigChangePinType::FillPinTypeMenu(FMenuBuilder& MenuBuilder)
 
 		const bool bHasAllTypes =
 			SortedTypes.Num() >=
-				FRigVMTemplateArgument::GetCompatibleTypes(FRigVMTemplateArgument::ETypeCategory_SingleAnyValue).Num();
+				FRigVMRegistry::Get().GetTypesForCategory(FRigVMRegistry::ETypeCategory_SingleAnyValue).Num();
 		
-		for(int32 TypeIndex=0; TypeIndex < SortedTypes.Num(); TypeIndex++)
+		for(int32 SortedIndex=0; SortedIndex < SortedTypes.Num(); SortedIndex++)
 		{
-			const FRigVMTemplateArgumentType& Type = SortedTypes[TypeIndex].Key;
-			const bool bIsFilteredOut = SortedTypes[TypeIndex].Value.bIsFilteredOut;
+			const int32& TypeIndex = SortedTypes[SortedIndex].Key;
+			const FRigVMTemplateArgumentType Type = FRigVMRegistry::Get().GetType(TypeIndex);
+			const bool bIsFilteredOut = SortedTypes[SortedIndex].Value.bIsFilteredOut;
 			if (Type.CPPTypeObject != nullptr && !IsValid(Type.CPPTypeObject))
 			{
 				continue;
@@ -323,8 +327,10 @@ void SControlRigChangePinType::FillPinTypeMenu(FMenuBuilder& MenuBuilder)
 
 				if(UScriptStruct* ScriptStruct = Cast<UScriptStruct>(Type.CPPTypeObject))
 				{
-					if(!FRigVMTemplateArgument::GetCompatibleTypes(FRigVMTemplateArgument::ETypeCategory_SingleMathStructValue).Contains(
-						FRigVMTemplateArgumentType(ScriptStruct->GetStructCPPName(), ScriptStruct)))
+					const FRigVMTemplateArgumentType MathType(*ScriptStruct->GetStructCPPName(), ScriptStruct);
+					const int32 MathTypeIndex = FRigVMRegistry::Get().GetTypeIndex(MathType);
+					
+					if(!FRigVMRegistry::Get().GetTypesForCategory(FRigVMRegistry::ETypeCategory_SingleMathStructValue).Contains(MathTypeIndex))
 					{
 						continue;
 					}
@@ -394,7 +400,8 @@ void SControlRigChangePinType::HandlePinTypeChanged(FRigVMTemplateArgumentType I
 			continue;
 		}
 
-		Controller->ResolveWildCardPin(ModelPin->GetPinPath(), InType.CPPType, InType.GetCPPTypeObjectPath(), true, true);
+		const int32 TypeIndex = FRigVMRegistry::Get().GetTypeIndex(InType);
+		Controller->ResolveWildCardPin(ModelPin->GetPinPath(), TypeIndex, true, true);
 	}
 }
 
diff --git a/Engine/Source/Developer/RigVMDeveloper/Private/RigVMDeveloperTypeUtils.cpp b/Engine/Source/Developer/RigVMDeveloper/Private/RigVMDeveloperTypeUtils.cpp
--- a/Engine/Source/Developer/RigVMDeveloper/Private/RigVMDeveloperTypeUtils.cpp
+++ b/Engine/Source/Developer/RigVMDeveloper/Private/RigVMDeveloperTypeUtils.cpp
@@ -267,43 +267,43 @@ FRigVMExternalVariable RigVMTypeUtils::ExternalVariableFromCPPType(const FName&
 	return Variable;
 }
 
-FEdGraphPinType RigVMTypeUtils::PinTypeFromCPPType(const FString& InCPPType, UObject* InCPPTypeObject)
+FEdGraphPinType RigVMTypeUtils::PinTypeFromCPPType(const FName& InCPPType, UObject* InCPPTypeObject)
 {
 	FEdGraphPinType PinType;
 	PinType.ResetToDefaults();
 	PinType.PinCategory = NAME_None;
 
-	FString BaseCPPType = InCPPType;
+	FName BaseCPPType = InCPPType;
 	PinType.ContainerType = EPinContainerType::None;
-	if (RigVMTypeUtils::IsArrayType(InCPPType))
+	if (RigVMTypeUtils::IsArrayType(InCPPType.ToString()))
 	{
-		BaseCPPType = RigVMTypeUtils::BaseTypeFromArrayType(InCPPType);
+		BaseCPPType = *RigVMTypeUtils::BaseTypeFromArrayType(InCPPType.ToString());
 		PinType.ContainerType = EPinContainerType::Array;
 	}
 
-	if (BaseCPPType == BoolType)
+	if (BaseCPPType == BoolTypeName)
 	{
 		PinType.PinCategory = UEdGraphSchema_K2::PC_Boolean;
 	}
-	else if (BaseCPPType == Int32Type)
+	else if (BaseCPPType == Int32TypeName)
 	{
 		PinType.PinCategory = UEdGraphSchema_K2::PC_Int;
 	}
-	else if (BaseCPPType == FloatType)
+	else if (BaseCPPType == FloatTypeName)
 	{
 		PinType.PinCategory = UEdGraphSchema_K2::PC_Real;
 		PinType.PinSubCategory = UEdGraphSchema_K2::PC_Float;
 	}
-	else if (BaseCPPType == DoubleType)
+	else if (BaseCPPType == DoubleTypeName)
 	{
 		PinType.PinCategory = UEdGraphSchema_K2::PC_Real;
 		PinType.PinSubCategory = UEdGraphSchema_K2::PC_Double;
 	}
-	else if (BaseCPPType == FNameType)
+	else if (BaseCPPType == FNameTypeName)
 	{
 		PinType.PinCategory = UEdGraphSchema_K2::PC_Name;
 	}
-	else if (BaseCPPType == FStringType)
+	else if (BaseCPPType == FStringTypeName)
 	{
 		PinType.PinCategory = UEdGraphSchema_K2::PC_String;
 	}
@@ -724,7 +724,7 @@ bool RigVMTypeUtils::CPPTypeFromExternalVariable(const FRigVMExternalVariable& I
 	return true;
 }
 
-bool RigVMTypeUtils::AreCompatible(const FString& InCPPTypeA, UObject* InCPPTypeObjectA, const FString& InCPPTypeB,	UObject* InCPPTypeObjectB)
+bool RigVMTypeUtils::AreCompatible(const FName& InCPPTypeA, UObject* InCPPTypeObjectA, const FName& InCPPTypeB,	UObject* InCPPTypeObjectB)
 {
 	return AreCompatible(PinTypeFromCPPType(InCPPTypeA, InCPPTypeObjectA), PinTypeFromCPPType(InCPPTypeB, InCPPTypeObjectB));
 }
diff --git a/Engine/Source/Developer/RigVMDeveloper/Private/RigVMModel/Nodes/RigVMArrayNode.cpp b/Engine/Source/Developer/RigVMDeveloper/Private/RigVMModel/Nodes/RigVMArrayNode.cpp
--- a/Engine/Source/Developer/RigVMDeveloper/Private/RigVMModel/Nodes/RigVMArrayNode.cpp
+++ b/Engine/Source/Developer/RigVMDeveloper/Private/RigVMModel/Nodes/RigVMArrayNode.cpp
@@ -265,6 +265,7 @@ const FRigVMTemplate* URigVMArrayNode::GetTemplate() const
 
 		if(TemplatesPerOpCode[OpCodeIndex])
 		{
+			CachedTemplate = TemplatesPerOpCode[OpCodeIndex]; 
 			return TemplatesPerOpCode[OpCodeIndex];
 		}
 		
@@ -274,10 +275,10 @@ const FRigVMTemplate* URigVMArrayNode::GetTemplate() const
 		{
 			Arguments.Reserve(8);
 			
-			static const TArray<FRigVMTemplateArgumentType>& SingleTypes = FRigVMTemplateArgument::GetCompatibleTypes(FRigVMTemplateArgument::ETypeCategory_SingleAnyValue);
-			static const TArray<FRigVMTemplateArgumentType>& ArrayTypes = FRigVMTemplateArgument::GetCompatibleTypes(FRigVMTemplateArgument::ETypeCategory_ArrayAnyValue);
+			static const TArray<int32>& SingleTypes = FRigVMRegistry::Get().GetTypesForCategory(FRigVMRegistry::ETypeCategory_SingleAnyValue);
+			static const TArray<int32>& ArrayTypes = FRigVMRegistry::Get().GetTypesForCategory(FRigVMRegistry::ETypeCategory_ArrayAnyValue);
 
-			static TArray<FRigVMTemplateArgumentType> ArrayValueTypes, ElementTypes;
+			static TArray<int32> ArrayValueTypes, ElementTypes;
 			if(ArrayValueTypes.IsEmpty())
 			{
 				ArrayValueTypes.Append(ArrayTypes);
@@ -285,73 +286,74 @@ const FRigVMTemplate* URigVMArrayNode::GetTemplate() const
 			}
 
 			UScriptStruct* ExecuteStruct = GetGraph()->GetExecuteContextStruct();
-			FRigVMTemplateArgumentType ExecuteType(ExecuteStruct->GetStructCPPName(), ExecuteStruct);
+			const FRigVMTemplateArgumentType ExecuteType(*ExecuteStruct->GetStructCPPName(), ExecuteStruct);
+			const int32 ExecuteTypeIndex = FRigVMRegistry::Get().FindOrAddType(ExecuteType);
 		
 			switch(OpCode)
 			{
 				case ERigVMOpCode::ArrayReset:
 				case ERigVMOpCode::ArrayReverse:
 				{
-					Arguments.Emplace(*FRigVMStruct::ExecuteContextName.ToString(), ERigVMPinDirection::IO, ExecuteType);
+					Arguments.Emplace(*FRigVMStruct::ExecuteContextName.ToString(), ERigVMPinDirection::IO, ExecuteTypeIndex);
 					Arguments.Emplace(*ArrayName, ERigVMPinDirection::IO, ArrayValueTypes);
 					break;
 				}
 				case ERigVMOpCode::ArrayGetNum:
 				{
 					Arguments.Emplace(*ArrayName, ERigVMPinDirection::Input, ArrayValueTypes);
-					Arguments.Emplace(*NumName, ERigVMPinDirection::Output, FRigVMTemplateArgumentType(RigVMTypeUtils::Int32Type, nullptr));
+					Arguments.Emplace(*NumName, ERigVMPinDirection::Output, RigVMTypeUtils::TypeIndex::Int32);
 					break;
 				} 
 				case ERigVMOpCode::ArraySetNum:
 				{
-					Arguments.Emplace(*FRigVMStruct::ExecuteContextName.ToString(), ERigVMPinDirection::IO, ExecuteType);
+					Arguments.Emplace(*FRigVMStruct::ExecuteContextName.ToString(), ERigVMPinDirection::IO, ExecuteTypeIndex);
 					Arguments.Emplace(*ArrayName, ERigVMPinDirection::IO, ArrayValueTypes);
-					Arguments.Emplace(*NumName, ERigVMPinDirection::Input, FRigVMTemplateArgumentType(RigVMTypeUtils::Int32Type, nullptr));
+					Arguments.Emplace(*NumName, ERigVMPinDirection::Input, RigVMTypeUtils::TypeIndex::Int32);
 					break;
 				}
 				case ERigVMOpCode::ArrayGetAtIndex:
 				{
 					Arguments.Emplace(*ArrayName, ERigVMPinDirection::Input, ArrayValueTypes);
-					Arguments.Emplace(*IndexName, ERigVMPinDirection::Input, FRigVMTemplateArgumentType(RigVMTypeUtils::Int32Type, nullptr));
+					Arguments.Emplace(*IndexName, ERigVMPinDirection::Input, RigVMTypeUtils::TypeIndex::Int32);
 					Arguments.Emplace(*ElementName, ERigVMPinDirection::Output, ElementTypes);
 					break;
 				}  
 				case ERigVMOpCode::ArraySetAtIndex:
 				case ERigVMOpCode::ArrayInsert:
 				{
-					Arguments.Emplace(*FRigVMStruct::ExecuteContextName.ToString(), ERigVMPinDirection::IO, ExecuteType);
+					Arguments.Emplace(*FRigVMStruct::ExecuteContextName.ToString(), ERigVMPinDirection::IO, ExecuteTypeIndex);
 					Arguments.Emplace(*ArrayName, ERigVMPinDirection::IO, ArrayValueTypes);
-					Arguments.Emplace(*IndexName, ERigVMPinDirection::Input, FRigVMTemplateArgumentType(RigVMTypeUtils::Int32Type, nullptr));
+					Arguments.Emplace(*IndexName, ERigVMPinDirection::Input, RigVMTypeUtils::TypeIndex::Int32);
 					Arguments.Emplace(*ElementName, ERigVMPinDirection::Input, ElementTypes);
 					break;
 				}
 				case ERigVMOpCode::ArrayAdd:
 				{
-					Arguments.Emplace(*FRigVMStruct::ExecuteContextName.ToString(), ERigVMPinDirection::IO, ExecuteType);
+					Arguments.Emplace(*FRigVMStruct::ExecuteContextName.ToString(), ERigVMPinDirection::IO, ExecuteTypeIndex);
 					Arguments.Emplace(*ArrayName, ERigVMPinDirection::IO, ArrayValueTypes);
 					Arguments.Emplace(*ElementName, ERigVMPinDirection::Input, ElementTypes);
-					Arguments.Emplace(*IndexName, ERigVMPinDirection::Output, FRigVMTemplateArgumentType(RigVMTypeUtils::Int32Type, nullptr));
+					Arguments.Emplace(*IndexName, ERigVMPinDirection::Output, RigVMTypeUtils::TypeIndex::Int32);
 					break;
 				}
 				case ERigVMOpCode::ArrayRemove:
 				{
-					Arguments.Emplace(*FRigVMStruct::ExecuteContextName.ToString(), ERigVMPinDirection::IO, ExecuteType);
+					Arguments.Emplace(*FRigVMStruct::ExecuteContextName.ToString(), ERigVMPinDirection::IO, ExecuteTypeIndex);
 					Arguments.Emplace(*ArrayName, ERigVMPinDirection::IO, ArrayValueTypes);
-					Arguments.Emplace(*IndexName, ERigVMPinDirection::Input, FRigVMTemplateArgumentType(RigVMTypeUtils::Int32Type, nullptr));
+					Arguments.Emplace(*IndexName, ERigVMPinDirection::Input, RigVMTypeUtils::TypeIndex::Int32);
 					break;
 				}
 				case ERigVMOpCode::ArrayFind:
 				{
 					Arguments.Emplace(*ArrayName, ERigVMPinDirection::Input, ArrayValueTypes);
 					Arguments.Emplace(*ElementName, ERigVMPinDirection::Input, ElementTypes);
-					Arguments.Emplace(*IndexName, ERigVMPinDirection::Output, FRigVMTemplateArgumentType(RigVMTypeUtils::Int32Type, nullptr));
-					Arguments.Emplace(*SuccessName, ERigVMPinDirection::Output, FRigVMTemplateArgumentType(RigVMTypeUtils::BoolType, nullptr));
+					Arguments.Emplace(*IndexName, ERigVMPinDirection::Output, RigVMTypeUtils::TypeIndex::Int32);
+					Arguments.Emplace(*SuccessName, ERigVMPinDirection::Output, RigVMTypeUtils::TypeIndex::Bool);
 					break;
 				}
 				case ERigVMOpCode::ArrayAppend:
 				case ERigVMOpCode::ArrayUnion:
 				{
-					Arguments.Emplace(*FRigVMStruct::ExecuteContextName.ToString(), ERigVMPinDirection::IO, ExecuteType);
+					Arguments.Emplace(*FRigVMStruct::ExecuteContextName.ToString(), ERigVMPinDirection::IO, ExecuteTypeIndex);
 					Arguments.Emplace(*ArrayName, ERigVMPinDirection::IO, ArrayValueTypes);
 					Arguments.Emplace(*OtherName, ERigVMPinDirection::Input, ArrayValueTypes);
 					break;
@@ -364,14 +366,14 @@ const FRigVMTemplate* URigVMArrayNode::GetTemplate() const
 				}
 				case ERigVMOpCode::ArrayIterator:
 				{
-					Arguments.Emplace(*FRigVMStruct::ExecuteContextName.ToString(), ERigVMPinDirection::IO, ExecuteType);
+					Arguments.Emplace(*FRigVMStruct::ExecuteContextName.ToString(), ERigVMPinDirection::IO, ExecuteTypeIndex);
 					Arguments.Emplace(*ArrayName, ERigVMPinDirection::Input, ArrayValueTypes);
 					Arguments.Emplace(*ElementName, ERigVMPinDirection::Output, ElementTypes);
-					Arguments.Emplace(*IndexName, ERigVMPinDirection::Output, FRigVMTemplateArgumentType(RigVMTypeUtils::Int32Type, nullptr));
-					Arguments.Emplace(*CountName, ERigVMPinDirection::Output, FRigVMTemplateArgumentType(RigVMTypeUtils::Int32Type, nullptr));
-					Arguments.Emplace(*RatioName, ERigVMPinDirection::Output, FRigVMTemplateArgumentType(RigVMTypeUtils::FloatType, nullptr));
-					Arguments.Emplace(*ContinueName, ERigVMPinDirection::Hidden, ExecuteType);
-					Arguments.Emplace(*CompletedName, ERigVMPinDirection::Output, ExecuteType);
+					Arguments.Emplace(*IndexName, ERigVMPinDirection::Output, RigVMTypeUtils::TypeIndex::Int32);
+					Arguments.Emplace(*CountName, ERigVMPinDirection::Output, RigVMTypeUtils::TypeIndex::Int32);
+					Arguments.Emplace(*RatioName, ERigVMPinDirection::Output, RigVMTypeUtils::TypeIndex::Float);
+					Arguments.Emplace(*ContinueName, ERigVMPinDirection::Hidden, ExecuteTypeIndex);
+					Arguments.Emplace(*CompletedName, ERigVMPinDirection::Output, ExecuteTypeIndex);
 					break;
 				}
 				case ERigVMOpCode::ArrayDifference:
diff --git a/Engine/Source/Developer/RigVMDeveloper/Private/RigVMModel/Nodes/RigVMIfNode.cpp b/Engine/Source/Developer/RigVMDeveloper/Private/RigVMModel/Nodes/RigVMIfNode.cpp
--- a/Engine/Source/Developer/RigVMDeveloper/Private/RigVMModel/Nodes/RigVMIfNode.cpp
+++ b/Engine/Source/Developer/RigVMDeveloper/Private/RigVMModel/Nodes/RigVMIfNode.cpp
@@ -34,16 +34,15 @@ const FRigVMTemplate* URigVMIfNode::GetTemplate() const
 		static TArray<FRigVMTemplateArgument> Arguments;
 		if(Arguments.IsEmpty())
 		{
-			static TArray<FRigVMTemplateArgumentType> Types;
+			static TArray<int32> Types;
 			if(Types.IsEmpty())
 			{
-				Types.Reserve(3);
-				Types.Append(FRigVMTemplateArgument::GetCompatibleTypes(FRigVMTemplateArgument::ETypeCategory_SingleAnyValue));
-				Types.Append(FRigVMTemplateArgument::GetCompatibleTypes(FRigVMTemplateArgument::ETypeCategory_ArrayAnyValue));
+				Types.Append(FRigVMRegistry::Get().GetTypesForCategory(FRigVMRegistry::ETypeCategory_SingleAnyValue));
+				Types.Append(FRigVMRegistry::Get().GetTypesForCategory(FRigVMRegistry::ETypeCategory_ArrayAnyValue));
 			}
 
 			Arguments.Reserve(4);
-			Arguments.Emplace(*ConditionName, ERigVMPinDirection::Input, FRigVMTemplateArgumentType(RigVMTypeUtils::BoolType, nullptr));
+			Arguments.Emplace(*ConditionName, ERigVMPinDirection::Input, RigVMTypeUtils::TypeIndex::Bool);
 			Arguments.Emplace(*TrueName, ERigVMPinDirection::Input, Types);
 			Arguments.Emplace(*FalseName, ERigVMPinDirection::Input, Types);
 			Arguments.Emplace(*ResultName, ERigVMPinDirection::Output, Types);
diff --git a/Engine/Source/Developer/RigVMDeveloper/Private/RigVMModel/Nodes/RigVMRerouteNode.cpp b/Engine/Source/Developer/RigVMDeveloper/Private/RigVMModel/Nodes/RigVMRerouteNode.cpp
--- a/Engine/Source/Developer/RigVMDeveloper/Private/RigVMModel/Nodes/RigVMRerouteNode.cpp
+++ b/Engine/Source/Developer/RigVMDeveloper/Private/RigVMModel/Nodes/RigVMRerouteNode.cpp
@@ -89,14 +89,16 @@ const FRigVMTemplate* URigVMRerouteNode::GetTemplate() const
 		static TArray<FRigVMTemplateArgument> Arguments;
 		if(Arguments.IsEmpty())
 		{
-			static TArray<FRigVMTemplateArgumentType> Types;
+			static TArray<int32> Types;
 			if(Types.IsEmpty())
 			{
-				Types.Append(FRigVMTemplateArgument::GetCompatibleTypes(FRigVMTemplateArgument::ETypeCategory_SingleAnyValue));
-				Types.Append(FRigVMTemplateArgument::GetCompatibleTypes(FRigVMTemplateArgument::ETypeCategory_ArrayAnyValue));
+				Types.Append(FRigVMRegistry::Get().GetTypesForCategory(FRigVMRegistry::ETypeCategory_SingleAnyValue));
+				Types.Append(FRigVMRegistry::Get().GetTypesForCategory(FRigVMRegistry::ETypeCategory_ArrayAnyValue));
 
 				UScriptStruct* ExecuteStruct = GetGraph()->GetExecuteContextStruct();
-				Types.Add(FRigVMTemplateArgumentType(ExecuteStruct->GetStructCPPName(), ExecuteStruct));
+				const FRigVMTemplateArgumentType ExecuteType(*ExecuteStruct->GetStructCPPName(), ExecuteStruct);
+				const int32 ExecuteTypeIndex = FRigVMRegistry::Get().FindOrAddType(ExecuteType);
+				Types.Add(ExecuteTypeIndex);
 			}
 		
 			Arguments.Emplace(TEXT("Value"), ERigVMPinDirection::IO, Types);
diff --git a/Engine/Source/Developer/RigVMDeveloper/Private/RigVMModel/Nodes/RigVMSelectNode.cpp b/Engine/Source/Developer/RigVMDeveloper/Private/RigVMModel/Nodes/RigVMSelectNode.cpp
--- a/Engine/Source/Developer/RigVMDeveloper/Private/RigVMModel/Nodes/RigVMSelectNode.cpp
+++ b/Engine/Source/Developer/RigVMDeveloper/Private/RigVMModel/Nodes/RigVMSelectNode.cpp
@@ -45,11 +45,11 @@ const FRigVMTemplate* URigVMSelectNode::GetTemplate() const
 		static TArray<FRigVMTemplateArgument> Arguments;
 		if(Arguments.IsEmpty())
 		{
-			static const TArray<FRigVMTemplateArgumentType>& SingleTypes = FRigVMTemplateArgument::GetCompatibleTypes(FRigVMTemplateArgument::ETypeCategory_SingleAnyValue);
-			static const TArray<FRigVMTemplateArgumentType>& ArrayTypes = FRigVMTemplateArgument::GetCompatibleTypes(FRigVMTemplateArgument::ETypeCategory_ArrayAnyValue);
-			static const TArray<FRigVMTemplateArgumentType>& ArrayArrayTypes = FRigVMTemplateArgument::GetCompatibleTypes(FRigVMTemplateArgument::ETypeCategory_ArrayArrayAnyValue);
+			static const TArray<int32>& SingleTypes = FRigVMRegistry::Get().GetTypesForCategory(FRigVMRegistry::ETypeCategory_SingleAnyValue);
+			static const TArray<int32>& ArrayTypes = FRigVMRegistry::Get().GetTypesForCategory(FRigVMRegistry::ETypeCategory_ArrayAnyValue);
+			static const TArray<int32>& ArrayArrayTypes = FRigVMRegistry::Get().GetTypesForCategory(FRigVMRegistry::ETypeCategory_ArrayArrayAnyValue);
 
-			static TArray<FRigVMTemplateArgumentType> ResultTypes, ValueTypes;
+			static TArray<int32> ResultTypes, ValueTypes;
 			if(ResultTypes.IsEmpty())
 			{
 				ResultTypes.Append(SingleTypes);
@@ -59,7 +59,7 @@ const FRigVMTemplate* URigVMSelectNode::GetTemplate() const
 			}
 
 			Arguments.Reserve(3);
-			Arguments.Emplace(*IndexName, ERigVMPinDirection::Input, FRigVMTemplateArgumentType(RigVMTypeUtils::Int32Type, nullptr));
+			Arguments.Emplace(*IndexName, ERigVMPinDirection::Input, RigVMTypeUtils::TypeIndex::Int32);
 			Arguments.Emplace(*ValueName, ERigVMPinDirection::Input, ValueTypes);
 			Arguments.Emplace(*ResultName, ERigVMPinDirection::Output, ResultTypes);
 		}
diff --git a/Engine/Source/Developer/RigVMDeveloper/Private/RigVMModel/Nodes/RigVMTemplateNode.cpp b/Engine/Source/Developer/RigVMDeveloper/Private/RigVMModel/Nodes/RigVMTemplateNode.cpp
--- a/Engine/Source/Developer/RigVMDeveloper/Private/RigVMModel/Nodes/RigVMTemplateNode.cpp
+++ b/Engine/Source/Developer/RigVMDeveloper/Private/RigVMModel/Nodes/RigVMTemplateNode.cpp
@@ -3,6 +3,41 @@
 #include "RigVMModel/Nodes/RigVMTemplateNode.h"
 #include "RigVMCore/RigVMTemplate.h"
 #include "RigVMModel/RigVMController.h"
+#include "UObject/ObjectSaveContext.h"
+
+////////////////////////////////////////////////////////////////////////////////////////////////
+
+void FRigVMTemplatePreferredType::UpdateStringFromIndex()
+{
+	static constexpr TCHAR Format[] = TEXT("%s,%s");
+	const FRigVMTemplateArgumentType Type = FRigVMRegistry::Get().GetType(TypeIndex);
+	TypeString = FString::Printf(Format, *Type.CPPType.ToString(), *Type.GetCPPTypeObjectPath().ToString());
+}
+
+void FRigVMTemplatePreferredType::UpdateIndexFromString()
+{
+	// during duplicate we may not have type strings here
+	if(TypeString.IsEmpty())
+	{
+		return;
+	}
+
+	FString Left, Right;
+	verify(TypeString.Split(TEXT(","), &Left, &Right));
+
+	UObject* CPPTypeObject = nullptr;
+
+	static const FString NoneString = FName(NAME_None).ToString(); 
+	if(Right != NoneString)
+	{
+		CPPTypeObject = URigVMPin::FindObjectFromCPPTypeObjectPath(Right);
+	}
+
+	const FRigVMTemplateArgumentType Type(*Left, CPPTypeObject);
+	TypeIndex = FRigVMRegistry::Get().FindOrAddType(Type);
+}
+
+////////////////////////////////////////////////////////////////////////////////////////////////
 
 URigVMTemplateNode::URigVMTemplateNode()
 	: Super()
@@ -13,6 +48,12 @@ URigVMTemplateNode::URigVMTemplateNode()
 {
 }
 
+void URigVMTemplateNode::PreSave(FObjectPreSaveContext SaveContext)
+{
+	Super::PreSave(SaveContext);
+	ConvertPreferredTypesToString();
+}
+
 void URigVMTemplateNode::PostLoad()
 {
 	Super::PostLoad();
@@ -24,10 +65,48 @@ void URigVMTemplateNode::PostLoad()
 	{
 		TemplateNotation = *SanitizedNotation;
 	}
-	
+
+	// upgrade from a previous version where we stored the preferred types as strings
+	if(!PreferredPermutationTypes_DEPRECATED.IsEmpty())
+	{
+		const FRigVMRegistry& Registry = FRigVMRegistry::Get();
+
+		PreferredPermutationPairs.Reserve(PreferredPermutationTypes_DEPRECATED.Num());
+
+		for(const FString& PreferredPermutation : PreferredPermutationTypes_DEPRECATED)
+		{
+			FString ArgName, CPPType;
+			PreferredPermutation.Split(TEXT(":"), &ArgName, &CPPType);
+
+			const FRigVMTemplateArgumentType Type = Registry.FindTypeFromCPPType(CPPType);
+			PreferredPermutationPairs.Emplace(*ArgName, Registry.GetTypeIndex(Type));
+		}
+
+		PreferredPermutationTypes_DEPRECATED.Reset();
+	}
+
+	ConvertPreferredTypesToTypeIndex();
 	InvalidateCache();
 }
 
+void URigVMTemplateNode::ConvertPreferredTypesToString()
+{
+	// we rely on the strings being serialized
+	for(FRigVMTemplatePreferredType& PreferredType : PreferredPermutationPairs)
+	{
+		PreferredType.UpdateStringFromIndex();
+	}
+}
+
+void URigVMTemplateNode::ConvertPreferredTypesToTypeIndex()
+{
+	// we rely on the strings being serialized - so on load we need to update the type index again 
+	for(FRigVMTemplatePreferredType& PreferredType : PreferredPermutationPairs)
+	{
+		PreferredType.UpdateIndexFromString();
+	}
+}
+
 UScriptStruct* URigVMTemplateNode::GetScriptStruct() const
 {
 	if(const FRigVMFunction* Function = GetResolvedFunction())
@@ -92,7 +171,7 @@ FText URigVMTemplateNode::GetToolTipTextForPin(const URigVMPin* InPin) const
 				const TArray<int32> PermutationIndices = GetFilteredPermutationsIndices();
 				if(PermutationIndices.Num() == GetTemplate()->NumPermutations())
 				{
-					if(Arg->GetTypes().Num() > 100)
+					if(Arg->GetTypeIndices().Num() > 100)
 					{
 						Tooltip = TEXT("Supports any type.");
 					}
@@ -103,7 +182,8 @@ FText URigVMTemplateNode::GetToolTipTextForPin(const URigVMPin* InPin) const
 					FString SupportedTypesJoined;
 					for (int32 Index : PermutationIndices)
 					{
-						FString Type = Arg->GetTypes()[Index].CPPType;
+						const int32 TypeIndex = Arg->GetTypeIndices()[Index];
+						FString Type = FRigVMRegistry::Get().GetType(TypeIndex).CPPType.ToString();
 						if (!FilteredPermutations.Contains(Index))
 						{
 							Type += TEXT(" : Breaks Connections");
@@ -135,22 +215,23 @@ bool URigVMTemplateNode::IsSingleton() const
 	return GetTemplate() == nullptr;
 }
 
-bool URigVMTemplateNode::SupportsType(const URigVMPin* InPin, const FString& InCPPType, FString* OutCPPType)
+bool URigVMTemplateNode::SupportsType(const URigVMPin* InPin, int32 InTypeIndex, int32* OutTypeIndex)
 {
+	const FRigVMRegistry& Registry = FRigVMRegistry::Get();
 	static const FString WildCardCPPType = RigVMTypeUtils::GetWildCardCPPType();
 	static const FString WildCardArrayCPPType = RigVMTypeUtils::ArrayTypeFromBaseType(WildCardCPPType);
 
 	const URigVMPin* RootPin = InPin->GetRootPin();
 
-	FString CPPType = InCPPType;
+	int32 TypeIndex = InTypeIndex;
 
 	if(InPin->GetParentPin() == RootPin && RootPin->IsArray())
 	{
-		CPPType = RigVMTypeUtils::ArrayTypeFromBaseType(CPPType);
+		TypeIndex = Registry.GetArrayTypeFromBaseTypeIndex(TypeIndex);
 	}
 
 	// we always support the unknown type
-	if((InCPPType == WildCardCPPType) || (InCPPType == WildCardArrayCPPType))
+	if(Registry.IsWildCardType(TypeIndex))
 	{
 		if(const FRigVMTemplate* Template = GetTemplate())
 		{
@@ -162,7 +243,7 @@ bool URigVMTemplateNode::SupportsType(const URigVMPin* InPin, const FString& InC
 					return false;
 				}
 
-				if(RigVMTypeUtils::IsArrayType(CPPType))
+				if(Registry.IsArrayType(TypeIndex))
 				{
 					if(Argument->GetArrayType() == FRigVMTemplateArgument::EArrayType_SingleValue)
 					{
@@ -177,9 +258,9 @@ bool URigVMTemplateNode::SupportsType(const URigVMPin* InPin, const FString& InC
 					}
 				}
 				
-				if(OutCPPType)
+				if(OutTypeIndex)
 				{
-					*OutCPPType = InCPPType;
+					*OutTypeIndex = InTypeIndex;
 				}
 				return true;
 			}
@@ -189,44 +270,44 @@ bool URigVMTemplateNode::SupportsType(const URigVMPin* InPin, const FString& InC
 	
 	if (const FRigVMTemplate* Template = GetTemplate())
 	{
-		const FString CacheKey = RootPin->GetName() + TEXT("|") + CPPType;
-		if (const TPair<bool, FRigVMTemplateArgumentType>* CachedResult = SupportedTypesCache.Find(CacheKey))
+		const TPair<FName,int32> CacheKey(RootPin->GetFName(), TypeIndex);
+		if (const TPair<bool, int32>* CachedResult = SupportedTypesCache.Find(CacheKey))
 		{
-			if(OutCPPType)
+			if(OutTypeIndex)
 			{
-				*OutCPPType = CachedResult->Value.CPPType;
+				*OutTypeIndex = CachedResult->Value;
 			}
 			return CachedResult->Key;
 		}
 
-		FRigVMTemplateArgumentType OutType;
-		if (Template->ArgumentSupportsType(RootPin->GetFName(), CPPType, &OutType))
+		int32 SupportedTypeIndex = INDEX_NONE;
+		if (Template->ArgumentSupportsTypeIndex(RootPin->GetFName(), TypeIndex, &SupportedTypeIndex))
 		{
-			if (OutCPPType)
+			if (OutTypeIndex)
 			{
-				(*OutCPPType) = OutType.CPPType;
+				(*OutTypeIndex) = SupportedTypeIndex;
 			}
 			return true;
 		}
 		return false;
 	}
 	
-	if(RootPin->GetCPPType() == CPPType)
+	if(RootPin->GetTypeIndex() == TypeIndex)
 	{
-		if(OutCPPType)
+		if(OutTypeIndex)
 		{
-			*OutCPPType = CPPType;
+			*OutTypeIndex = TypeIndex;
 		}
 		return true;
 	}
 	return false;
 }
 
-bool URigVMTemplateNode::FilteredSupportsType(const URigVMPin* InPin, const FString& InCPPType, FString* OutCPPType, bool bAllowFloatingPointCasts)
+bool URigVMTemplateNode::FilteredSupportsType(const URigVMPin* InPin, int32 InTypeIndex, int32* OutTypeIndex, bool bAllowFloatingPointCasts)
 {
-	if (OutCPPType)
+	if (OutTypeIndex)
 	{
-		*OutCPPType = FString(); 
+		*OutTypeIndex = INDEX_NONE; 
 	}
 
 	const URigVMPin* RootPin = InPin;
@@ -247,7 +328,15 @@ bool URigVMTemplateNode::FilteredSupportsType(const URigVMPin* InPin, const FStr
 
 	if (bIsStructElement)
 	{
-		return InPin->GetCPPType() == InCPPType;
+		const bool bResult = InPin->GetTypeIndex() == InTypeIndex;
+		if(bResult)
+		{
+			if(OutTypeIndex)
+			{
+				*OutTypeIndex = InTypeIndex; 
+			}
+		}
+		return bResult;
 	}
 
 	const FRigVMTemplateArgument* Argument = GetTemplate()->FindArgument(RootPin->GetFName());
@@ -256,22 +345,24 @@ bool URigVMTemplateNode::FilteredSupportsType(const URigVMPin* InPin, const FStr
 		return false;
 	}
 
-	FString RootCPPType = InCPPType;
+	const FRigVMRegistry& Registry = FRigVMRegistry::Get();
+
+	int32 RootTypeIndex = InTypeIndex;
 	if (bIsArrayElement)
 	{
-		RootCPPType = RigVMTypeUtils::ArrayTypeFromBaseType(InCPPType);
+		RootTypeIndex = Registry.GetArrayTypeFromBaseTypeIndex(RootTypeIndex);
 	}
 
 	if (FilteredPermutations.Num() == GetTemplate()->NumPermutations())
 	{
-		return GetTemplate()->ArgumentSupportsType(RootPin->GetFName(), InCPPType);
+		return GetTemplate()->ArgumentSupportsTypeIndex(RootPin->GetFName(), InTypeIndex, OutTypeIndex);
 	}
 
-	const TArray<FRigVMTemplateArgumentType>& Types = Argument->GetTypes();
-	for (int32 PermutationIndex : FilteredPermutations)
+	const TArray<int32>& TypeIndices = Argument->GetTypeIndices();
+	for (const int32& PermutationIndex : FilteredPermutations)
 	{
-		const FRigVMTemplateArgumentType& FilteredType = Types[PermutationIndex];
-		if (FilteredType.Matches(RootCPPType, bAllowFloatingPointCasts))
+		const int32& FilteredTypeIndex = TypeIndices[PermutationIndex];
+		if (Registry.CanMatchTypes(FilteredTypeIndex, RootTypeIndex, bAllowFloatingPointCasts))
 		{
 			return true;
 		}
@@ -355,25 +446,26 @@ FString URigVMTemplateNode::GetInitialDefaultValueForPin(const FName& InRootPinN
 
 		const FRigVMTemplate* Template = GetTemplate();
 		const FRigVMTemplateArgument* Argument = Template->FindArgument(InRootPinName);
-		const FRigVMTemplateArgumentType Type = Argument->GetTypes()[PermutationIndex];
+		const int32 TypeIndex = Argument->GetTypeIndices()[PermutationIndex];
 
 		if(const FRigVMFunction* Permutation = Template->GetPermutation(PermutationIndex))
 		{
 			const TSharedPtr<FStructOnScope> StructOnScope = MakeShareable(new FStructOnScope(Permutation->Struct));
 			const FRigVMStruct* DefaultStruct = (const FRigVMStruct*)StructOnScope->GetStructMemory();
 
-			const bool bUseQuotes = Type.CPPType != RigVMTypeUtils::FStringType && Type.CPPType != RigVMTypeUtils::FNameType;
+			const bool bUseQuotes = TypeIndex != RigVMTypeUtils::TypeIndex::FString && TypeIndex != RigVMTypeUtils::TypeIndex::FName;
 			NewDefaultValue = DefaultStruct->ExportToFullyQualifiedText(
 				Cast<UScriptStruct>(StructOnScope->GetStruct()), InRootPinName, nullptr, bUseQuotes);
 		}
 		else
 		{
-			if (Type.IsArray())
+			if (FRigVMRegistry::Get().IsArrayType(TypeIndex))
 			{
 				NewDefaultValue = TEXT("()");
 			}
 			else
 			{
+				const FRigVMTemplateArgumentType& Type = FRigVMRegistry::Get().GetType(TypeIndex);
 				if (UScriptStruct* ScriptStruct = Cast<UScriptStruct>(Type.CPPTypeObject))
 				{
 					TArray<uint8, TAlignedHeapAllocator<16>> TempBuffer;
@@ -385,36 +477,36 @@ FString URigVMTemplateNode::GetInitialDefaultValueForPin(const FName& InRootPinN
 					ScriptStruct->ExportText(NewDefaultValue, TempBuffer.GetData(), nullptr, nullptr, PPF_None, nullptr);
 					ScriptStruct->DestroyStruct(TempBuffer.GetData());				
 				}
-				else if (UEnum* Enum = Cast<UEnum>(Type.CPPTypeObject))
+				else if (const UEnum* Enum = Cast<UEnum>(Type.CPPTypeObject))
 				{
 					NewDefaultValue = Enum->GetNameStringByValue(0);
 				}
-				else if(UClass* Class = Cast<UClass>(Type.CPPTypeObject))
+				else if(const UClass* Class = Cast<UClass>(Type.CPPTypeObject))
 				{
 					// not supporting objects yet
 					ensure(false);
 				}
-				else if (Type.CPPType == RigVMTypeUtils::FloatType)
+				else if (Type.CPPType == RigVMTypeUtils::FloatTypeName)
 				{
 					NewDefaultValue = TEXT("0.000000");				
 				}
-				else if (Type.CPPType == RigVMTypeUtils::DoubleType)
+				else if (Type.CPPType == RigVMTypeUtils::DoubleTypeName)
 				{
 					NewDefaultValue = TEXT("0.000000");
 				}
-				else if (Type.CPPType == RigVMTypeUtils::Int32Type)
+				else if (Type.CPPType == RigVMTypeUtils::Int32TypeName)
 				{
 					NewDefaultValue = TEXT("0");
 				}
-				else if (Type.CPPType == RigVMTypeUtils::BoolType)
+				else if (Type.CPPType == RigVMTypeUtils::BoolTypeName)
 				{
 					NewDefaultValue = TEXT("False");
 				}
-				else if (Type.CPPType == RigVMTypeUtils::FStringType)
+				else if (Type.CPPType == RigVMTypeUtils::FStringTypeName)
 				{
 					NewDefaultValue = TEXT("");
 				}
-				else if (Type.CPPType == RigVMTypeUtils::FNameType)
+				else if (Type.CPPType == RigVMTypeUtils::FNameTypeName)
 				{
 					NewDefaultValue = TEXT("");
 				}
@@ -471,37 +563,35 @@ const TArray<int32>& URigVMTemplateNode::GetFilteredPermutationsIndices() const
 	return FilteredPermutations;
 }
 
-TArray<FRigVMTemplateArgumentType> URigVMTemplateNode::GetFilteredTypesForPin(URigVMPin* InPin) const
+TArray<int32> URigVMTemplateNode::GetFilteredTypesForPin(URigVMPin* InPin) const
 {
 	ensureMsgf(InPin->GetNode() == this, TEXT("GetFilteredTypesForPin of %s with pin from another node %s"), *GetNodePath(), *InPin->GetPinPath(true));
-	TArray<FRigVMTemplateArgumentType> FilteredTypes;
+	TArray<int32> FilteredTypes;
 
 	if (FilteredPermutations.IsEmpty())
 	{
 		return FilteredTypes;	
 	}
 
-	if (!PreferredPermutationTypes.IsEmpty())
+	if (!PreferredPermutationPairs.IsEmpty())
 	{
-		for (const FString& PreferredType : PreferredPermutationTypes)
+		for (const FRigVMTemplatePreferredType& PreferredType : PreferredPermutationPairs)
 		{
-			FString ArgName, CPPType;
-			PreferredType.Split(TEXT(":"), &ArgName, &CPPType);
-			if (InPin->GetName() == ArgName)
+			if (InPin->GetFName() == PreferredType.Argument)
 			{
-				const FRigVMTemplateArgument* Argument = GetTemplate()->FindArgument(*ArgName);
-				for (const FRigVMTemplateArgumentType& Type : Argument->GetTypes())
+				const FRigVMTemplateArgument* Argument = GetTemplate()->FindArgument(PreferredType.Argument);
+				for (const int32& TypeIndex : Argument->GetTypeIndices())
 				{
-					if (Type.CPPType == CPPType)
+					if (TypeIndex == PreferredType.TypeIndex)
 					{
-						return {Type};						
+						return {TypeIndex};						
 					}
 				}				
 			}
 		}
 	}
 	
-	URigVMPin* RootPin = InPin;
+	const URigVMPin* RootPin = InPin;
 	bool bIsArrayElement = false;
 	if (URigVMPin* ParentPin = InPin->GetParentPin())
 	{
@@ -510,12 +600,12 @@ TArray<FRigVMTemplateArgumentType> URigVMTemplateNode::GetFilteredTypesForPin(UR
 	}
 	if (const FRigVMTemplateArgument* Argument = GetTemplate()->FindArgument(RootPin->GetFName()))
 	{
-		FilteredTypes = Argument->GetSupportedTypes(FilteredPermutations);
+		FilteredTypes = Argument->GetSupportedTypeIndices(FilteredPermutations);
 		if (bIsArrayElement)
 		{
-			for (FRigVMTemplateArgumentType& ArrayType : FilteredTypes)
+			for (int32& ArrayType : FilteredTypes)
 			{
-				ArrayType.ConvertToBaseElement();
+				ArrayType = FRigVMRegistry::Get().GetBaseTypeFromArrayTypeIndex(ArrayType);
 			}
 		}
 	}
@@ -544,7 +634,7 @@ TArray<int32> URigVMTemplateNode::GetNewFilteredPermutations(URigVMPin* InPin, U
 
 	if (bIsStructElement)
 	{
-		if (RigVMTypeUtils::AreCompatible(InPin->GetCPPType(), InPin->GetCPPTypeObject(), LinkedPin->GetCPPType(), LinkedPin->GetCPPTypeObject()))
+		if (FRigVMRegistry::Get().CanMatchTypes(InPin->GetTypeIndex(), LinkedPin->GetTypeIndex(), true))
 		{
 			return FilteredPermutations;
 		}
@@ -553,7 +643,7 @@ TArray<int32> URigVMTemplateNode::GetNewFilteredPermutations(URigVMPin* InPin, U
 	TArray<int32> PermutationsToTry = FilteredPermutations;
 
 	// Reduce permutations to the ones respecting the preferred types
-	TArray<int32> PreferredPermutations = FindPermutationsForTypes(PreferredPermutationTypes, true);
+	const TArray<int32> PreferredPermutations = FindPermutationsForTypes(PreferredPermutationPairs, true);
 	PermutationsToTry = PermutationsToTry.FilterByPredicate([&](const int32& OtherPermutation) { return PreferredPermutations.Contains(OtherPermutation); });
 	
 	bool bLinkedIsTemplate = false;
@@ -566,12 +656,12 @@ TArray<int32> URigVMTemplateNode::GetNewFilteredPermutations(URigVMPin* InPin, U
 			{
 				for(int32 PermutationIndex : PermutationsToTry)
 				{
-					FRigVMTemplateArgumentType Type = Argument->GetTypes()[PermutationIndex];
+					int32 TypeIndex = Argument->GetTypeIndices()[PermutationIndex];
 					if (bIsArrayElement)
 					{
-						Type.ConvertToBaseElement();
+						TypeIndex = FRigVMRegistry::Get().GetBaseTypeFromArrayTypeIndex(TypeIndex);
 					}
-					if (LinkedTemplate->FilteredSupportsType(LinkedPin, Type.CPPType))
+					if (LinkedTemplate->FilteredSupportsType(LinkedPin, TypeIndex))
 					{
 						NewFilteredPermutations.Add(PermutationIndex);
 					}
@@ -584,17 +674,18 @@ TArray<int32> URigVMTemplateNode::GetNewFilteredPermutations(URigVMPin* InPin, U
 	{
 		if (const FRigVMTemplateArgument* Argument = GetTemplate()->FindArgument(RootPin->GetFName()))
 		{
-			const FString LinkedCPPType = LinkedPin->GetCPPType();
-			for(int32 PermIndex : PermutationsToTry)
+			const int32 LinkedTypeIndex = LinkedPin->GetTypeIndex();
+			for(int32 PermutationIndex : PermutationsToTry)
 			{
-				FRigVMTemplateArgumentType Type = Argument->GetTypes()[PermIndex];
+				int32 TypeIndex = Argument->GetTypeIndices()[PermutationIndex];
 				if (bIsArrayElement)
 				{
-					Type.ConvertToBaseElement();
+					TypeIndex = FRigVMRegistry::Get().GetBaseTypeFromArrayTypeIndex(TypeIndex);
 				}
-				if (Type.Matches(LinkedCPPType))
+
+				if (FRigVMRegistry::Get().CanMatchTypes(TypeIndex, LinkedTypeIndex, true))
 				{
-					NewFilteredPermutations.Add(PermIndex);
+					NewFilteredPermutations.Add(PermutationIndex);
 				}
 			}
 		}
@@ -602,38 +693,38 @@ TArray<int32> URigVMTemplateNode::GetNewFilteredPermutations(URigVMPin* InPin, U
 	return NewFilteredPermutations;
 }
 
-TArray<int32> URigVMTemplateNode::GetNewFilteredPermutations(URigVMPin* InPin, const TArray<FRigVMTemplateArgumentType>& InTypes)
+TArray<int32> URigVMTemplateNode::GetNewFilteredPermutations(URigVMPin* InPin, const TArray<int32>& InTypeIndices)
 {
 	TArray<int32> NewFilteredPermutations;
 	NewFilteredPermutations.Reserve(FilteredPermutations.Num());
 
 	URigVMPin* RootPin = InPin;
-	TArray<FRigVMTemplateArgumentType> RootTypes = InTypes;
+	TArray<int32> RootTypes = InTypeIndices;
 	if (URigVMPin* ParentPin = InPin->GetParentPin())
 	{
 		RootPin = ParentPin;
-		for (FRigVMTemplateArgumentType& Type : RootTypes)
+		for (int32& TypeIndex : RootTypes)
 		{
-			Type.ConvertToArray();
+			TypeIndex = FRigVMRegistry::Get().GetArrayTypeFromBaseTypeIndex(TypeIndex);
 		}
 	}
 
 	TArray<int32> PermutationsToTry = FilteredPermutations;
 
 	// Reduce permutations to the ones respecting the preferred types
-	TArray<int32> PreferredPermutations = FindPermutationsForTypes(PreferredPermutationTypes, true);
+	const TArray<int32> PreferredPermutations = FindPermutationsForTypes(PreferredPermutationPairs, true);
 	PermutationsToTry = PermutationsToTry.FilterByPredicate([&](const int32& OtherPermutation) { return PreferredPermutations.Contains(OtherPermutation); });
 	
 	if (const FRigVMTemplateArgument* Argument = GetTemplate()->FindArgument(RootPin->GetFName()))
 	{
-		const TArray<FRigVMTemplateArgumentType>& Types = Argument->GetTypes();
-		for(int32 PermIndex : PermutationsToTry)
+		const TArray<int32>& TypeIndices = Argument->GetTypeIndices();
+		for(int32 PermutationIndex : PermutationsToTry)
 		{
-			for (FRigVMTemplateArgumentType& RootType : RootTypes)
+			for (const int32& RootTypeIndex : RootTypes)
 			{				
-				if (Types[PermIndex].Matches(RootType.CPPType))
+				if (FRigVMRegistry::Get().CanMatchTypes(TypeIndices[PermutationIndex], RootTypeIndex, true))
 				{
-					NewFilteredPermutations.Add(PermIndex);
+					NewFilteredPermutations.Add(PermutationIndex);
 					break;
 				}
 			}			
@@ -642,25 +733,21 @@ TArray<int32> URigVMTemplateNode::GetNewFilteredPermutations(URigVMPin* InPin, c
 	return NewFilteredPermutations;
 }
 
-TArray<int32> URigVMTemplateNode::FindPermutationsForTypes(const TArray<FString>& ArgumentTypes, bool bAllowCasting)
+TArray<int32> URigVMTemplateNode::FindPermutationsForTypes(const TArray<FRigVMTemplatePreferredType>& ArgumentTypes, bool bAllowCasting) const
 {
 	TArray<int32> Permutations;
 	if (const FRigVMTemplate* Template = GetTemplate())
 	{
+		const FRigVMRegistry& Registry = FRigVMRegistry::Get();
+		
 		TArray<const FRigVMTemplateArgument*> Args;
-		TArray<FString> CPPTypes;
-		for (const FString& TypePairString : ArgumentTypes)
+		TArray<int32> TypeIndices;
+		for (const FRigVMTemplatePreferredType& ArgumentType : ArgumentTypes)
 		{
-			FString ArgName, Type;
-			if (!TypePairString.Split(TEXT(":"), &ArgName, &Type))
-			{
-				return {};
-			}
-
-			if (const FRigVMTemplateArgument* Argument = Template->FindArgument(*ArgName))
+			if (const FRigVMTemplateArgument* Argument = Template->FindArgument(ArgumentType.Argument))
 			{
 				Args.Add(Argument);
-				CPPTypes.Add(Type);
+				TypeIndices.Add(ArgumentType.TypeIndex);
 			}
 			else
 			{
@@ -675,8 +762,8 @@ TArray<int32> URigVMTemplateNode::FindPermutationsForTypes(const TArray<FString>
 			{
 				const FRigVMTemplateArgument* Argument = Args[ArgIndex];
 				{
-					if ((bAllowCasting && !Argument->GetTypes()[i].Matches(CPPTypes[ArgIndex])) ||
-						(!bAllowCasting && Argument->GetTypes()[i].CPPType != CPPTypes[ArgIndex]))
+					if ((bAllowCasting && !Registry.CanMatchTypes(Argument->GetTypeIndices()[i], TypeIndices[ArgIndex], true)) ||
+						(!bAllowCasting && Argument->GetTypeIndices()[i] != TypeIndices[ArgIndex]))
 					{
 						bAllArgsMatched = false;
 						break;
@@ -693,17 +780,17 @@ TArray<int32> URigVMTemplateNode::FindPermutationsForTypes(const TArray<FString>
 	return Permutations;
 }
 
-TArray<FString> URigVMTemplateNode::GetArgumentTypesForPermutation(const int32 InPermutationIndex)
+TArray<FRigVMTemplatePreferredType> URigVMTemplateNode::GetArgumentTypesForPermutation(const int32 InPermutationIndex) const
 {
-	TArray<FString> ArgTypes;
+	TArray<FRigVMTemplatePreferredType> ArgTypes;
 	if (const FRigVMTemplate* Template = GetTemplate())
 	{
 		for (int32 ArgIndex = 0; ArgIndex < Template->NumArguments(); ++ArgIndex)
 		{
 			const FRigVMTemplateArgument* Argument = Template->GetArgument(ArgIndex);
-			if (Argument->GetTypes().Num() > InPermutationIndex)
+			if (Argument->GetTypeIndices().Num() > InPermutationIndex)
 			{
-				ArgTypes.Add(Argument->GetName().ToString() + TEXT(":") + Argument->GetTypes()[InPermutationIndex].CPPType);
+				ArgTypes.Emplace(Argument->GetName(), Argument->GetTypeIndices()[InPermutationIndex]);
 			}
 			else
 			{
@@ -716,9 +803,9 @@ TArray<FString> URigVMTemplateNode::GetArgumentTypesForPermutation(const int32 I
 	return ArgTypes;
 }
 
-bool URigVMTemplateNode::PinNeedsFilteredTypesUpdate(URigVMPin* InPin, const TArray<FRigVMTemplateArgumentType>& InTypes)
+bool URigVMTemplateNode::PinNeedsFilteredTypesUpdate(URigVMPin* InPin, const TArray<int32>& InTypeIndices)
 {
-	TArray<int32> NewFilteredPermutations = GetNewFilteredPermutations(InPin, InTypes);
+	const TArray<int32> NewFilteredPermutations = GetNewFilteredPermutations(InPin, InTypeIndices);
 	if (NewFilteredPermutations.Num() == FilteredPermutations.Num())
 	{
 		return false;
@@ -728,8 +815,7 @@ bool URigVMTemplateNode::PinNeedsFilteredTypesUpdate(URigVMPin* InPin, const TAr
 
 bool URigVMTemplateNode::PinNeedsFilteredTypesUpdate(URigVMPin* InPin, URigVMPin* LinkedPin)
 {
-	TArray<int32> NewFilteredPermutations = GetNewFilteredPermutations(InPin, LinkedPin);
-
+	const TArray<int32> NewFilteredPermutations = GetNewFilteredPermutations(InPin, LinkedPin);
 	if (NewFilteredPermutations.Num() == FilteredPermutations.Num())
 	{
 		return false;
@@ -743,7 +829,7 @@ bool URigVMTemplateNode::UpdateFilteredPermutations(URigVMPin* InPin, URigVMPin*
 	ensureMsgf(InPin->GetNode() == this, TEXT("Updating filtered permutations of %s with pin from another node %s"), *GetNodePath(), *InPin->GetPinPath(true));
 	ensureMsgf(LinkedPin->GetNode() != this, TEXT("Updating filtered permutations of %s with linked pin from same node %s"), *GetNodePath(), *LinkedPin->GetPinPath(true));
 
-	TArray<int32> NewFilteredPermutations = GetNewFilteredPermutations(InPin, LinkedPin);	
+	const TArray<int32> NewFilteredPermutations = GetNewFilteredPermutations(InPin, LinkedPin);	
 	if (NewFilteredPermutations.Num() == FilteredPermutations.Num())
 	{
 		return false;
@@ -758,9 +844,9 @@ bool URigVMTemplateNode::UpdateFilteredPermutations(URigVMPin* InPin, URigVMPin*
 	return true;
 }
 
-bool URigVMTemplateNode::UpdateFilteredPermutations(URigVMPin* InPin, const TArray<FRigVMTemplateArgumentType>& InTypes)
+bool URigVMTemplateNode::UpdateFilteredPermutations(URigVMPin* InPin, const TArray<int32>& InTypeIndices)
 {
-	TArray<int32> NewFilteredPermutations = GetNewFilteredPermutations(InPin, InTypes);
+	const TArray<int32> NewFilteredPermutations = GetNewFilteredPermutations(InPin, InTypeIndices);
 	if (NewFilteredPermutations.Num() == FilteredPermutations.Num())
 	{
 		return false;
@@ -798,9 +884,9 @@ void URigVMTemplateNode::InitializeFilteredPermutations()
 {
 	if (const FRigVMTemplate* Template = GetTemplate())
 	{
-		if (!PreferredPermutationTypes.IsEmpty())
+		if (!PreferredPermutationPairs.IsEmpty())
 		{
-			FilteredPermutations = FindPermutationsForTypes(PreferredPermutationTypes);
+			FilteredPermutations = FindPermutationsForTypes(PreferredPermutationPairs);
 		}
 		else
 		{
@@ -822,7 +908,7 @@ void URigVMTemplateNode::InitializeFilteredPermutationsFromTypes()
 
 	if (const FRigVMTemplate* Template = GetTemplate())
 	{
-		TArray<FString> ArgTypes;
+		TArray<FRigVMTemplatePreferredType> ArgTypes;
 		for (int32 ArgIndex = 0; ArgIndex < Template->NumArguments(); ++ArgIndex)
 		{
 			const FRigVMTemplateArgument* Argument = Template->GetArgument(ArgIndex);
@@ -830,16 +916,16 @@ void URigVMTemplateNode::InitializeFilteredPermutationsFromTypes()
 			{
 				if (!Pin->IsWildCard())
 				{
-					ArgTypes.Add(Argument->GetName().ToString() + TEXT(":") + Pin->GetCPPType());
+					ArgTypes.Emplace(Argument->GetName(), Pin->GetTypeIndex());
 				}
 			}
 		}
 
-		TArray<int32> Permutations = FindPermutationsForTypes(ArgTypes);
+		const TArray<int32> Permutations = FindPermutationsForTypes(ArgTypes);
 		if (!Permutations.IsEmpty())
 		{
 			FilteredPermutations = Permutations;
-			PreferredPermutationTypes = ArgTypes;
+			PreferredPermutationPairs = ArgTypes;
 		}
 		else
 		{
diff --git a/Engine/Source/Developer/RigVMDeveloper/Private/RigVMModel/RigVMController.cpp b/Engine/Source/Developer/RigVMDeveloper/Private/RigVMModel/RigVMController.cpp
--- a/Engine/Source/Developer/RigVMDeveloper/Private/RigVMModel/RigVMController.cpp
+++ b/Engine/Source/Developer/RigVMDeveloper/Private/RigVMModel/RigVMController.cpp
@@ -473,7 +473,7 @@ TArray<FString> URigVMController::GetAddNodePythonCommands(URigVMNode* Node) con
 									*Pin->GetCPPType(),
 									*Pin->GetCPPTypeObject()->GetPathName()));
 
-							TypeMap.Add(Argument->GetName(), FRigVMTemplateArgumentType(Pin->GetCPPType(), Pin->GetCPPTypeObject()));
+							TypeMap.Add(Argument->GetName(), Pin->GetTypeIndex());
 							Template->Resolve(TypeMap, Permutations, false);
 						}
 					}
@@ -798,14 +798,14 @@ URigVMUnitNode* URigVMController::AddUnitNode(UScriptStruct* InScriptStruct, con
 		TArray<int32> OldPermutations = TemplateNode->FilteredPermutations;
 		int32 PermutationIndex = Template->FindPermutation(Function);
 		TemplateNode->FilteredPermutations = {PermutationIndex};
-		const TArray<FString> NewPreferredPermutationTypes = TemplateNode->GetArgumentTypesForPermutation(PermutationIndex);
+		const TArray<FRigVMTemplatePreferredType> NewPreferredPermutationTypes = TemplateNode->GetArgumentTypesForPermutation(PermutationIndex);
 		if (bSetupUndoRedo)
 		{
 			FRigVMSetTemplateFilteredPermutationsAction Action(TemplateNode, nullptr, OldPermutations);
 			ActionStack->AddAction(Action);
 			ActionStack->AddAction(FRigVMSetPreferredTemplatePermutationsAction(TemplateNode, NewPreferredPermutationTypes));
 		}		
-		TemplateNode->PreferredPermutationTypes = NewPreferredPermutationTypes;
+		TemplateNode->PreferredPermutationPairs = NewPreferredPermutationTypes;
 		UpdateTemplateNodePinTypes(TemplateNode, bSetupUndoRedo);
 
 		if (UnitNodeCreatedContext.IsValid())
@@ -1821,8 +1821,10 @@ bool URigVMController::UnresolveTemplateNodes(const TArray<URigVMTemplateNode*>&
 	Algo::Transform(InNodes, Nodes, [](URigVMTemplateNode* InNode) -> URigVMNode* { return InNode; });
 	TArray<TPair<FString, FString>> LinkedPaths = GetLinkedPinPaths(Nodes);
 
+	const FRigVMRegistry& Registry = FRigVMRegistry::Get();
+	
 	// Find pins outside our set of nodes which were resolved to a type
-	TMap<URigVMPin*, FRigVMTemplateArgumentType> PinsToResolve;
+	TMap<URigVMPin*, int32> PinsToResolve;
 	for (TPair<FString, FString>& LinkPath : LinkedPaths)
 	{
 		const FString& SourcePath = LinkPath.Key;
@@ -1867,14 +1869,14 @@ bool URigVMController::UnresolveTemplateNodes(const TArray<URigVMTemplateNode*>&
 		{
 			if (!SourcePin->IsWildCard())
 			{
-				PinsToResolve.Add(SourcePin, FRigVMTemplateArgumentType(SourcePin->GetCPPType(), SourcePin->GetCPPTypeObject()));
+				PinsToResolve.Add(SourcePin, SourcePin->GetTypeIndex());
 			}
 		}
 		else if(bTargetOutside && !bSourceOutside)
 		{
 			if (!TargetPin->IsWildCard())
 			{
-				PinsToResolve.Add(TargetPin, FRigVMTemplateArgumentType(TargetPin->GetCPPType(), TargetPin->GetCPPTypeObject()));
+				PinsToResolve.Add(TargetPin, TargetPin->GetTypeIndex());
 			}
 		}			
 	}
@@ -1888,13 +1890,13 @@ bool URigVMController::UnresolveTemplateNodes(const TArray<URigVMTemplateNode*>&
 			continue;
 		}
 
-		if (!Node->PreferredPermutationTypes.IsEmpty())
+		if (!Node->PreferredPermutationPairs.IsEmpty())
 		{
 			if (bSetupUndoRedo)
 			{
 				ActionStack->AddAction(FRigVMSetPreferredTemplatePermutationsAction(Node, {}));
 			}
-			Node->PreferredPermutationTypes = {};
+			Node->PreferredPermutationPairs = {};
 		}
 		
 		TArray<int32> OldPermutations = Node->FilteredPermutations;
@@ -1921,19 +1923,19 @@ bool URigVMController::UnresolveTemplateNodes(const TArray<URigVMTemplateNode*>&
 		// if the types' containers don't match the current pin containers try again
 		bool bTypesWereAdjusted = false;
 		FRigVMTemplate::FTypeMap AdjustedTypes = Types;
-		for(TPair<FName, FRigVMTemplateArgumentType>& TypePair : AdjustedTypes)
+		for(TPair<FName, int32>& TypePair : AdjustedTypes)
 		{
 			if(URigVMPin* Pin = Node->FindPin(TypePair.Key.ToString()))
 			{
-				if(Pin->IsArray() != TypePair.Value.IsArray())
+				if(Pin->IsArray() != Registry.IsArrayType(TypePair.Value))
 				{
 					if(Pin->IsArray())
 					{
-						TypePair.Value.CPPType = RigVMTypeUtils::ArrayTypeFromBaseType(TypePair.Value.CPPType);
+						TypePair.Value = Registry.GetArrayTypeFromBaseTypeIndex(TypePair.Value);
 					}
 					else
 					{
-						TypePair.Value.CPPType = RigVMTypeUtils::BaseTypeFromArrayType(TypePair.Value.CPPType);
+						TypePair.Value = Registry.GetBaseTypeFromArrayTypeIndex(TypePair.Value);
 					}
 					bTypesWereAdjusted = true;
 				}
@@ -1949,10 +1951,10 @@ bool URigVMController::UnresolveTemplateNodes(const TArray<URigVMTemplateNode*>&
 			}
 		}
 
-		for(const TPair<FName, FRigVMTemplateArgumentType>& TypePair : Types)
+		for(const TPair<FName, int32>& TypePair : Types)
 		{
 			const FName& PinName = TypePair.Key;
-			const FRigVMTemplateArgumentType& ExpectedType = TypePair.Value;
+			const int32& ExpectedTypeIndex = TypePair.Value;
 			
 			URigVMPin* Pin = Node->FindPin(PinName.ToString());
 			if(Pin == nullptr)
@@ -1965,19 +1967,18 @@ bool URigVMController::UnresolveTemplateNodes(const TArray<URigVMTemplateNode*>&
 				return false;
 			}
 
-			if(ChangePinType(Pin, ExpectedType.CPPType, ExpectedType.CPPTypeObject, bSetupUndoRedo, false, true, !Pin->IsArray()))
+			if(ChangePinType(Pin, ExpectedTypeIndex, bSetupUndoRedo, false, true, !Pin->IsArray()))
 			{
 				bChangedAnyPin = true;
 
 				if(Pin->IsArray())
 				{
-					FRigVMTemplateArgumentType ElementType = ExpectedType;
-					ElementType.CPPType = RigVMTypeUtils::BaseTypeFromArrayType(ElementType.CPPType);
+					const int32 ElementTypeIndex = FRigVMRegistry::Get().GetBaseTypeFromArrayTypeIndex(ExpectedTypeIndex);
 					
 					TArray<URigVMPin*> SubPins = Pin->GetSubPins();
 					for (URigVMPin* SubPin : SubPins)
 					{
-						ChangePinType(SubPin, ElementType.CPPType, ElementType.CPPTypeObject, bSetupUndoRedo, false, true, true);
+						ChangePinType(SubPin, ElementTypeIndex, bSetupUndoRedo, false, true, true);
 					}
 				}
 			}
@@ -1992,7 +1993,7 @@ bool URigVMController::UnresolveTemplateNodes(const TArray<URigVMTemplateNode*>&
 			return A->IsWildCard() == B->IsWildCard();
 		}), bSetupUndoRedo);
 	
-	for (TPair<URigVMPin*, FRigVMTemplateArgumentType>& Pin : PinsToResolve)
+	for (TPair<URigVMPin*, int32>& Pin : PinsToResolve)
 	{
 		TArray<URigVMLink*> LinksToBreak;
 		for (URigVMLink* Link : Pin.Key->GetLinks())
@@ -2001,7 +2002,7 @@ bool URigVMController::UnresolveTemplateNodes(const TArray<URigVMTemplateNode*>&
 			{
 				if (Nodes.Contains(OppositePin->GetNode()))
 				{
-					if (OppositePin->GetCPPType() != Pin.Value.CPPType)
+					if (OppositePin->GetTypeIndex() != Pin.Value)
 					{
 						LinksToBreak.Add(Link);
 					}
@@ -3334,6 +3335,13 @@ FString URigVMController::ExportNodesToText(const TArray<FName>& InNodeNames)
 					AllNodeNames.AddUnique(Injection->Node->GetFName());
 				}
 			}
+
+			if(URigVMTemplateNode* TemplateNode = Cast<URigVMTemplateNode>(Node))
+			{
+				// make sure the preferred types are saved
+				// as strings - not just as indices
+				TemplateNode->ConvertPreferredTypesToString();
+			}
 		}
 	}
 
@@ -3533,6 +3541,14 @@ TArray<FName> URigVMController::ImportNodesFromText(const FString& InText, bool
 			continue;
 		}
 
+		if(URigVMTemplateNode* TemplateNode = Cast<URigVMTemplateNode>(CreatedNode))
+		{
+			// make sure the update preferred type pairs
+			// from string to index - since indices may
+			// have changed between sessions
+			TemplateNode->ConvertPreferredTypesToTypeIndex();
+		}
+
 		Graph->Nodes.Add(CreatedNode);
 
 		if (bSetupUndoRedo)
@@ -3701,7 +3717,7 @@ TArray<FName> URigVMController::ImportNodesFromText(const FString& InText, bool
 							{
 								if (!FirstTemplateNode->IsSingleton())
 								{
-									TArray<FRigVMTemplateArgumentType> InputTypes = GetWildcardFilteredTypes(SourcePin);
+									TArray<int32> InputTypes = GetWildcardFilteredTypeIndices(SourcePin);
 									if (InputTypes.Num() > 0)
 									{
 										PrepareTemplatePinForType(TargetPin, InputTypes, bSetupUndoRedo);										
@@ -3712,7 +3728,7 @@ TArray<FName> URigVMController::ImportNodesFromText(const FString& InText, bool
 							{
 								if (!SecondTemplateNode->IsSingleton())
 								{
-									TArray<FRigVMTemplateArgumentType> OutTypes = GetWildcardFilteredTypes(TargetPin);
+									TArray<int32> OutTypes = GetWildcardFilteredTypeIndices(TargetPin);
 									if (OutTypes.Num() > 0)
 									{
 										PrepareTemplatePinForType(SourcePin, OutTypes, bSetupUndoRedo);										
@@ -9255,7 +9271,7 @@ bool URigVMController::AddLink(URigVMPin* OutputPin, URigVMPin* InputPin, bool b
 		{
 			if (!FirstTemplateNode->IsSingleton())
 			{
-				TArray<FRigVMTemplateArgumentType> InputTypes = GetWildcardFilteredTypes(SecondToResolve);
+				TArray<int32> InputTypes = GetWildcardFilteredTypeIndices(SecondToResolve);
 				if (InputTypes.Num() > 0)
 				{
 					if (!PrepareTemplatePinForType(FirstToResolve, InputTypes, bSetupUndoRedo))
@@ -9270,7 +9286,7 @@ bool URigVMController::AddLink(URigVMPin* OutputPin, URigVMPin* InputPin, bool b
 		{
 			if (!SecondTemplateNode->IsSingleton())
 			{
-				TArray<FRigVMTemplateArgumentType> OutTypes = GetWildcardFilteredTypes(FirstToResolve);
+				TArray<int32> OutTypes = GetWildcardFilteredTypeIndices(FirstToResolve);
 				if (OutTypes.Num() > 0)
 				{
 					if (!PrepareTemplatePinForType(SecondToResolve, OutTypes, bSetupUndoRedo))
@@ -11644,16 +11660,16 @@ URigVMRerouteNode* URigVMController::AddFreeRerouteNode(bool bShowAsFullNode, co
 	Node->InitializeFilteredPermutations();
 	if (InCPPType != RigVMTypeUtils::GetWildCardCPPType() && InCPPType != RigVMTypeUtils::GetWildCardArrayCPPType())
 	{
-		PrepareTemplatePinForType(ValuePin, {FRigVMTemplateArgumentType(ValuePin->GetCPPType(), ValuePin->GetCPPTypeObject())}, bSetupUndoRedo);
+		PrepareTemplatePinForType(ValuePin, {ValuePin->GetTypeIndex()}, bSetupUndoRedo);
 		TArray<int32> FilterPermutations = Node->GetFilteredPermutationsIndices();
 		if (FilterPermutations.Num() == 1)
 		{
-			const TArray<FString> NewPreferredPermutationTypes = Node->GetArgumentTypesForPermutation(FilterPermutations[0]);
+			const TArray<FRigVMTemplatePreferredType> NewPreferredPermutationTypes = Node->GetArgumentTypesForPermutation(FilterPermutations[0]);
 			if (bSetupUndoRedo)
 			{
 				ActionStack->AddAction(FRigVMSetPreferredTemplatePermutationsAction(Node, NewPreferredPermutationTypes));
 			}
-			Node->PreferredPermutationTypes = NewPreferredPermutationTypes;
+			Node->PreferredPermutationPairs = NewPreferredPermutationTypes;
 		}
 	}
 
@@ -11836,16 +11852,16 @@ URigVMIfNode* URigVMController::AddIfNode(const FString& InCPPType, const FName&
 	Node->InitializeFilteredPermutations();
 	if (InCPPType != RigVMTypeUtils::GetWildCardCPPType() && InCPPType != RigVMTypeUtils::GetWildCardArrayCPPType())
 	{
-		PrepareTemplatePinForType(TruePin, {FRigVMTemplateArgumentType(TruePin->GetCPPType(), TruePin->GetCPPTypeObject())}, bSetupUndoRedo);
+		PrepareTemplatePinForType(TruePin, {TruePin->GetTypeIndex()}, bSetupUndoRedo);
 		TArray<int32> FilterPermutations = Node->GetFilteredPermutationsIndices();
 		if (FilterPermutations.Num() == 1)
 		{
-			const TArray<FString> NewPreferredPermutationTypes = Node->GetArgumentTypesForPermutation(FilterPermutations[0]);
+			const TArray<FRigVMTemplatePreferredType> NewPreferredPermutationTypes = Node->GetArgumentTypesForPermutation(FilterPermutations[0]);
 			if (bSetupUndoRedo)
 			{
 				ActionStack->AddAction(FRigVMSetPreferredTemplatePermutationsAction(Node, NewPreferredPermutationTypes));
 			}
-			Node->PreferredPermutationTypes = NewPreferredPermutationTypes;
+			Node->PreferredPermutationPairs = NewPreferredPermutationTypes;
 		}
 	}
 	
@@ -11957,16 +11973,16 @@ URigVMSelectNode* URigVMController::AddSelectNode(const FString& InCPPType, cons
 	Node->InitializeFilteredPermutations();
 	if (InCPPType != RigVMTypeUtils::GetWildCardCPPType() && InCPPType != RigVMTypeUtils::GetWildCardArrayCPPType())
 	{
-		PrepareTemplatePinForType(ResultPin, {FRigVMTemplateArgumentType(ResultPin->GetCPPType(), ResultPin->GetCPPTypeObject())}, bSetupUndoRedo);
+		PrepareTemplatePinForType(ResultPin, {ResultPin->GetTypeIndex()}, bSetupUndoRedo);
 		TArray<int32> FilterPermutations = Node->GetFilteredPermutationsIndices();
 		if (FilterPermutations.Num() == 1)
 		{
-			const TArray<FString> NewPreferredPermutationTypes = Node->GetArgumentTypesForPermutation(FilterPermutations[0]);
+			const TArray<FRigVMTemplatePreferredType> NewPreferredPermutationTypes = Node->GetArgumentTypesForPermutation(FilterPermutations[0]);
 			if (bSetupUndoRedo)
 			{
 				ActionStack->AddAction(FRigVMSetPreferredTemplatePermutationsAction(Node, NewPreferredPermutationTypes));
 			}
-			Node->PreferredPermutationTypes = NewPreferredPermutationTypes;
+			Node->PreferredPermutationPairs = NewPreferredPermutationTypes;
 		}
 	}
 	
@@ -12059,8 +12075,9 @@ URigVMTemplateNode* URigVMController::AddTemplateNode(const FName& InNotation, c
 		const FRigVMTemplateArgument* Arg = Template->GetArgument(ArgIndex);
 
 		URigVMPin* Pin = NewObject<URigVMPin>(Node, Arg->GetName());
-		const FRigVMTemplateArgumentType& Type = Types.FindChecked(Arg->GetName());
-		Pin->CPPType = Type.CPPType;
+		const int32& TypeIndex = Types.FindChecked(Arg->GetName());
+		const FRigVMTemplateArgumentType Type = FRigVMRegistry::Get().GetType(TypeIndex);
+		Pin->CPPType = Type.CPPType.ToString();
 		Pin->CPPTypeObject = Type.CPPTypeObject;
 		if (Pin->CPPTypeObject)
 		{
@@ -12542,16 +12559,16 @@ URigVMArrayNode* URigVMController::AddArrayNode(ERigVMOpCode InOpCode, const FSt
 	if (InCPPType != RigVMTypeUtils::GetWildCardCPPType() && InCPPType != RigVMTypeUtils::GetWildCardArrayCPPType())
 	{
 		URigVMPin* ArrayPin = Node->FindPin(URigVMArrayNode::ArrayName);
-		PrepareTemplatePinForType(ArrayPin, {FRigVMTemplateArgumentType(ArrayPin->GetCPPType(), ArrayPin->GetCPPTypeObject())}, bSetupUndoRedo);
+		PrepareTemplatePinForType(ArrayPin, {ArrayPin->GetTypeIndex()}, bSetupUndoRedo);
 		TArray<int32> FilterPermutations = Node->GetFilteredPermutationsIndices();
 		if (FilterPermutations.Num() == 1)
 		{
-			const TArray<FString> NewPreferredPermutationTypes = Node->GetArgumentTypesForPermutation(FilterPermutations[0]);
+			const TArray<FRigVMTemplatePreferredType> NewPreferredPermutationTypes = Node->GetArgumentTypesForPermutation(FilterPermutations[0]);
 			if (bSetupUndoRedo)
 			{
 				ActionStack->AddAction(FRigVMSetPreferredTemplatePermutationsAction(Node, NewPreferredPermutationTypes));
 			}
-			Node->PreferredPermutationTypes = NewPreferredPermutationTypes;
+			Node->PreferredPermutationPairs = NewPreferredPermutationTypes;
 		}
 	}
 
@@ -14949,7 +14966,7 @@ bool URigVMController::FullyResolveTemplateNode(URigVMTemplateNode* InNode, int3
 		for (URigVMPin* Pin : InNode->GetPins())
 		{
 			check(!Pin->IsWildCard());
-			TypeMap.Add(Pin->GetFName(), FRigVMTemplateArgumentType(Pin->GetCPPType(), Pin->GetCPPTypeObject()));
+			TypeMap.Add(Pin->GetFName(), Pin->GetTypeIndex());
 		}
 
 		TArray<int32> Permutations;
@@ -14966,11 +14983,11 @@ bool URigVMController::FullyResolveTemplateNode(URigVMTemplateNode* InNode, int3
 	// find all existing pins that we may need to change
 	TArray<FRigVMTemplateArgument> MissingPins;
 	TArray<URigVMPin*> PinsToRemove;
-	TMap<URigVMPin*, FRigVMTemplateArgumentType> PinTypesToChange;
+	TMap<URigVMPin*, int32> PinTypesToChange;
 	for(int32 ArgIndex = 0; ArgIndex < Template->NumArguments(); ArgIndex++)
 	{
 		const FRigVMTemplateArgument* Argument = Template->GetArgument(ArgIndex);
-		const FRigVMTemplateArgumentType ResolvedType = Argument->GetSupportedTypes(PermutationIndices)[0];
+		const int32 ResolvedTypeIndex = Argument->GetSupportedTypeIndices(PermutationIndices)[0];
 
 		URigVMPin* Pin = InNode->FindPin(Argument->GetName().ToString());
 		if(Pin == nullptr)
@@ -14982,9 +14999,9 @@ bool URigVMController::FullyResolveTemplateNode(URigVMTemplateNode* InNode, int3
 			return false;
 		}
 
-		if(Pin->GetCPPType() != ResolvedType.CPPType)
+		if(Pin->GetTypeIndex() != ResolvedTypeIndex)
 		{
-			PinTypesToChange.Add(Pin, ResolvedType);
+			PinTypesToChange.Add(Pin, ResolvedTypeIndex);
 		}
 	}
 
@@ -14997,13 +15014,13 @@ bool URigVMController::FullyResolveTemplateNode(URigVMTemplateNode* InNode, int3
 			for (TFieldIterator<FProperty> It(StructToVisit, EFieldIterationFlags::None); It; ++It)
 			{
 				const FRigVMTemplateArgument ExpectedArgument(*It);
-				const FRigVMTemplateArgumentType ExpectedType = ExpectedArgument.GetSupportedTypes()[0];
+				const int32 ExpectedTypeIndex = ExpectedArgument.GetSupportedTypeIndices()[0];
 
 				if(URigVMPin* Pin = InNode->FindPin(It->GetFName().ToString()))
 				{
-					if(Pin->GetCPPType() != ExpectedType.CPPType)
+					if(Pin->GetTypeIndex() != ExpectedTypeIndex)
 					{
-						PinTypesToChange.Add(Pin, ExpectedType);
+						PinTypesToChange.Add(Pin, ExpectedTypeIndex);
 					}
 				}
 				else
@@ -15043,19 +15060,18 @@ bool URigVMController::FullyResolveTemplateNode(URigVMTemplateNode* InNode, int3
 
 	// update the incorrectly typed pins
 	bool bNeedsTemplateUpdate = false;
-	for(const TPair<URigVMPin*, FRigVMTemplateArgumentType>& Pair : PinTypesToChange)
+	for(const TPair<URigVMPin*, int32>& Pair : PinTypesToChange)
 	{
 		URigVMPin* Pin = Pair.Key;
-		const FRigVMTemplateArgumentType& ExpectedType = Pair.Value;
+		const int32& ExpectedTypeIndex = Pair.Value;
 		
 		if(!Pin->IsWildCard())
 		{
-			if((Pin->GetCPPType() != ExpectedType.CPPType) ||
-				(Pin->GetCPPTypeObject() != ExpectedType.CPPTypeObject))
+			if(Pin->GetTypeIndex() != ExpectedTypeIndex)
 			{
 				bNeedsTemplateUpdate = true;
-				const FString CPPType = Pin->IsArray() ? RigVMTypeUtils::GetWildCardArrayCPPType() : RigVMTypeUtils::GetWildCardCPPType();
-				if(!ChangePinType(Pin, CPPType, RigVMTypeUtils::GetWildCardCPPTypeObject(), bSetupUndoRedo, false, true, true))
+				const int32 WildCardIndex = Pin->IsArray() ? RigVMTypeUtils::TypeIndex::WildCardArray : RigVMTypeUtils::TypeIndex::WildCard;
+				if(!ChangePinType(Pin, WildCardIndex, bSetupUndoRedo, false, true, true))
 				{
 					if(bSetupUndoRedo)
 					{
@@ -15070,7 +15086,7 @@ bool URigVMController::FullyResolveTemplateNode(URigVMTemplateNode* InNode, int3
 		{
 			bNeedsTemplateUpdate = true;
 			
-			if (!UpdateFilteredPermutations(Pin, {ExpectedType}, bSetupUndoRedo))
+			if (!UpdateFilteredPermutations(Pin, {ExpectedTypeIndex}, bSetupUndoRedo))
 			{
 				if(bSetupUndoRedo)
 				{
@@ -15120,7 +15136,7 @@ bool URigVMController::FullyResolveTemplateNode(URigVMTemplateNode* InNode, int3
 	return true;
 }
 
-bool URigVMController::PrepareTemplatePinForType(URigVMPin* InPin, const TArray<FRigVMTemplateArgumentType>& InTypes, bool bSetupUndoRedo)
+bool URigVMController::PrepareTemplatePinForType(URigVMPin* InPin, const TArray<int32>& InTypeIndices, bool bSetupUndoRedo)
 {
 	if (!IsValidGraph())
 	{
@@ -15145,13 +15161,13 @@ bool URigVMController::PrepareTemplatePinForType(URigVMPin* InPin, const TArray<
 
 	// If the pin is an element of an array, get the array pin and convert the types
 	URigVMPin* RootPin = InPin;
-	TArray<FRigVMTemplateArgumentType> Types = InTypes;
+	TArray<int32> Types = InTypeIndices;
 	if (InPin->IsArrayElement())
 	{
 		RootPin = InPin->GetParentPin();
-		for (FRigVMTemplateArgumentType& Type : Types)
+		for (int32& Type : Types)
 		{
-			Type.ConvertToArray();
+			Type = FRigVMRegistry::Get().GetArrayTypeFromBaseTypeIndex(Type);
 		}
 	}
 
@@ -15159,7 +15175,7 @@ bool URigVMController::PrepareTemplatePinForType(URigVMPin* InPin, const TArray<
 	// the type required matches the one in the pin
 	if (!TemplateNode->GetTemplate()->FindArgument(RootPin->GetFName()))
 	{
-		if (Types.Num() == 1 && Types[0].Matches(RootPin->GetCPPType()))
+		if (Types.Num() == 1 && FRigVMRegistry::Get().CanMatchTypes(Types[0], RootPin->GetTypeIndex(), true))
 		{
 			return true;
 		}
@@ -15167,9 +15183,9 @@ bool URigVMController::PrepareTemplatePinForType(URigVMPin* InPin, const TArray<
 	}
 
 	bool bFilteredSupportsType = false;
-	for (const FRigVMTemplateArgumentType& Type : Types)
+	for (const int32& Type : Types)
 	{
-		if (TemplateNode->FilteredSupportsType(RootPin, Type.CPPType))
+		if (TemplateNode->FilteredSupportsType(RootPin, Type))
 		{
 			bFilteredSupportsType = true;
 			break;
@@ -15179,18 +15195,18 @@ bool URigVMController::PrepareTemplatePinForType(URigVMPin* InPin, const TArray<
 	bool bSupportsType = false;
 	if (!bFilteredSupportsType)
 	{
-		if (!TemplateNode->PreferredPermutationTypes.IsEmpty())
+		if (!TemplateNode->PreferredPermutationPairs.IsEmpty())
 		{
 			if (bSetupUndoRedo)
 			{
 				ActionStack->AddAction(FRigVMSetPreferredTemplatePermutationsAction(TemplateNode, {}));
 			}
-			TemplateNode->PreferredPermutationTypes = {};
+			TemplateNode->PreferredPermutationPairs = {};
 		}
 		
-		for (const FRigVMTemplateArgumentType& Type : Types)
+		for (const int32& Type : Types)
 		{
-			if (TemplateNode->SupportsType(RootPin, Type.CPPType))
+			if (TemplateNode->SupportsType(RootPin, Type))
 			{
 				bSupportsType = true;
 				break;
@@ -15234,7 +15250,7 @@ bool URigVMController::PrepareTemplatePinForType(URigVMPin* InPin, const TArray<
 				InitializeAllTemplateFiltersInGraph(true, false);
 
 				// Update our node with the requested type
-				UpdateFilteredPermutations(InPin, InTypes, true);
+				UpdateFilteredPermutations(InPin, InTypeIndices, true);
 				
 				// Update types in case some pins on this node disappear, and should no longer propagate
 				UpdateTemplateNodePinTypes(TemplateNode, true);
@@ -15271,13 +15287,13 @@ bool URigVMController::PrepareTemplatePinForType(URigVMPin* InPin, const TArray<
 		// so no links to break
 		if (LinksToBreak.IsEmpty())
 		{
-			if (!TemplateNode->PreferredPermutationTypes.IsEmpty())
+			if (!TemplateNode->PreferredPermutationPairs.IsEmpty())
 			{
 				return false;
 			}
 			
 			UnresolveTemplateNodes({TemplateNode}, bSetupUndoRedo);
-			PrepareTemplatePinForType(InPin, InTypes, bSetupUndoRedo);
+			PrepareTemplatePinForType(InPin, InTypeIndices, bSetupUndoRedo);
 			return true;
 		}
 
@@ -15303,7 +15319,7 @@ bool URigVMController::PrepareTemplatePinForType(URigVMPin* InPin, const TArray<
 			}
 
 			ensure(LinksToBreak.Num() > 0);
-			PrepareTemplatePinForType(InPin, InTypes, bSetupUndoRedo);
+			PrepareTemplatePinForType(InPin, InTypeIndices, bSetupUndoRedo);
 			return true;
 		}
 	}
@@ -15311,12 +15327,12 @@ bool URigVMController::PrepareTemplatePinForType(URigVMPin* InPin, const TArray<
 	return false;
 }
 
-TArray<FRigVMTemplateArgumentType> URigVMController::GetWildcardFilteredTypes(URigVMPin* InPin)
+TArray<int32> URigVMController::GetWildcardFilteredTypeIndices(URigVMPin* InPin)
 {
-	TArray<FRigVMTemplateArgumentType> Types;
+	TArray<int32> Types;
 	if (!InPin->IsWildCard())
 	{
-		Types.Add(FRigVMTemplateArgumentType(InPin->GetCPPType(), InPin->GetCPPTypeObject()));
+		Types.Add(InPin->GetTypeIndex());
 		return Types;
 	}
 
@@ -15361,7 +15377,7 @@ bool URigVMController::ResolveWildCardPin(const FString& InPinPath, const FStrin
 
 	if (URigVMPin* Pin = Graph->FindPin(InPinPath))
 	{
-		if (ResolveWildCardPin(Pin, FRigVMTemplateArgumentType(CPPType, CPPTypeObject), bSetupUndoRedo, bPrintPythonCommand))
+		if (ResolveWildCardPin(Pin, FRigVMTemplateArgumentType(*CPPType, CPPTypeObject), bSetupUndoRedo, bPrintPythonCommand))
 		{
 			if (bPrintPythonCommand)
 			{
@@ -15395,6 +15411,35 @@ bool URigVMController::ResolveWildCardPin(URigVMPin* InPin, const FRigVMTemplate
 		return false;
 	}
 
+	return ResolveWildCardPin(InPin, FRigVMRegistry::Get().GetTypeIndex(InType), bSetupUndoRedo, bPrintPythonCommand);
+}
+
+bool URigVMController::ResolveWildCardPin(const FString& InPinPath, int32 InTypeIndex, bool bSetupUndoRedo,
+	bool bPrintPythonCommand)
+{
+	if (!IsValidGraph())
+	{
+		return false;
+	}
+
+	if (!bIsTransacting && !IsGraphEditable())
+	{
+		return false;
+	}
+
+	URigVMGraph* Graph = GetGraph();
+	check(Graph);
+
+	if (URigVMPin* Pin = Graph->FindPin(InPinPath))
+	{
+		return ResolveWildCardPin(Pin, InTypeIndex, bSetupUndoRedo, bPrintPythonCommand);
+	}
+	return false;
+}
+
+bool URigVMController::ResolveWildCardPin(URigVMPin* InPin, int32 InTypeIndex, bool bSetupUndoRedo,
+	bool bPrintPythonCommand)
+{
 	if(InPin->IsWildCard())
 	{
 		ensure(InPin->GetNode()->IsA<URigVMTemplateNode>());
@@ -15407,7 +15452,7 @@ bool URigVMController::ResolveWildCardPin(URigVMPin* InPin, const FRigVMTemplate
 		}
 		
 		URigVMTemplateNode* TemplateNode = CastChecked<URigVMTemplateNode>(InPin->GetNode());
-		if (!PrepareTemplatePinForType(InPin, {InType}, bSetupUndoRedo))
+		if (!PrepareTemplatePinForType(InPin, {InTypeIndex}, bSetupUndoRedo))
 		{
 			if (bSetupUndoRedo)
 			{
@@ -15419,12 +15464,12 @@ bool URigVMController::ResolveWildCardPin(URigVMPin* InPin, const FRigVMTemplate
 		const TArray<int32>& FilteredPermutations = TemplateNode->GetFilteredPermutationsIndices();
 		if (FilteredPermutations.Num() == 1)
 		{
-			const TArray<FString> NewPreferredPermutationTypes = TemplateNode->GetArgumentTypesForPermutation(FilteredPermutations[0]);
+			const TArray<FRigVMTemplatePreferredType> NewPreferredPermutationTypes = TemplateNode->GetArgumentTypesForPermutation(FilteredPermutations[0]);
 			if (bSetupUndoRedo)
 			{
 				ActionStack->AddAction(FRigVMSetPreferredTemplatePermutationsAction(TemplateNode, NewPreferredPermutationTypes));
 			}
-			TemplateNode->PreferredPermutationTypes = NewPreferredPermutationTypes;
+			TemplateNode->PreferredPermutationPairs = NewPreferredPermutationTypes;
 		}
 		if (bSetupUndoRedo)
 		{
@@ -15437,7 +15482,7 @@ bool URigVMController::ResolveWildCardPin(URigVMPin* InPin, const FRigVMTemplate
 	{
 		OpenUndoBracket(TEXT("Resolving wildcard pin"));
 		UnresolveTemplateNodes({*InPin->GetNode()->GetName()}, bSetupUndoRedo);
-		if (ResolveWildCardPin(InPin, InType, bSetupUndoRedo, bPrintPythonCommand))
+		if (ResolveWildCardPin(InPin, InTypeIndex, bSetupUndoRedo, bPrintPythonCommand))
 		{
 			CloseUndoBracket();
 			return true;
@@ -15478,7 +15523,7 @@ bool URigVMController::UpdateFilteredPermutations(URigVMPin* InPin, URigVMPin* I
 	return false;
 }
 
-bool URigVMController::UpdateFilteredPermutations(URigVMPin* InPin, const TArray<FRigVMTemplateArgumentType>& InTypes, bool bSetupUndoRedo)
+bool URigVMController::UpdateFilteredPermutations(URigVMPin* InPin, const TArray<int32>& InTypeIndices, bool bSetupUndoRedo)
 {
 	URigVMTemplateNode* Node = Cast<URigVMTemplateNode>(InPin->GetNode());
 	if (!Node)
@@ -15492,7 +15537,7 @@ bool URigVMController::UpdateFilteredPermutations(URigVMPin* InPin, const TArray
 		OldPermutations = Node->GetFilteredPermutationsIndices();
 	}
 
-	if (Node->UpdateFilteredPermutations(InPin, InTypes))
+	if (Node->UpdateFilteredPermutations(InPin, InTypeIndices))
 	{
 		if (bSetupUndoRedo)
 		{
@@ -15519,7 +15564,7 @@ bool URigVMController::UpdateTemplateNodePinTypes(URigVMTemplateNode* InNode, bo
 			continue;
 		}
 		
-		TArray<FRigVMTemplateArgumentType> Types = InNode->GetFilteredTypesForPin(Pin);
+		TArray<int32> Types = InNode->GetFilteredTypesForPin(Pin);
 		if (Types.IsEmpty())
 		{
 			continue;
@@ -15531,7 +15576,7 @@ bool URigVMController::UpdateTemplateNodePinTypes(URigVMTemplateNode* InNode, bo
 			bool bCanReduceToSingleType = true;
 			for (int32 i=1; i<Types.Num(); ++i)
 			{
-				if (!RigVMTypeUtils::AreCompatible(Types[0].CPPType, Types[0].CPPTypeObject, Types[i].CPPType, Types[i].CPPTypeObject))
+				if (!FRigVMRegistry::Get().CanMatchTypes(Types[0], Types[i], true))
 				{
 					bCanReduceToSingleType = false;
 					break;
@@ -15546,7 +15591,7 @@ bool URigVMController::UpdateTemplateNodePinTypes(URigVMTemplateNode* InNode, bo
 				{
 					for (int32 i=0; i<Types.Num(); ++i)
 					{
-						if (Types[i].CPPType == Pin->GetCPPType())
+						if (Types[i] == Pin->GetTypeIndex())
 						{
 							PreferredIndex = i;
 							break;
@@ -15589,12 +15634,10 @@ bool URigVMController::UpdateTemplateNodePinTypes(URigVMTemplateNode* InNode, bo
 		else if (Types.Num() == 1)
 		{
 			// Resolve
-			FString CPPType = Types[0].CPPType;
-			UObject* CPPObjectType = Types[0].CPPTypeObject;
-			if (Pin->GetCPPType() != CPPType || Pin->GetCPPTypeObject() != CPPObjectType)
+			if (Pin->GetTypeIndex() != Types[0])
 			{
 				bAnyTypeChanged = true;
-				ChangePinType(Pin, CPPType, CPPObjectType, bSetupUndoRedo, false, false, false);
+				ChangePinType(Pin, Types[0], bSetupUndoRedo, false, false, false);
 			}
 		}
 		else
@@ -15645,7 +15688,7 @@ bool URigVMController::PropagateTemplateFilteredTypes(URigVMTemplateNode* InNode
 
 			if (!bIsTemplate)
 			{
-				if (!InNode->FilteredSupportsType(Pin, OtherPin->GetCPPType()))
+				if (!InNode->FilteredSupportsType(Pin, OtherPin->GetTypeIndex()))
 				{
 					URigVMLink* Link = Pin->FindLinkForPin(OtherPin);
 					ensureMsgf(!ActionStack->BracketActions.IsEmpty(), TEXT("Unexpected link broken %s in package %s"), *Link->GetPinPathRepresentation(), *GetPackage()->GetPathName());
@@ -15697,7 +15740,7 @@ bool URigVMController::RecomputeAllTemplateFilteredTypes(bool bSetupUndoRedo)
 	check(Graph);
 
 	// Save all template pin types, in case we need them after initializing
-	TMap<URigVMPin*, FRigVMTemplateArgumentType> TypesBeforeRecomputing;
+	TMap<URigVMPin*, int32> TypesBeforeRecomputing;
 	for (URigVMNode* Node : Graph->GetNodes())
 	{
 		if (URigVMTemplateNode* TemplateNode = Cast<URigVMTemplateNode>(Node))
@@ -15707,7 +15750,7 @@ bool URigVMController::RecomputeAllTemplateFilteredTypes(bool bSetupUndoRedo)
 				continue;
 			}
 
-			if (!TemplateNode->PreferredPermutationTypes.IsEmpty())
+			if (!TemplateNode->PreferredPermutationPairs.IsEmpty())
 			{
 				continue;
 			}
@@ -15723,7 +15766,7 @@ bool URigVMController::RecomputeAllTemplateFilteredTypes(bool bSetupUndoRedo)
 						continue;
 					}
 
-					TypesBeforeRecomputing.Add(Pin, FRigVMTemplateArgumentType(Pin->GetCPPType(), Pin->GetCPPTypeObject()));
+					TypesBeforeRecomputing.Add(Pin, Pin->GetTypeIndex());
 				}
 			}
 		}
@@ -15781,13 +15824,13 @@ bool URigVMController::RecomputeAllTemplateFilteredTypes(bool bSetupUndoRedo)
 		// If pin is a struct member, we should resolve to that type
 		if (URigVMTemplateNode* OutputNode = Cast<URigVMTemplateNode>(OutputPin->GetNode()))
 		{
-			if (!OutputNode->IsSingleton() && OutputNode->PreferredPermutationTypes.IsEmpty())
+			if (!OutputNode->IsSingleton() && OutputNode->PreferredPermutationPairs.IsEmpty())
 			{
 				if (OutputPin->IsStructMember())
 				{
 					URigVMPin* RootPin = OutputPin->GetRootPin();
-					const FRigVMTemplateArgumentType Type = TypesBeforeRecomputing.FindChecked(RootPin);
-					if (UpdateFilteredPermutations(RootPin, {Type}, bSetupUndoRedo))
+					const int32 TypeIndex = TypesBeforeRecomputing.FindChecked(RootPin);
+					if (UpdateFilteredPermutations(RootPin, {TypeIndex}, bSetupUndoRedo))
 					{
 						PropagateTemplateFilteredTypes(OutputNode, bSetupUndoRedo);
 					}
@@ -15800,11 +15843,11 @@ bool URigVMController::RecomputeAllTemplateFilteredTypes(bool bSetupUndoRedo)
 		{
 			if (!InputNode->IsSingleton())
 			{
-				if (InputPin->IsStructMember() && InputNode->PreferredPermutationTypes.IsEmpty())
+				if (InputPin->IsStructMember() && InputNode->PreferredPermutationPairs.IsEmpty())
 				{
 					URigVMPin* RootPin = InputPin->GetRootPin();
-					const FRigVMTemplateArgumentType Type = TypesBeforeRecomputing.FindChecked(RootPin);
-					if (UpdateFilteredPermutations(RootPin, {Type}, bSetupUndoRedo))
+					const int32 TypeIndex = TypesBeforeRecomputing.FindChecked(RootPin);
+					if (UpdateFilteredPermutations(RootPin, {TypeIndex}, bSetupUndoRedo))
 					{
 						PropagateTemplateFilteredTypes(InputNode, bSetupUndoRedo);
 					}
@@ -15988,36 +16031,49 @@ bool URigVMController::ChangePinType(URigVMPin* InPin, const FString& InCPPType,
 		return false;
 	}
 
-	FName CPPTypeObjectPath(NAME_None);
-	if(InCPPTypeObject)
+	if (FRigVMPropertyDescription::RequiresCPPTypeObject(InCPPType) && !InCPPTypeObject)
 	{
-		CPPTypeObjectPath = *InCPPTypeObject->GetPathName();
+		return false;
 	}
 
-	if (FRigVMPropertyDescription::RequiresCPPTypeObject(InCPPType) && !InCPPTypeObject)
+	const FRigVMTemplateArgumentType Type(*InCPPType, InCPPTypeObject);
+	const int32 TypeIndex = FRigVMRegistry::Get().GetTypeIndex(Type);
+	return ChangePinType(InPin, TypeIndex, bSetupUndoRedo, bSetupOrphanPins, bBreakLinks, bRemoveSubPins);
+}
+
+bool URigVMController::ChangePinType(URigVMPin* InPin, int32 InTypeIndex, bool bSetupUndoRedo, bool bSetupOrphanPins, bool bBreakLinks, bool bRemoveSubPins)
+{
+	if (!bIsTransacting && !IsGraphEditable())
+	{
+		return false;
+	}
+	
+	if (InTypeIndex == INDEX_NONE)
 	{
 		return false;
 	}
 
 	// only allow valid pin cpp types on template nodes
-	FString CPPType = InCPPType;
+	int32 TypeIndex = InTypeIndex;
 	if(URigVMTemplateNode* TemplateNode = Cast<URigVMTemplateNode>(InPin->GetNode()))
 	{
-		if(!TemplateNode->SupportsType(InPin, InCPPType, &CPPType))
+		if(!TemplateNode->SupportsType(InPin, InTypeIndex, &TypeIndex))
 		{
-			ReportErrorf(TEXT("ChangePinType: %s doesn't support type '%s'."), *InPin->GetPinPath(), *InCPPType);
+			const FRigVMTemplateArgumentType Type = FRigVMRegistry::Get().GetType(InTypeIndex);
+			ReportErrorf(TEXT("ChangePinType: %s doesn't support type '%s'."), *InPin->GetPinPath(), *Type.CPPType.ToString());
 			return false;
 		}
 
 		// change the pin's type to be of an array as well
-		if(InPin->IsRootPin() && RigVMTypeUtils::IsArrayType(CPPType) != InPin->IsArray())
+		const bool bIsArrayType = FRigVMRegistry::Get().IsArrayType(TypeIndex);
+		if(InPin->IsRootPin() && bIsArrayType != InPin->IsArray())
 		{
 			// nothing to do here - leave the type as is 
 		}
 		else
 		{
-			const FString BaseCPPType = RigVMTypeUtils::IsArrayType(CPPType) ? RigVMTypeUtils::BaseTypeFromArrayType(CPPType) : CPPType;
-			CPPType = InPin->IsArray() ? RigVMTypeUtils::ArrayTypeFromBaseType(BaseCPPType) : BaseCPPType;
+			const int32 BaseTypeIndex = bIsArrayType ? FRigVMRegistry::Get().GetBaseTypeFromArrayTypeIndex(TypeIndex) : TypeIndex;
+			TypeIndex = InPin->IsArray() ? FRigVMRegistry::Get().GetArrayTypeFromBaseTypeIndex(BaseTypeIndex) : BaseTypeIndex;
 		}
 	}
 
@@ -16077,7 +16133,7 @@ bool URigVMController::ChangePinType(URigVMPin* InPin, const FString& InCPPType,
 
 	if (bSetupUndoRedo)
 	{
-		ActionStack->AddAction(FRigVMChangePinTypeAction(InPin, CPPType, CPPTypeObjectPath, bSetupOrphanPins, bBreakLinks, bRemoveSubPins));
+		ActionStack->AddAction(FRigVMChangePinTypeAction(InPin, TypeIndex, bSetupOrphanPins, bBreakLinks, bRemoveSubPins));
 	}
 	
 	// compute the number of remaining wildcard pins
@@ -16089,10 +16145,11 @@ bool URigVMController::ChangePinType(URigVMPin* InPin, const FString& InCPPType,
 	const FPinState PreviousPinState = GetPinState(InPin);
 	const FString PreviousCPPType = InPin->CPPType;
 	
-	InPin->CPPType = CPPType;
-	InPin->CPPTypeObjectPath = CPPTypeObjectPath;
-	InPin->CPPTypeObject = InCPPTypeObject;
-	InPin->bIsDynamicArray = RigVMTypeUtils::IsArrayType(CPPType);
+	const FRigVMTemplateArgumentType Type = FRigVMRegistry::Get().GetType(TypeIndex);
+	InPin->CPPType = Type.CPPType.ToString();
+	InPin->CPPTypeObjectPath = Type.GetCPPTypeObjectPath();
+	InPin->CPPTypeObject = Type.CPPTypeObject;
+	InPin->bIsDynamicArray = FRigVMRegistry::Get().IsArrayType(TypeIndex);
 	// we might want to use GetPinInitialDefaultValue here for a better default value
 	InPin->DefaultValue = FString();
 
@@ -16118,7 +16175,7 @@ bool URigVMController::ChangePinType(URigVMPin* InPin, const FString& InCPPType,
 
 	if (InPin->IsArray())
 	{
-		const FString BaseCPPType = RigVMTypeUtils::BaseTypeFromArrayType(CPPType);
+		const int32 BaseTypeIndex = FRigVMRegistry::Get().GetBaseTypeFromArrayTypeIndex(TypeIndex);
 		for (int32 i=0; i<InPin->GetSubPins().Num(); ++i)
 		{
 			URigVMPin* SubPin = InPin->GetSubPins()[i];
@@ -16126,7 +16183,7 @@ bool URigVMController::ChangePinType(URigVMPin* InPin, const FString& InCPPType,
 			{
 				continue;
 			}
-			ChangePinType(SubPin, BaseCPPType, InCPPTypeObject, bSetupUndoRedo, bSetupOrphanPins, bBreakLinks, bRemoveSubPins);
+			ChangePinType(SubPin, BaseTypeIndex, bSetupUndoRedo, bSetupOrphanPins, bBreakLinks, bRemoveSubPins);
 		}
 	}
 
diff --git a/Engine/Source/Developer/RigVMDeveloper/Private/RigVMModel/RigVMControllerActions.cpp b/Engine/Source/Developer/RigVMDeveloper/Private/RigVMModel/RigVMControllerActions.cpp
--- a/Engine/Source/Developer/RigVMDeveloper/Private/RigVMModel/RigVMControllerActions.cpp
+++ b/Engine/Source/Developer/RigVMDeveloper/Private/RigVMModel/RigVMControllerActions.cpp
@@ -788,7 +788,7 @@ FRigVMRemoveNodeAction::FRigVMRemoveNodeAction(URigVMNode* InNode, URigVMControl
 		else
 		{
 			InverseAction.AddAction(FRigVMAddTemplateNodeAction(TemplateNode), InController);
-			InverseAction.AddAction(FRigVMSetPreferredTemplatePermutationsAction(TemplateNode, TemplateNode->PreferredPermutationTypes), InController);
+			InverseAction.AddAction(FRigVMSetPreferredTemplatePermutationsAction(TemplateNode, TemplateNode->PreferredPermutationPairs), InController);
 			InverseAction.AddAction(FRigVMSetTemplateFilteredPermutationsAction(TemplateNode, nullptr, {}), InController);
 
 			for (URigVMPin* Pin : TemplateNode->GetPins())
@@ -801,7 +801,7 @@ FRigVMRemoveNodeAction::FRigVMRemoveNodeAction(URigVMNode* InNode, URigVMControl
 						{
 							if (!Argument->IsSingleton())
 							{
-								InverseAction.AddAction(FRigVMChangePinTypeAction(Pin, Pin->GetCPPType(), *Pin->GetCPPTypeObject()->GetPathName(), false, false, false), InController);
+								InverseAction.AddAction(FRigVMChangePinTypeAction(Pin, Pin->GetTypeIndex(), false, false, false), InController);
 							}
 						}
 					}			
@@ -1418,10 +1418,10 @@ bool FRigVMSetTemplateFilteredPermutationsAction::Redo(URigVMController* InContr
 	return FRigVMBaseAction::Redo(InController);
 }
 
-FRigVMSetPreferredTemplatePermutationsAction::FRigVMSetPreferredTemplatePermutationsAction(URigVMTemplateNode* InNode, const TArray<FString>& InPreferredPermutationTypes)
+FRigVMSetPreferredTemplatePermutationsAction::FRigVMSetPreferredTemplatePermutationsAction(URigVMTemplateNode* InNode, const TArray<FRigVMTemplatePreferredType>& InPreferredTypes)
 : NodePath(InNode->GetNodePath())
-, OldPreferredPermutationTypes(InNode->PreferredPermutationTypes)
-, NewPreferredPermutationTypes(InPreferredPermutationTypes)
+, OldPreferredPermutationTypes(InNode->PreferredPermutationPairs)
+, NewPreferredPermutationTypes(InPreferredTypes)
 {
 	
 }
@@ -1452,7 +1452,7 @@ bool FRigVMSetPreferredTemplatePermutationsAction::Undo(URigVMController* InCont
 
 	if (URigVMTemplateNode* TemplateNode = Cast<URigVMTemplateNode>(InController->GetGraph()->FindNode(NodePath)))
 	{
-		TemplateNode->PreferredPermutationTypes = OldPreferredPermutationTypes;
+		TemplateNode->PreferredPermutationPairs = OldPreferredPermutationTypes;
 		return true;
 	}
 	return false;
@@ -1465,7 +1465,7 @@ bool FRigVMSetPreferredTemplatePermutationsAction::Redo(URigVMController* InCont
 	{
 		return false;
 	}
-	TemplateNode->PreferredPermutationTypes = NewPreferredPermutationTypes;
+	TemplateNode->PreferredPermutationPairs = NewPreferredPermutationTypes;
 	return FRigVMBaseAction::Redo(InController);
 }
 
@@ -1569,29 +1569,21 @@ bool FRigVMBreakLinkAction::Redo(URigVMController* InController)
 
 FRigVMChangePinTypeAction::FRigVMChangePinTypeAction()
 : PinPath()
-, OldCPPType()
-, OldCPPTypeObjectPath(NAME_None)
-, NewCPPType()
-, NewCPPTypeObjectPath(NAME_None)
+, OldTypeIndex(INDEX_NONE)
+, NewTypeIndex(INDEX_NONE)
 , bSetupOrphanPins(true)
 , bBreakLinks(true)
 , bRemoveSubPins(true)
 {}
 
-FRigVMChangePinTypeAction::FRigVMChangePinTypeAction(URigVMPin* InPin, const FString& InCppType, const FName& InCppTypeObjectPath, bool InSetupOrphanPins, bool InBreakLinks, bool InRemoveSubPins)
+FRigVMChangePinTypeAction::FRigVMChangePinTypeAction(URigVMPin* InPin, int32 InTypeIndex, bool InSetupOrphanPins, bool InBreakLinks, bool InRemoveSubPins)
 : PinPath(InPin->GetPinPath())
-, OldCPPType(InPin->GetCPPType())
-, OldCPPTypeObjectPath(NAME_None)
-, NewCPPType(InCppType)
-, NewCPPTypeObjectPath(InCppTypeObjectPath)
+, OldTypeIndex(InPin->GetTypeIndex())
+, NewTypeIndex(InTypeIndex)
 , bSetupOrphanPins(InSetupOrphanPins)
 , bBreakLinks(InBreakLinks)
 , bRemoveSubPins(InRemoveSubPins)
 {
-	if (UObject* CPPTypeObject = InPin->GetCPPTypeObject())
-	{
-		OldCPPTypeObjectPath = *CPPTypeObject->GetPathName();
-	}
 }
 
 bool FRigVMChangePinTypeAction::Undo(URigVMController* InController)
@@ -1600,13 +1592,34 @@ bool FRigVMChangePinTypeAction::Undo(URigVMController* InController)
 	{
 		return false;
 	}
-	
-	return InController->ChangePinType(PinPath, OldCPPType, OldCPPTypeObjectPath, false, bSetupOrphanPins, bBreakLinks, bRemoveSubPins);
+
+	if(const URigVMGraph* Graph = InController->GetGraph())
+	{
+		if(URigVMPin* Pin = Graph->FindPin(PinPath))
+		{
+			return InController->ChangePinType(Pin, OldTypeIndex, false, bSetupOrphanPins, bBreakLinks, bRemoveSubPins);
+		}
+	}
+	return false;
 }
 
 bool FRigVMChangePinTypeAction::Redo(URigVMController* InController)
 {
-	if(!InController->ChangePinType(PinPath, NewCPPType, NewCPPTypeObjectPath, false, bSetupOrphanPins, bBreakLinks, bRemoveSubPins))
+	if(const URigVMGraph* Graph = InController->GetGraph())
+	{
+		if(URigVMPin* Pin = Graph->FindPin(PinPath))
+		{
+			if(!InController->ChangePinType(Pin, NewTypeIndex, false, bSetupOrphanPins, bBreakLinks, bRemoveSubPins))
+			{
+				return false;
+			}
+		}
+		else
+		{
+			return false;
+		}
+	}
+	else
 	{
 		return false;
 	}
diff --git a/Engine/Source/Developer/RigVMDeveloper/Private/RigVMModel/RigVMPin.cpp b/Engine/Source/Developer/RigVMDeveloper/Private/RigVMModel/RigVMPin.cpp
--- a/Engine/Source/Developer/RigVMDeveloper/Private/RigVMModel/RigVMPin.cpp
+++ b/Engine/Source/Developer/RigVMDeveloper/Private/RigVMModel/RigVMPin.cpp
@@ -191,6 +191,7 @@ URigVMPin::URigVMPin()
 	, CPPTypeObjectPath(NAME_None)
 	, DefaultValue(FString())
 	, BoundVariablePath_DEPRECATED()
+	, LastKnownTypeIndex(INDEX_NONE)
 {
 #if UE_BUILD_DEBUG
 	CachedPinPath = GetPinPath();
@@ -612,7 +613,23 @@ FString URigVMPin::GetArrayElementCppType() const
 
 FRigVMTemplateArgumentType URigVMPin::GetTemplateArgumentType() const
 {
-	return FRigVMTemplateArgumentType(GetCPPType(), GetCPPTypeObject());
+	return FRigVMRegistry::Get().GetType(GetTypeIndex());
+}
+
+int32 URigVMPin::GetTypeIndex() const
+{
+	if(LastKnownCPPType != GetCPPType())
+	{
+		LastKnownTypeIndex = INDEX_NONE;
+	}
+	if(LastKnownTypeIndex == INDEX_NONE)
+	{
+		LastKnownCPPType = GetCPPType();
+		const FRigVMTemplateArgumentType Type(*LastKnownCPPType, GetCPPTypeObject());
+		LastKnownTypeIndex = FRigVMRegistry::Get().FindOrAddType(Type);
+		ensure(LastKnownTypeIndex != INDEX_NONE);
+	}
+	return LastKnownTypeIndex;
 }
 
 bool URigVMPin::IsStringType() const
@@ -1226,7 +1243,7 @@ bool URigVMPin::CanBeBoundToVariable(const FRigVMExternalVariable& InExternalVar
 	}
 
 	const FString CPPBaseType = IsArray() ? GetArrayElementCppType() : GetCPPType();
-	return RigVMTypeUtils::AreCompatible(CPPBaseType, GetCPPTypeObject(), ExternalCPPType.ToString(), ExternalCPPTypeObject);
+	return RigVMTypeUtils::AreCompatible(*CPPBaseType, GetCPPTypeObject(), ExternalCPPType, ExternalCPPTypeObject);
 }
 
 bool URigVMPin::ShowInDetailsPanelOnly() const
@@ -1298,6 +1315,9 @@ void URigVMPin::UpdateTypeInformationIfRequired() const
 		{
 			URigVMPin* MutableThis = (URigVMPin*)this;
 			MutableThis->CPPTypeObject = FindObjectFromCPPTypeObjectPath(CPPTypeObjectPath.ToString());
+			MutableThis->CPPType = RigVMTypeUtils::PostProcessCPPType(CPPType, CPPTypeObject);
+			MutableThis->LastKnownTypeIndex = FRigVMRegistry::Get().FindOrAddType(FRigVMTemplateArgumentType(*CPPType, CPPTypeObject));
+			MutableThis->LastKnownCPPType = MutableThis->CPPType; 
 		}
 	}
 
@@ -1678,7 +1698,7 @@ bool URigVMPin::CanLink(URigVMPin* InSourcePin, URigVMPin* InTargetPin, FString*
 		static const FString Float = TEXT("float");
 		static const FString Double = TEXT("double");
 
-		if (RigVMTypeUtils::AreCompatible(InSourcePin->CPPType, InSourcePin->CPPTypeObject, InTargetPin->CPPType, InTargetPin->CPPTypeObject))
+		if (FRigVMRegistry::Get().CanMatchTypes(InSourcePin->GetTypeIndex(), InTargetPin->GetTypeIndex(), true))
 		{
 			bCPPTypesDiffer = false;
 		}
@@ -1686,11 +1706,11 @@ bool URigVMPin::CanLink(URigVMPin* InSourcePin, URigVMPin* InTargetPin, FString*
 		if (bCPPTypesDiffer)
 		{
 			{
-				auto TemplateNodeSupportsType = [](URigVMPin* InPin, const FString& InCPPType, FString* OutFailureReason) -> bool
+				auto TemplateNodeSupportsType = [](URigVMPin* InPin, const int32& InTypeIndex, FString* OutFailureReason) -> bool
 				{
 					if (URigVMTemplateNode* TemplateNode = Cast<URigVMTemplateNode>(InPin->GetNode()))
 					{
-						if (TemplateNode->SupportsType(InPin, InCPPType))
+						if (TemplateNode->SupportsType(InPin, InTypeIndex))
 						{
 							if (OutFailureReason)
 							{
@@ -1733,11 +1753,11 @@ bool URigVMPin::CanLink(URigVMPin* InSourcePin, URigVMPin* InTargetPin, FString*
 				
 				if(IsPinValidForTypeChange(InSourcePin, false, InUserLinkDirection))
 				{
-					return TemplateNodeSupportsType(InSourcePin, InTargetPin->GetCPPType(), OutFailureReason);
+					return TemplateNodeSupportsType(InSourcePin, InTargetPin->GetTypeIndex(), OutFailureReason);
 				}
 				else if(IsPinValidForTypeChange(InTargetPin, true, InUserLinkDirection))
 				{
-					return TemplateNodeSupportsType(InTargetPin, InSourcePin->GetCPPType(), OutFailureReason);
+					return TemplateNodeSupportsType(InTargetPin, InSourcePin->GetTypeIndex(), OutFailureReason);
 				}
 			}
 		}
diff --git a/Engine/Source/Developer/RigVMDeveloper/Public/RigVMDeveloperTypeUtils.h b/Engine/Source/Developer/RigVMDeveloper/Public/RigVMDeveloperTypeUtils.h
--- a/Engine/Source/Developer/RigVMDeveloper/Public/RigVMDeveloperTypeUtils.h
+++ b/Engine/Source/Developer/RigVMDeveloper/Public/RigVMDeveloperTypeUtils.h
@@ -23,7 +23,7 @@ namespace RigVMTypeUtils
 
 	RIGVMDEVELOPER_API FRigVMExternalVariable ExternalVariableFromCPPType(const FName& InName, const FString& InCPPType, UObject* InCPPTypeObject, bool bInPublic = false, bool bInReadonly = false);
 
-	RIGVMDEVELOPER_API FEdGraphPinType PinTypeFromCPPType(const FString& InCPPType, UObject* InCPPTypeObject);
+	RIGVMDEVELOPER_API FEdGraphPinType PinTypeFromCPPType(const FName& InCPPType, UObject* InCPPTypeObject);
 
 	RIGVMDEVELOPER_API FEdGraphPinType PinTypeFromExternalVariable(const FRigVMExternalVariable& InExternalVariable);
 
@@ -45,7 +45,7 @@ namespace RigVMTypeUtils
 
 	RIGVMDEVELOPER_API bool CPPTypeFromExternalVariable(const FRigVMExternalVariable& InExternalVariable, FString& OutCPPType, UObject** OutCPPTypeObject);
 
-	RIGVMDEVELOPER_API bool AreCompatible(const FString& InCPPTypeA, UObject* InCPPTypeObjectA, const FString& InCPPTypeB, UObject* InCPPTypeObjectB);
+	RIGVMDEVELOPER_API bool AreCompatible(const FName& InCPPTypeA, UObject* InCPPTypeObjectA, const FName& InCPPTypeB, UObject* InCPPTypeObjectB);
 
 	RIGVMDEVELOPER_API bool AreCompatible(const FRigVMExternalVariable& InTypeA, const FRigVMExternalVariable& InTypeB, const FString& InSegmentPathA = FString(), const FString& InSegmentPathB = FString());
 	
diff --git a/Engine/Source/Developer/RigVMDeveloper/Public/RigVMModel/Nodes/RigVMTemplateNode.h b/Engine/Source/Developer/RigVMDeveloper/Public/RigVMModel/Nodes/RigVMTemplateNode.h
--- a/Engine/Source/Developer/RigVMDeveloper/Public/RigVMModel/Nodes/RigVMTemplateNode.h
+++ b/Engine/Source/Developer/RigVMDeveloper/Public/RigVMModel/Nodes/RigVMTemplateNode.h
@@ -6,6 +6,34 @@
 #include "RigVMCore/RigVMTemplate.h"
 #include "RigVMTemplateNode.generated.h"
 
+USTRUCT()
+struct RIGVMDEVELOPER_API FRigVMTemplatePreferredType
+{
+	GENERATED_BODY()
+	
+	FRigVMTemplatePreferredType()
+		: Argument(NAME_None)
+		, TypeIndex(INDEX_NONE)
+	{}
+	
+	FRigVMTemplatePreferredType(const FName& InArgument, int32 InTypeIndex)
+		: Argument(InArgument)
+		, TypeIndex(InTypeIndex)
+	{}
+	
+	void UpdateStringFromIndex();
+	void UpdateIndexFromString();
+
+	UPROPERTY()
+	FName Argument;
+
+	UPROPERTY()
+	int32 TypeIndex;
+
+	UPROPERTY()
+	FString TypeString;
+};
+
 /**
  * The Template Node represents an unresolved function.
  * Template nodes can morph into all functions implementing
@@ -22,6 +50,7 @@ public:
 	URigVMTemplateNode();
 
 	// UObject interface
+	virtual void PreSave(FObjectPreSaveContext SaveContext) override;
 	virtual void PostLoad() override;
 
 	// URigVMNode interface
@@ -43,10 +72,10 @@ public:
 	virtual bool IsSingleton() const;
 
 	// returns true if a pin supports a given type
-	bool SupportsType(const URigVMPin* InPin, const FString& InCPPType, FString* OutCPPType = nullptr);
+	bool SupportsType(const URigVMPin* InPin, int32 InTypeIndex, int32* OutTypeIndex = nullptr);
 
 	// returns true if a pin supports a given type after filtering
-	bool FilteredSupportsType(const URigVMPin* InPIn, const FString& InCPPType, FString* OutCPPType = nullptr, bool bAllowFloatingPointCasts = true);
+	bool FilteredSupportsType(const URigVMPin* InPin, int32 InTypeIndex, int32* OutTypeIndex = nullptr, bool bAllowFloatingPointCasts = true);
 
 	// returns the resolved functions for the template
 	TArray<const FRigVMFunction*> GetResolvedPermutations() const;
@@ -75,14 +104,14 @@ public:
 	const TArray<int32>& GetFilteredPermutationsIndices() const;
 
 	// returns the filtered types of this pin
-	TArray<FRigVMTemplateArgumentType> GetFilteredTypesForPin(URigVMPin* InPin) const;
+	TArray<int32> GetFilteredTypesForPin(URigVMPin* InPin) const;
 
 	// returns true if updating pin filters with InTypes would result in different filters 
-	bool PinNeedsFilteredTypesUpdate(URigVMPin* InPin, const TArray<FRigVMTemplateArgumentType>& InTypes);
+	bool PinNeedsFilteredTypesUpdate(URigVMPin* InPin, const TArray<int32>& InTypeIndices);
 	bool PinNeedsFilteredTypesUpdate(URigVMPin* InPin, URigVMPin* LinkedPin);
 
 	// updates the filtered permutations given a link or the types for a pin
-	bool UpdateFilteredPermutations(URigVMPin* InPin, const TArray<FRigVMTemplateArgumentType>& InTypes);
+	bool UpdateFilteredPermutations(URigVMPin* InPin, const TArray<int32>& InTypeIndices);
 	bool UpdateFilteredPermutations(URigVMPin* InPin, URigVMPin* LinkedPin);
 
 	// initializes the filtered permutations to all possible permutations
@@ -90,16 +119,22 @@ public:
 
 	// Initializes the filtered permutations and preferred permutation from the types of the pins
 	void InitializeFilteredPermutationsFromTypes();
-	
+
+	// Converts the preferred types per pin from index to string
+	void ConvertPreferredTypesToString();
+
+	// Converts the preferred types per pin from string to indices
+	void ConvertPreferredTypesToTypeIndex();
+
 protected:
 
 	virtual void InvalidateCache() override;
 	
 	TArray<int32> GetNewFilteredPermutations(URigVMPin* InPin, URigVMPin* LinkedPin);
-	TArray<int32> GetNewFilteredPermutations(URigVMPin* InPin, const TArray<FRigVMTemplateArgumentType>& InTypes);
+	TArray<int32> GetNewFilteredPermutations(URigVMPin* InPin, const TArray<int32>& InTypeIndices);
 
-	TArray<int32> FindPermutationsForTypes(const TArray<FString>& ArgumentTypes, bool bAllowCasting = false);
-	TArray<FString> GetArgumentTypesForPermutation(const int32 InPermutationIndex);
+	TArray<int32> FindPermutationsForTypes(const TArray<FRigVMTemplatePreferredType>& ArgumentTypes, bool bAllowCasting = false) const;
+	TArray<FRigVMTemplatePreferredType> GetArgumentTypesForPermutation(const int32 InPermutationIndex) const;
 
 	UPROPERTY()
 	FName TemplateNotation;
@@ -110,10 +145,13 @@ protected:
 	// Indicates a preferred permuation using the types of the arguments
 	// Each element is in the format "ArgumentName:CPPType"
 	UPROPERTY()
-	TArray<FString> PreferredPermutationTypes;
-	
+	TArray<FString> PreferredPermutationTypes_DEPRECATED;
+
+	UPROPERTY()
+	TArray<FRigVMTemplatePreferredType> PreferredPermutationPairs;
+
 	TArray<int32> FilteredPermutations;
-	TMap<FString, TPair<bool, FRigVMTemplateArgumentType>> SupportedTypesCache;
+	TMap<TPair<FName,int32>, TPair<bool, int32>> SupportedTypesCache;
 
 	mutable const FRigVMTemplate* CachedTemplate;
 	mutable const FRigVMFunction* CachedFunction;
diff --git a/Engine/Source/Developer/RigVMDeveloper/Public/RigVMModel/RigVMController.h b/Engine/Source/Developer/RigVMDeveloper/Public/RigVMModel/RigVMController.h
--- a/Engine/Source/Developer/RigVMDeveloper/Public/RigVMModel/RigVMController.h
+++ b/Engine/Source/Developer/RigVMDeveloper/Public/RigVMModel/RigVMController.h
@@ -379,7 +379,9 @@ public:
 	UFUNCTION(BlueprintCallable, Category = RigVMController)
 	bool ResolveWildCardPin(const FString& InPinPath, const FString& InCPPType, const FName& InCPPTypeObjectPath, bool bSetupUndoRedo = true, bool bPrintPythonCommand = false);
 	bool ResolveWildCardPin(URigVMPin* InPin, const FRigVMTemplateArgumentType& InType, bool bSetupUndoRedo = true, bool bPrintPythonCommand = false);
-		
+	bool ResolveWildCardPin(const FString& InPinPath, int32 InTypeIndex, bool bSetupUndoRedo = true, bool bPrintPythonCommand = false);
+	bool ResolveWildCardPin(URigVMPin* InPin, int32 InTypeIndex, bool bSetupUndoRedo = true, bool bPrintPythonCommand = false);
+
 	// Adds a Function / Struct Node to the edited Graph as an injected node
 	// UnitNode represent a RIGVM_METHOD declaration on a USTRUCT.
 	// This causes a NodeAdded modified event.
@@ -1050,15 +1052,15 @@ private:
 	// If any of the types is supported (without breaking any links), then the filtered permutations will be updated and the change will
 	// propagate to other nodes in the graph.
 	// If it is not supported, we will attempt to find and break any links that do not support this change
-	bool PrepareTemplatePinForType(URigVMPin* InPin, const TArray<FRigVMTemplateArgumentType>& InTypes, bool bSetupUndoRedo);
+	bool PrepareTemplatePinForType(URigVMPin* InPin, const TArray<int32>& InTypeIndices, bool bSetupUndoRedo);
 
 	// Get filtered types for a wildcard node. If template node, that means just returning its filtered wildcard types, but if it's another type of node (select, if, rereoute), iterate
 	// its connections to figure out the filtered types
-	TArray<FRigVMTemplateArgumentType> GetWildcardFilteredTypes(URigVMPin* InPin);
+	TArray<int32> GetWildcardFilteredTypeIndices(URigVMPin* InPin);
 	
 	// Updates the permutations allowed without having to break any links
 	bool UpdateFilteredPermutations(URigVMPin* InPin, URigVMPin* InLinkedPin, bool bSetupUndoRedo);
-	bool UpdateFilteredPermutations(URigVMPin* InPin, const TArray<FRigVMTemplateArgumentType>& InTypes, bool bSetupUndoRedo);
+	bool UpdateFilteredPermutations(URigVMPin* InPin, const TArray<int32>& InTypeIndices, bool bSetupUndoRedo);
 
 	// Changes Pin types if filtered types of a pin are unique
 	bool UpdateTemplateNodePinTypes(URigVMTemplateNode* InNode, bool bSetupUndoRedo);
@@ -1070,6 +1072,7 @@ private:
 	bool ChangePinType(const FString& InPinPath, const FString& InCPPType, const FName& InCPPTypeObjectPath, bool bSetupUndoRedo, bool bSetupOrphanPins = true, bool bBreakLinks = true, bool bRemoveSubPins = true);
 	bool ChangePinType(URigVMPin* InPin, const FString& InCPPType, const FName& InCPPTypeObjectPath, bool bSetupUndoRedo, bool bSetupOrphanPins = true, bool bBreakLinks = true, bool bRemoveSubPins = true);
 	bool ChangePinType(URigVMPin* InPin, const FString& InCPPType,UObject* InCPPTypeObject, bool bSetupUndoRedo, bool bSetupOrphanPins = true, bool bBreakLinks = true, bool bRemoveSubPins = true);
+	bool ChangePinType(URigVMPin* InPin, int32 InTypeIndex, bool bSetupUndoRedo, bool bSetupOrphanPins = true, bool bBreakLinks = true, bool bRemoveSubPins = true);
 
 #if WITH_EDITOR
 	void RewireLinks(URigVMPin* OldPin, URigVMPin* NewPin, bool bAsInput, bool bSetupUndoRedo, TArray<URigVMLink*> InLinks = TArray<URigVMLink*>());
diff --git a/Engine/Source/Developer/RigVMDeveloper/Public/RigVMModel/RigVMControllerActions.h b/Engine/Source/Developer/RigVMDeveloper/Public/RigVMModel/RigVMControllerActions.h
--- a/Engine/Source/Developer/RigVMDeveloper/Public/RigVMModel/RigVMControllerActions.h
+++ b/Engine/Source/Developer/RigVMDeveloper/Public/RigVMModel/RigVMControllerActions.h
@@ -1158,7 +1158,7 @@ struct FRigVMSetPreferredTemplatePermutationsAction : public FRigVMBaseAction
 public:
 
 	FRigVMSetPreferredTemplatePermutationsAction() {}
-	FRigVMSetPreferredTemplatePermutationsAction(URigVMTemplateNode* InNode, const TArray<FString>& InPreferredPermutation);
+	FRigVMSetPreferredTemplatePermutationsAction(URigVMTemplateNode* InNode, const TArray<FRigVMTemplatePreferredType>& InPreferredTypes);
 	virtual ~FRigVMSetPreferredTemplatePermutationsAction() {};
 	virtual UScriptStruct* GetScriptStruct() const override { return FRigVMSetTemplateFilteredPermutationsAction::StaticStruct(); }
 	virtual bool Merge(const FRigVMBaseAction* Other);
@@ -1169,10 +1169,10 @@ public:
 	FString NodePath;
 
 	UPROPERTY()
-	TArray<FString> OldPreferredPermutationTypes;
+	TArray<FRigVMTemplatePreferredType> OldPreferredPermutationTypes;
 
 	UPROPERTY()
-	TArray<FString> NewPreferredPermutationTypes;
+	TArray<FRigVMTemplatePreferredType> NewPreferredPermutationTypes;
 };
 
 /**
@@ -1294,7 +1294,7 @@ struct FRigVMChangePinTypeAction : public FRigVMBaseAction
 public:
 
 	FRigVMChangePinTypeAction();
-	FRigVMChangePinTypeAction(URigVMPin* InPin, const FString& InCppType, const FName& InCppTypeObjectPath, bool InSetupOrphanPins, bool InBreakLinks, bool InRemoveSubPins);
+	FRigVMChangePinTypeAction(URigVMPin* InPin, int32 InTypeIndex, bool InSetupOrphanPins, bool InBreakLinks, bool InRemoveSubPins);
 	virtual ~FRigVMChangePinTypeAction() {};
 	virtual UScriptStruct* GetScriptStruct() const override { return FRigVMChangePinTypeAction::StaticStruct(); }
 	virtual bool Undo(URigVMController* InController) override;
@@ -1304,16 +1304,10 @@ public:
 	FString PinPath;
 
 	UPROPERTY()
-	FString OldCPPType;
+	int32 OldTypeIndex;
 
 	UPROPERTY()
-	FName OldCPPTypeObjectPath;
-
-	UPROPERTY()
-	FString NewCPPType;
-
-	UPROPERTY()
-	FName NewCPPTypeObjectPath;
+	int32 NewTypeIndex;
 
 	UPROPERTY()
 	bool bSetupOrphanPins;
diff --git a/Engine/Source/Developer/RigVMDeveloper/Public/RigVMModel/RigVMPin.h b/Engine/Source/Developer/RigVMDeveloper/Public/RigVMModel/RigVMPin.h
--- a/Engine/Source/Developer/RigVMDeveloper/Public/RigVMModel/RigVMPin.h
+++ b/Engine/Source/Developer/RigVMDeveloper/Public/RigVMModel/RigVMPin.h
@@ -242,6 +242,9 @@ public:
 	// Returns the argument type this pin would represent within a template
 	FRigVMTemplateArgumentType GetTemplateArgumentType() const;
 
+	// Returns the argument type index this pin would represent within a template
+	int32 GetTypeIndex() const;
+
 	// Returns true if the C++ data type is FString or FName
 	UFUNCTION(BlueprintCallable, Category = RigVMPin)
 	bool IsStringType() const;
@@ -509,6 +512,9 @@ private:
 	UPROPERTY()
 	FString BoundVariablePath_DEPRECATED;
 
+	mutable FString LastKnownCPPType;
+	mutable int32 LastKnownTypeIndex;
+
 	static const FString OrphanPinPrefix;
 
 	friend class URigVMController;
diff --git a/Engine/Source/Runtime/RigVM/Private/RigVMCore/RigVMFunction.cpp b/Engine/Source/Runtime/RigVM/Private/RigVMCore/RigVMFunction.cpp
--- a/Engine/Source/Runtime/RigVM/Private/RigVMCore/RigVMFunction.cpp
+++ b/Engine/Source/Runtime/RigVM/Private/RigVMCore/RigVMFunction.cpp
@@ -3,6 +3,7 @@
 #include "RigVMCore/RigVMFunction.h"
 #include "RigVMCore/RigVMTemplate.h"
 #include "RigVMCore/RigVMRegistry.h"
+#include "RigVMCore/RigVMExternalVariable.h"
 
 FString FRigVMFunction::GetName() const
 {
@@ -49,6 +50,26 @@ FString FRigVMFunction::GetModuleRelativeHeaderPath() const
 	return FString();
 }
 
+const TArray<int32>& FRigVMFunction::GetArgumentTypeIndices() const
+{
+	if(ArgumentTypeIndices.IsEmpty() && !Arguments.IsEmpty())
+	{
+		for(const FRigVMFunctionArgument& Argument : Arguments)
+		{
+			if(const FProperty* Property = Struct->FindPropertyByName(Argument.Name))
+			{
+				FName CPPType = NAME_None;
+				UObject* CPPTypeObject = nullptr;
+				FRigVMExternalVariable::GetTypeFromProperty(Property, CPPType, CPPTypeObject);
+
+				const FRigVMTemplateArgumentType Type(CPPType, CPPTypeObject);
+				ArgumentTypeIndices.Add(FRigVMRegistry::Get().FindOrAddType(Type));
+			}
+		}
+	}
+	return ArgumentTypeIndices;
+}
+
 bool FRigVMFunction::IsAdditionalArgument(const FRigVMFunctionArgument& InArgument) const
 {
 #if WITH_EDITOR
diff --git a/Engine/Source/Runtime/RigVM/Private/RigVMCore/RigVMRegistry.cpp b/Engine/Source/Runtime/RigVM/Private/RigVMCore/RigVMRegistry.cpp
--- a/Engine/Source/Runtime/RigVM/Private/RigVMCore/RigVMRegistry.cpp
+++ b/Engine/Source/Runtime/RigVM/Private/RigVMCore/RigVMRegistry.cpp
@@ -2,6 +2,7 @@
 
 #include "RigVMCore/RigVMRegistry.h"
 #include "RigVMCore/RigVMStruct.h"
+#include "RigVMTypeUtils.h"
 #include "RigVMModule.h"
 #include "UObject/UObjectIterator.h"
 
@@ -10,13 +11,572 @@ const FName FRigVMRegistry::TemplateNameMetaName = TEXT("TemplateName");
 
 FRigVMRegistry& FRigVMRegistry::Get()
 {
+	s_RigVMRegistry.InitializeIfNeeded();
 	return s_RigVMRegistry;
 }
 
+void FRigVMRegistry::InitializeIfNeeded()
+{
+	if(!Types.IsEmpty())
+	{
+		return;
+	}
+
+	Types.Reserve(512);
+	TypeToIndex.Reserve(512);
+	TypesPerCategory.Reserve(18);
+	
+	TypesPerCategory.Add(ETypeCategory_SingleAnyValue, TArray<int32>()).Reserve(256);
+	TypesPerCategory.Add(ETypeCategory_ArrayAnyValue, TArray<int32>()).Reserve(256);
+	TypesPerCategory.Add(ETypeCategory_ArrayArrayAnyValue, TArray<int32>()).Reserve(256);
+	TypesPerCategory.Add(ETypeCategory_SingleSimpleValue, TArray<int32>()).Reserve(8);
+	TypesPerCategory.Add(ETypeCategory_ArraySimpleValue, TArray<int32>()).Reserve(8);
+	TypesPerCategory.Add(ETypeCategory_ArrayArraySimpleValue, TArray<int32>()).Reserve(8);
+	TypesPerCategory.Add(ETypeCategory_SingleMathStructValue, TArray<int32>()).Reserve(MathTypes.Num());
+	TypesPerCategory.Add(ETypeCategory_ArrayMathStructValue, TArray<int32>()).Reserve(MathTypes.Num());
+	TypesPerCategory.Add(ETypeCategory_ArrayArrayMathStructValue, TArray<int32>()).Reserve(MathTypes.Num());
+	TypesPerCategory.Add(ETypeCategory_SingleScriptStructValue, TArray<int32>()).Reserve(128);
+	TypesPerCategory.Add(ETypeCategory_ArrayScriptStructValue, TArray<int32>()).Reserve(128);
+	TypesPerCategory.Add(ETypeCategory_ArrayArrayScriptStructValue, TArray<int32>()).Reserve(128);
+	TypesPerCategory.Add(ETypeCategory_SingleEnumValue, TArray<int32>()).Reserve(128);
+	TypesPerCategory.Add(ETypeCategory_ArrayEnumValue, TArray<int32>()).Reserve(128);
+	TypesPerCategory.Add(ETypeCategory_ArrayArrayEnumValue, TArray<int32>()).Reserve(128);
+	TypesPerCategory.Add(ETypeCategory_SingleObjectValue, TArray<int32>()).Reserve(128);
+	TypesPerCategory.Add(ETypeCategory_ArrayObjectValue, TArray<int32>()).Reserve(128);
+	TypesPerCategory.Add(ETypeCategory_ArrayArrayObjectValue, TArray<int32>()).Reserve(128);
+
+	RigVMTypeUtils::TypeIndex::Bool = FindOrAddType(FRigVMTemplateArgumentType(RigVMTypeUtils::BoolTypeName, nullptr));
+	RigVMTypeUtils::TypeIndex::Float = FindOrAddType(FRigVMTemplateArgumentType(RigVMTypeUtils::FloatTypeName, nullptr));
+	RigVMTypeUtils::TypeIndex::Double = FindOrAddType(FRigVMTemplateArgumentType(RigVMTypeUtils::DoubleTypeName, nullptr));
+	RigVMTypeUtils::TypeIndex::Int32 = FindOrAddType(FRigVMTemplateArgumentType(RigVMTypeUtils::Int32TypeName, nullptr));
+	RigVMTypeUtils::TypeIndex::UInt8 = FindOrAddType(FRigVMTemplateArgumentType(RigVMTypeUtils::UInt8TypeName, nullptr));
+	RigVMTypeUtils::TypeIndex::FName = FindOrAddType(FRigVMTemplateArgumentType(RigVMTypeUtils::FNameTypeName, nullptr));
+	RigVMTypeUtils::TypeIndex::FString = FindOrAddType(FRigVMTemplateArgumentType(RigVMTypeUtils::FStringTypeName, nullptr));
+	RigVMTypeUtils::TypeIndex::WildCard = FindOrAddType(FRigVMTemplateArgumentType(RigVMTypeUtils::GetWildCardCPPTypeName(), RigVMTypeUtils::GetWildCardCPPTypeObject()));
+	RigVMTypeUtils::TypeIndex::BoolArray = FindOrAddType(FRigVMTemplateArgumentType(RigVMTypeUtils::BoolArrayTypeName, nullptr));
+	RigVMTypeUtils::TypeIndex::FloatArray = FindOrAddType(FRigVMTemplateArgumentType(RigVMTypeUtils::FloatArrayTypeName, nullptr));
+	RigVMTypeUtils::TypeIndex::DoubleArray = FindOrAddType(FRigVMTemplateArgumentType(RigVMTypeUtils::DoubleArrayTypeName, nullptr));
+	RigVMTypeUtils::TypeIndex::Int32Array = FindOrAddType(FRigVMTemplateArgumentType(RigVMTypeUtils::Int32ArrayTypeName, nullptr));
+	RigVMTypeUtils::TypeIndex::UInt8Array = FindOrAddType(FRigVMTemplateArgumentType(RigVMTypeUtils::UInt8ArrayTypeName, nullptr));
+	RigVMTypeUtils::TypeIndex::FNameArray = FindOrAddType(FRigVMTemplateArgumentType(RigVMTypeUtils::FNameArrayTypeName, nullptr));
+	RigVMTypeUtils::TypeIndex::FStringArray = FindOrAddType(FRigVMTemplateArgumentType(RigVMTypeUtils::FStringArrayTypeName, nullptr));
+	RigVMTypeUtils::TypeIndex::WildCardArray = FindOrAddType(FRigVMTemplateArgumentType(RigVMTypeUtils::GetWildCardArrayCPPTypeName(), RigVMTypeUtils::GetWildCardCPPTypeObject()));
+
+	// register the default math types
+	for(UScriptStruct* MathType : MathTypes)
+	{
+		FindOrAddType(FRigVMTemplateArgumentType(*MathType->GetStructCPPName(), MathType));
+	}
+
+	// add all user defined structs
+	for (TObjectIterator<UScriptStruct> ScriptIt; ScriptIt; ++ScriptIt)
+	{
+		UScriptStruct* ScriptStruct = *ScriptIt;
+		if(!IsAllowedType(ScriptStruct))
+		{
+			continue;
+		}
+
+		// if this is a C++ type - skip it
+		if(ScriptStruct->IsNative())
+		{
+			continue;
+		}
+
+		const FString CPPType = ScriptStruct->GetStructCPPName();
+		FindOrAddType(FRigVMTemplateArgumentType(*CPPType, ScriptStruct));
+	}
+
+	// add all user defined enums
+	for (TObjectIterator<UEnum> EnumIt; EnumIt; ++EnumIt)
+	{
+		UEnum* Enum = (*EnumIt);
+		if(!IsAllowedType(Enum))
+		{
+			continue;
+		}
+
+		// if this is a C++ type - skip it
+		if(Enum->IsNative())
+		{
+			continue;
+		}
+
+		const FString CPPType = Enum->CppType.IsEmpty() ? Enum->GetName() : Enum->CppType;
+		FindOrAddType(FRigVMTemplateArgumentType(*CPPType, Enum));
+	}
+}
+
 void FRigVMRegistry::Refresh()
 {
 }
 
+int32 FRigVMRegistry::FindOrAddType(const FRigVMTemplateArgumentType& InType)
+{
+	int32 Index = GetTypeIndex(InType);
+	if(Index == INDEX_NONE)
+	{
+		{
+			FTypeInfo Info;
+			Info.Type = InType;
+			Info.bIsArray = InType.IsArray();
+			Index = Types.Add(Info);
+		}
+
+		TypeToIndex.Add(InType, Index);
+
+		// add the type to the category map
+		static constexpr TCHAR ArrayArrayPrefix[] = TEXT("TArray<TArray<");
+		const int32 ArrayDimension = Types[Index].bIsArray ?
+			(InType.CPPType.ToString().StartsWith(ArrayArrayPrefix) ? 2 : 1) : 0;
+		const UObject* CPPTypeObject = Types[Index].Type.CPPTypeObject;
+
+		// simple types
+		if(CPPTypeObject == nullptr)
+		{
+			switch(ArrayDimension)
+			{
+				default:
+				case 0:
+				{
+					TypesPerCategory.FindChecked(ETypeCategory_SingleSimpleValue).Add(Index);
+					TypesPerCategory.FindChecked(ETypeCategory_SingleAnyValue).Add(Index);
+					break;
+				}
+				case 1:
+				{
+					TypesPerCategory.FindChecked(ETypeCategory_ArraySimpleValue).Add(Index);
+					TypesPerCategory.FindChecked(ETypeCategory_ArrayAnyValue).Add(Index);
+					break;
+				}
+				case 2:
+				{
+					TypesPerCategory.FindChecked(ETypeCategory_ArrayArraySimpleValue).Add(Index);
+					TypesPerCategory.FindChecked(ETypeCategory_ArrayArrayAnyValue).Add(Index);
+					break;
+				}
+			}
+		}
+		else if(const UClass* Class = Cast<UClass>(CPPTypeObject))
+		{
+			if(IsAllowedType(Class))
+			{
+				switch(ArrayDimension)
+				{
+					default:
+					case 0:
+					{
+						TypesPerCategory.FindChecked(ETypeCategory_SingleObjectValue).Add(Index);
+						TypesPerCategory.FindChecked(ETypeCategory_SingleAnyValue).Add(Index);
+						break;
+					}
+					case 1:
+					{
+						TypesPerCategory.FindChecked(ETypeCategory_ArrayObjectValue).Add(Index);
+						TypesPerCategory.FindChecked(ETypeCategory_ArrayAnyValue).Add(Index);
+						break;
+					}
+					case 2:
+					{
+						TypesPerCategory.FindChecked(ETypeCategory_ArrayArrayObjectValue).Add(Index);
+						TypesPerCategory.FindChecked(ETypeCategory_ArrayArrayAnyValue).Add(Index);
+						break;
+					}
+				}
+			}	
+		}
+		else if(const UEnum* Enum = Cast<UEnum>(CPPTypeObject))
+		{
+			if(IsAllowedType(Enum))
+			{
+				switch(ArrayDimension)
+				{
+					default:
+					case 0:
+					{
+						TypesPerCategory.FindChecked(ETypeCategory_SingleEnumValue).Add(Index);
+						TypesPerCategory.FindChecked(ETypeCategory_SingleAnyValue).Add(Index);
+						break;
+					}
+					case 1:
+					{
+						TypesPerCategory.FindChecked(ETypeCategory_ArrayEnumValue).Add(Index);
+						TypesPerCategory.FindChecked(ETypeCategory_ArrayAnyValue).Add(Index);
+						break;
+					}
+					case 2:
+					{
+						TypesPerCategory.FindChecked(ETypeCategory_ArrayArrayEnumValue).Add(Index);
+						TypesPerCategory.FindChecked(ETypeCategory_ArrayArrayAnyValue).Add(Index);
+						break;
+					}
+				}
+			}
+		}
+		else if(const UStruct* Struct = Cast<UStruct>(CPPTypeObject))
+		{
+			if(IsAllowedType(Struct))
+			{
+				if(MathTypes.Contains(CPPTypeObject))
+				{
+					switch(ArrayDimension)
+					{
+						default:
+						case 0:
+						{
+							TypesPerCategory.FindChecked(ETypeCategory_SingleMathStructValue).Add(Index);
+							break;
+						}
+						case 1:
+						{
+							TypesPerCategory.FindChecked(ETypeCategory_ArrayMathStructValue).Add(Index);
+							break;
+						}
+						case 2:
+						{
+							TypesPerCategory.FindChecked(ETypeCategory_ArrayArrayMathStructValue).Add(Index);
+							break;
+						}
+					}
+				}
+				
+				switch(ArrayDimension)
+				{
+					default:
+					case 0:
+					{
+						TypesPerCategory.FindChecked(ETypeCategory_SingleScriptStructValue).Add(Index);
+						TypesPerCategory.FindChecked(ETypeCategory_SingleAnyValue).Add(Index);
+						break;
+					}
+					case 1:
+					{
+						TypesPerCategory.FindChecked(ETypeCategory_ArrayScriptStructValue).Add(Index);
+						TypesPerCategory.FindChecked(ETypeCategory_ArrayAnyValue).Add(Index);
+						break;
+					}
+					case 2:
+					{
+						TypesPerCategory.FindChecked(ETypeCategory_ArrayArrayScriptStructValue).Add(Index);
+						TypesPerCategory.FindChecked(ETypeCategory_ArrayArrayAnyValue).Add(Index);
+						break;
+					}
+				}
+			}
+		}
+
+		// register the opposing type
+		if(!Types[Index].bIsArray)
+		{
+			FRigVMTemplateArgumentType ArrayType = InType;
+			ArrayType.ConvertToArray();
+			Types[Index].ArrayTypeIndex = FindOrAddType(ArrayType);
+			if(Types[Index].ArrayTypeIndex != INDEX_NONE)
+			{
+				Types[Types[Index].ArrayTypeIndex].BaseTypeIndex = Index;
+			}
+		}
+		else
+		{
+			FRigVMTemplateArgumentType ElementType = InType;
+			ElementType.ConvertToBaseElement();
+			Types[Index].BaseTypeIndex = FindOrAddType(ElementType);
+			if(Types[Index].BaseTypeIndex != INDEX_NONE)
+			{
+				Types[Types[Index].BaseTypeIndex].ArrayTypeIndex = Index;
+
+				// also automatically do the two dimensional array
+				if(GetArrayDimensionsForType(Index) == 1)
+				{
+					FRigVMTemplateArgumentType ArrayArrayType = InType;
+					ArrayArrayType.ConvertToArray();
+					Types[Index].ArrayTypeIndex = FindOrAddType(ArrayArrayType);
+					if(Types[Index].ArrayTypeIndex != INDEX_NONE)
+					{
+						Types[Types[Index].ArrayTypeIndex].BaseTypeIndex = Index;
+					}
+				}
+			}
+		}
+
+		// if the type is a structure
+		// then add all of its sub property types
+		if(!Types[Index].bIsArray)
+		{
+			if(const UStruct* Struct = Cast<UStruct>(Types[Index].Type.CPPTypeObject))
+			{
+				for (TFieldIterator<FProperty> It(Struct); It; ++It)
+				{
+					FProperty* Property = *It;
+					if(IsAllowedType(Property, true))
+					{
+						// by creating a template argument for the child property
+						// the type will be added by calling ::FindOrAddType recursively.
+						FRigVMTemplateArgument DummyArgument(Property);
+					}
+				}
+			}
+		}
+	}
+	
+	return Index;
+}
+
+int32 FRigVMRegistry::GetTypeIndex(const FRigVMTemplateArgumentType& InType) const
+{
+	if(const int32* Index = TypeToIndex.Find(InType))
+	{
+		return *Index;
+	}
+	return INDEX_NONE;
+}
+
+const FRigVMTemplateArgumentType& FRigVMRegistry::GetType(int32 InTypeIndex) const
+{
+	if(ensure(Types.IsValidIndex(InTypeIndex)))
+	{
+		return Types[InTypeIndex].Type;
+	}
+	static FRigVMTemplateArgumentType EmptyType;
+	return EmptyType;
+}
+
+const FRigVMTemplateArgumentType& FRigVMRegistry::FindTypeFromCPPType(const FString& InCPPType) const
+{
+	const int32 TypeIndex = GetTypeIndexFromCPPType(InCPPType);
+	if(ensure(Types.IsValidIndex(TypeIndex)))
+	{
+		return Types[TypeIndex].Type;
+	}
+
+	static FRigVMTemplateArgumentType EmptyType;
+	return EmptyType;
+}
+
+int32 FRigVMRegistry::GetTypeIndexFromCPPType(const FString& InCPPType) const
+{
+	if(ensure(!InCPPType.IsEmpty()))
+	{
+		const FName CPPTypeName = *InCPPType;
+		return Types.IndexOfByPredicate([CPPTypeName](const FTypeInfo& Info) -> bool
+		{
+			return Info.Type.CPPType == CPPTypeName;
+		});
+	}
+	return INDEX_NONE;
+}
+
+bool FRigVMRegistry::IsArrayType(int32 InTypeIndex) const
+{
+	if(ensure(Types.IsValidIndex(InTypeIndex)))
+	{
+		return Types[InTypeIndex].bIsArray;
+	}
+	return false;
+}
+
+int32 FRigVMRegistry::GetArrayDimensionsForType(int32 InTypeIndex) const
+{
+	if(ensure(Types.IsValidIndex(InTypeIndex)))
+	{
+		const FTypeInfo& Info = Types[InTypeIndex];
+		if(Info.bIsArray)
+		{
+			return 1 + GetArrayDimensionsForType(Info.BaseTypeIndex);
+		}
+	}
+	return 0;
+}
+
+bool FRigVMRegistry::IsWildCardType(int32 InTypeIndex) const
+{
+	return RigVMTypeUtils::TypeIndex::WildCard == InTypeIndex ||
+		RigVMTypeUtils::TypeIndex::WildCardArray == InTypeIndex;
+}
+
+bool FRigVMRegistry::CanMatchTypes(int32 InTypeIndexA, int32 InTypeIndexB, bool bAllowFloatingPointCasts) const
+{
+	if(!ensure(Types.IsValidIndex(InTypeIndexA)) || !ensure(Types.IsValidIndex(InTypeIndexB)))
+	{
+		return false;
+	}
+
+	if(InTypeIndexA == InTypeIndexB)
+	{
+		return true;
+	}
+
+	if(bAllowFloatingPointCasts)
+	{
+		// swap order since float is known to registered before double
+		if(InTypeIndexA > InTypeIndexB)
+		{
+			Swap(InTypeIndexA, InTypeIndexB);
+		}
+		if(InTypeIndexA == RigVMTypeUtils::TypeIndex::Float && InTypeIndexB == RigVMTypeUtils::TypeIndex::Double)
+		{
+			return true;
+		}
+		if(InTypeIndexA == RigVMTypeUtils::TypeIndex::FloatArray && InTypeIndexB == RigVMTypeUtils::TypeIndex::DoubleArray)
+		{
+			return true;
+		}
+	}
+	return false;
+}
+
+const TArray<int32>& FRigVMRegistry::GetCompatibleTypes(int32 InTypeIndex) const
+{
+	if(InTypeIndex == RigVMTypeUtils::TypeIndex::Float)
+	{
+		static const TArray<int32> CompatibleTypes = {RigVMTypeUtils::TypeIndex::Double};
+		return CompatibleTypes;
+	}
+	if(InTypeIndex == RigVMTypeUtils::TypeIndex::Double)
+	{
+		static const TArray<int32> CompatibleTypes = {RigVMTypeUtils::TypeIndex::Float};
+		return CompatibleTypes;
+	}
+	if(InTypeIndex == RigVMTypeUtils::TypeIndex::FloatArray)
+	{
+		static const TArray<int32> CompatibleTypes = {RigVMTypeUtils::TypeIndex::DoubleArray};
+		return CompatibleTypes;
+	}
+	if(InTypeIndex == RigVMTypeUtils::TypeIndex::DoubleArray)
+	{
+		static const TArray<int32> CompatibleTypes = {RigVMTypeUtils::TypeIndex::FloatArray};
+		return CompatibleTypes;
+	}
+
+	static const TArray<int32> EmptyTypes;
+	return EmptyTypes;
+}
+
+const TArray<int32>& FRigVMRegistry::GetTypesForCategory(ETypeCategory InCategory)
+{
+	check(InCategory != ETypeCategory_Invalid);
+	return TypesPerCategory.FindChecked(InCategory);
+}
+
+int32 FRigVMRegistry::GetArrayTypeFromBaseTypeIndex(int32 InTypeIndex) const
+{
+	if(ensure(Types.IsValidIndex(InTypeIndex)))
+	{
+		return Types[InTypeIndex].ArrayTypeIndex;
+	}
+	return INDEX_NONE;
+}
+
+int32 FRigVMRegistry::GetBaseTypeFromArrayTypeIndex(int32 InTypeIndex) const
+{
+	if(ensure(Types.IsValidIndex(InTypeIndex)))
+	{
+		return Types[InTypeIndex].BaseTypeIndex;
+	}
+	return INDEX_NONE;
+}
+
+bool FRigVMRegistry::IsAllowedType(const FProperty* InProperty, bool bCheckFlags)
+{
+	if(bCheckFlags)
+	{
+		if(!InProperty->HasAnyPropertyFlags(
+			CPF_BlueprintVisible |
+			CPF_BlueprintReadOnly |
+			CPF_Edit))
+		{
+			return false;
+		}
+	}
+
+	if(InProperty->IsA<FBoolProperty>() ||
+		InProperty->IsA<FUInt32Property>() ||
+		InProperty->IsA<FInt8Property>() ||
+		InProperty->IsA<FInt16Property>() ||
+		InProperty->IsA<FIntProperty>() ||
+		InProperty->IsA<FInt64Property>() ||
+		InProperty->IsA<FFloatProperty>() ||
+		InProperty->IsA<FDoubleProperty>() ||
+		InProperty->IsA<FNumericProperty>() ||
+		InProperty->IsA<FNameProperty>() ||
+		InProperty->IsA<FStrProperty>())
+	{
+		return true;
+	}
+
+	if(const FArrayProperty* ArrayProperty  = CastField<FArrayProperty>(InProperty))
+	{
+		return IsAllowedType(ArrayProperty->Inner, false);
+	}
+	if(const FStructProperty* StructProperty = CastField<FStructProperty>(InProperty))
+	{
+		return IsAllowedType(StructProperty->Struct);
+	}
+	if(const FObjectProperty* ObjectProperty = CastField<FObjectProperty>(InProperty))
+	{
+		return IsAllowedType(ObjectProperty->PropertyClass);
+	}
+	if(const FEnumProperty* EnumProperty = CastField<FEnumProperty>(InProperty))
+	{
+		return IsAllowedType(EnumProperty->GetEnum());
+	}
+	if(const FByteProperty* ByteProperty = CastField<FByteProperty>(InProperty))
+	{
+		if(const UEnum* Enum = ByteProperty->Enum)
+		{
+			return IsAllowedType(Enum);
+		}
+		return true;
+	}
+	return false;
+}
+
+bool FRigVMRegistry::IsAllowedType(const UEnum* InEnum)
+{
+	return !InEnum->HasAnyFlags(DisallowedFlags()) && InEnum->HasAllFlags(NeededFlags());
+}
+
+bool FRigVMRegistry::IsAllowedType(const UStruct* InStruct)
+{
+	if(InStruct->HasAnyFlags(DisallowedFlags()) || !InStruct->HasAllFlags(NeededFlags()))
+	{
+		return false;
+	}
+	if(InStruct->IsChildOf(FRigVMStruct::StaticStruct()))
+	{
+		return false;
+	}
+	if(InStruct->IsChildOf(FRigVMUnknownType::StaticStruct()))
+	{
+		return false;
+	}
+	if(InStruct->IsChildOf(FRigVMExecuteContext::StaticStruct()))
+	{
+		return false;
+	}
+	for (TFieldIterator<FProperty> It(InStruct); It; ++It)
+	{
+		if(!IsAllowedType(*It))
+		{
+			return false;
+		}
+	}
+	return true;
+}
+
+bool FRigVMRegistry::IsAllowedType(const UClass* InClass)
+{
+	if(InClass->HasAnyClassFlags(CLASS_Hidden | CLASS_Abstract))
+	{
+		return false;
+	}
+
+	// note: currently we don't allow UObjects
+	return false;
+	//return IsAllowedType(Cast<UStruct>(InClass));
+}
+
+
 void FRigVMRegistry::Register(const TCHAR* InName, FRigVMFunctionPtr InFunctionPtr, UScriptStruct* InStruct, const TArray<FRigVMFunctionArgument>& InArguments)
 {
 	if (FindFunction(InName) != nullptr)
@@ -133,35 +693,35 @@ const FRigVMTemplate* FRigVMRegistry::GetOrAddTemplateFromArguments(const FName&
 	{
 		if(!Argument.IsSingleton() && NumPermutations > 1)
 		{
-			if(Argument.Types.Num() != NumPermutations)
+			if(Argument.TypeIndices.Num() != NumPermutations)
 			{
 				UE_LOG(LogRigVM, Error, TEXT("Failed to add template '%s' since the arguments' types counts don't match."), *InName.ToString());
 				return nullptr;
 			}
 		}
-		NumPermutations = FMath::Max(NumPermutations, Argument.Types.Num()); 
+		NumPermutations = FMath::Max(NumPermutations, Argument.TypeIndices.Num()); 
 	}
 
 	// if any of the arguments are wildcards we'll need to update the types
 	for(FRigVMTemplateArgument& Argument : Template.Arguments)
 	{
-		if(Argument.Types.Num() == 1 && Argument.Types[0].IsWildCard())
+		if(Argument.TypeIndices.Num() == 1 && IsWildCardType(Argument.TypeIndices[0]))
 		{
-			if(Argument.Types[0].IsArray())
+			if(IsArrayType(Argument.TypeIndices[0]))
 			{
-				Argument.Types = FRigVMTemplateArgument::GetCompatibleTypes(FRigVMTemplateArgument::ETypeCategory_ArrayAnyValue);
+				Argument.TypeIndices = GetTypesForCategory(ETypeCategory_ArrayAnyValue);
 			}
 			else
 			{
-				Argument.Types = FRigVMTemplateArgument::GetCompatibleTypes(FRigVMTemplateArgument::ETypeCategory_SingleAnyValue);
+				Argument.TypeIndices = GetTypesForCategory(ETypeCategory_SingleAnyValue);
 			}
 
-			for (int32 i = 0; i < Argument.Types.Num(); ++i)
+			for (int32 i = 0; i < Argument.TypeIndices.Num(); ++i)
 			{
-				Argument.TypeToPermutations.Add(Argument.Types[i].CPPType, {i});				
+				Argument.TypeToPermutations.Add(Argument.TypeIndices[i], {i});				
 			}
 			
-			NumPermutations = FMath::Max(NumPermutations, Argument.Types.Num()); 
+			NumPermutations = FMath::Max(NumPermutations, Argument.TypeIndices.Num()); 
 		}
 	}
 
@@ -170,18 +730,18 @@ const FRigVMTemplate* FRigVMRegistry::GetOrAddTemplateFromArguments(const FName&
 	{
 		for(FRigVMTemplateArgument& Argument : Template.Arguments)
 		{
-			if(Argument.Types.Num() == 1)
+			if(Argument.TypeIndices.Num() == 1)
 			{
-				const FRigVMTemplateArgumentType Type = Argument.Types[0];
-				Argument.Types.SetNum(NumPermutations);
+				const int32 TypeIndex = Argument.TypeIndices[0];
+				Argument.TypeIndices.SetNum(NumPermutations);
 				TArray<int32> ArgTypePermutations;
 				ArgTypePermutations.SetNum(NumPermutations);
 				for(int32 Index=0;Index<NumPermutations;Index++)
 				{
-					Argument.Types[Index] = Type;
+					Argument.TypeIndices[Index] = TypeIndex;
 					ArgTypePermutations[Index] = Index;
 				}
-				Argument.TypeToPermutations.Add(Type.CPPType, ArgTypePermutations);
+				Argument.TypeToPermutations.Add(TypeIndex, ArgTypePermutations);
 			}
 		}
 	}
diff --git a/Engine/Source/Runtime/RigVM/Private/RigVMCore/RigVMTemplate.cpp b/Engine/Source/Runtime/RigVM/Private/RigVMCore/RigVMTemplate.cpp
--- a/Engine/Source/Runtime/RigVM/Private/RigVMCore/RigVMTemplate.cpp
+++ b/Engine/Source/Runtime/RigVM/Private/RigVMCore/RigVMTemplate.cpp
@@ -27,7 +27,8 @@ FRigVMTemplateArgument::FRigVMTemplateArgument(FProperty* InProperty)
 
 	FString ExtendedType;
 	const FString CPPType = InProperty->GetCPPType(&ExtendedType);
-	FRigVMTemplateArgumentType Type(CPPType + ExtendedType);
+	const FName CPPTypeName = *(CPPType + ExtendedType);
+	FRigVMTemplateArgumentType Type(CPPTypeName);
 
 	if (FArrayProperty* ArrayProperty = CastField<FArrayProperty>(InProperty))
 	{
@@ -46,62 +47,64 @@ FRigVMTemplateArgument::FRigVMTemplateArgument(FProperty* InProperty)
 	{
 		Type.CPPTypeObject = ByteProperty->Enum;
 	}
-	Type.CPPType = RigVMTypeUtils::PostProcessCPPType(Type.CPPType, Type.CPPTypeObject);
+	Type.CPPType = *RigVMTypeUtils::PostProcessCPPType(Type.CPPType.ToString(), Type.CPPTypeObject);
 
-	Types.Add(Type);
-	TypeToPermutations.Add(Type.CPPType, {0});
+	const int32 TypeIndex = FRigVMRegistry::Get().FindOrAddType(Type); 
+
+	TypeIndices.Add(TypeIndex);
+	TypeToPermutations.Add(TypeIndex, {0});
 }
 
-FRigVMTemplateArgument::FRigVMTemplateArgument(const FName& InName, ERigVMPinDirection InDirection, const FRigVMTemplateArgumentType& InType)
+FRigVMTemplateArgument::FRigVMTemplateArgument(const FName& InName, ERigVMPinDirection InDirection, int32 InTypeIndex)
 : Index(INDEX_NONE)
 , Name(InName)
 , Direction(InDirection)
-, Types({InType})
+, TypeIndices({InTypeIndex})
 {
-	TypeToPermutations.Add(InType.CPPType, {0});	
+	TypeToPermutations.Add(InTypeIndex, {0});	
 }
 
-FRigVMTemplateArgument::FRigVMTemplateArgument(const FName& InName, ERigVMPinDirection InDirection, const TArray<FRigVMTemplateArgumentType>& InTypes)
+FRigVMTemplateArgument::FRigVMTemplateArgument(const FName& InName, ERigVMPinDirection InDirection, const TArray<int32>& InTypeIndices)
 : Index(INDEX_NONE)
 , Name(InName)
 , Direction(InDirection)
-, Types(InTypes)
+, TypeIndices(InTypeIndices)
 {
-	check(Types.Num() > 0);
+	check(TypeIndices.Num() > 0);
 
-	for(int32 TypeIndex=0;TypeIndex<Types.Num();TypeIndex++)
+	for(int32 TypeIndex=0;TypeIndex<TypeIndices.Num();TypeIndex++)
 	{
-		if (TArray<int32>* Permutations = TypeToPermutations.Find(Types[TypeIndex].CPPType))
+		if (TArray<int32>* Permutations = TypeToPermutations.Find(TypeIndices[TypeIndex]))
 		{
 			Permutations->Add(TypeIndex);
 		}
 		else
 		{
-			TypeToPermutations.Add(Types[TypeIndex].CPPType, {TypeIndex});
+			TypeToPermutations.Add(TypeIndices[TypeIndex], {TypeIndex});
 		}
 	}
 }
 
-bool FRigVMTemplateArgument::SupportsType(const FString& InCPPType, FRigVMTemplateArgumentType* OutType) const
+bool FRigVMTemplateArgument::SupportsTypeIndex(int32 InTypeIndex, int32* OutTypeIndex) const
 {
-	if (const TArray<int32>* Permutations = TypeToPermutations.Find(InCPPType))
+	if (const TArray<int32>* Permutations = TypeToPermutations.Find(InTypeIndex))
 	{
-		if(OutType)
+		if(OutTypeIndex)
 		{
-			(*OutType) = Types[(*Permutations)[0]];
+			(*OutTypeIndex) = TypeIndices[(*Permutations)[0]];
 		}
 		return true;		
 	}
 
 	// Try to find compatible type
-	TArray<FString> CompatibleTypes = FRigVMTemplateArgumentType::GetCompatibleTypes(InCPPType);
-	for (const FString& Compatible : CompatibleTypes)
+	const TArray<int32>& CompatibleTypes = FRigVMRegistry::Get().GetCompatibleTypes(InTypeIndex);
+	for (const int32& CompatibleTypeIndex : CompatibleTypes)
 	{
-		if (const TArray<int32>* Permutations = TypeToPermutations.Find(Compatible))
+		if (const TArray<int32>* Permutations = TypeToPermutations.Find(CompatibleTypeIndex))
 		{
-			if(OutType)
+			if(OutTypeIndex)
 			{
-				(*OutType) = Types[(*Permutations)[0]];
+				(*OutTypeIndex) = TypeIndices[(*Permutations)[0]];
 			}
 			return true;		
 		}
@@ -121,10 +124,10 @@ bool FRigVMTemplateArgument::IsSingleton(const TArray<int32>& InPermutationIndic
 		return false;
 	}
 
-	const FRigVMTemplateArgumentType TypeToCheck = Types[InPermutationIndices[0]];
+	const int32 TypeIndexToCheck = TypeIndices[InPermutationIndices[0]];
 	for (int32 PermutationIndex = 1; PermutationIndex < InPermutationIndices.Num(); PermutationIndex++)
 	{
-		if (Types[InPermutationIndices[PermutationIndex]] != TypeToCheck)
+		if (TypeIndices[InPermutationIndices[PermutationIndex]] != TypeIndexToCheck)
 		{
 			return false;
 		}
@@ -134,18 +137,19 @@ bool FRigVMTemplateArgument::IsSingleton(const TArray<int32>& InPermutationIndic
 
 FRigVMTemplateArgument::EArrayType FRigVMTemplateArgument::GetArrayType() const
 {
-	if(Types.Num() > 0)
+	if(TypeIndices.Num() > 0)
 	{
-		const EArrayType ArrayType = Types[0].IsArray() ? EArrayType::EArrayType_ArrayValue : EArrayType::EArrayType_SingleValue;
+		const FRigVMRegistry& Registry = FRigVMRegistry::Get();
+		const EArrayType ArrayType = Registry.IsArrayType(TypeIndices[0]) ? EArrayType::EArrayType_ArrayValue : EArrayType::EArrayType_SingleValue;
 		
 		if(IsSingleton())
 		{
 			return ArrayType;
 		}
 
-		for(int32 PermutationIndex=1;PermutationIndex<Types.Num();PermutationIndex++)
+		for(int32 PermutationIndex=1;PermutationIndex<TypeIndices.Num();PermutationIndex++)
 		{
-			const EArrayType OtherArrayType = Types[PermutationIndex].IsArray() ? EArrayType::EArrayType_ArrayValue : EArrayType::EArrayType_SingleValue;
+			const EArrayType OtherArrayType = Registry.IsArrayType(TypeIndices[PermutationIndex]) ? EArrayType::EArrayType_ArrayValue : EArrayType::EArrayType_SingleValue;
 			if(OtherArrayType != ArrayType)
 			{
 				return EArrayType::EArrayType_Mixed;
@@ -158,320 +162,58 @@ FRigVMTemplateArgument::EArrayType FRigVMTemplateArgument::GetArrayType() const
 	return EArrayType_Invalid;
 }
 
-const TArray<FRigVMTemplateArgumentType>& FRigVMTemplateArgument::GetCompatibleTypes(ETypeCategory InCategory)
-{
-	check(InCategory != ETypeCategory_Invalid);
-	
-	static TMap<ETypeCategory, TArray<FRigVMTemplateArgumentType>> CompatibleTypes;
-	if(CompatibleTypes.IsEmpty())
-	{
-		CompatibleTypes.Add(ETypeCategory_SingleAnyValue, TArray<FRigVMTemplateArgumentType>());
-		CompatibleTypes.Add(ETypeCategory_ArrayAnyValue, TArray<FRigVMTemplateArgumentType>());
-		CompatibleTypes.Add(ETypeCategory_ArrayArrayAnyValue, TArray<FRigVMTemplateArgumentType>());
-		CompatibleTypes.Add(ETypeCategory_SingleSimpleValue, TArray<FRigVMTemplateArgumentType>());
-		CompatibleTypes.Add(ETypeCategory_ArraySimpleValue, TArray<FRigVMTemplateArgumentType>());
-		CompatibleTypes.Add(ETypeCategory_ArrayArraySimpleValue, TArray<FRigVMTemplateArgumentType>());
-		CompatibleTypes.Add(ETypeCategory_SingleMathStructValue, TArray<FRigVMTemplateArgumentType>());
-		CompatibleTypes.Add(ETypeCategory_ArrayMathStructValue, TArray<FRigVMTemplateArgumentType>());
-		CompatibleTypes.Add(ETypeCategory_ArrayArrayMathStructValue, TArray<FRigVMTemplateArgumentType>());
-		CompatibleTypes.Add(ETypeCategory_SingleScriptStructValue, TArray<FRigVMTemplateArgumentType>());
-		CompatibleTypes.Add(ETypeCategory_ArrayScriptStructValue, TArray<FRigVMTemplateArgumentType>());
-		CompatibleTypes.Add(ETypeCategory_ArrayArrayScriptStructValue, TArray<FRigVMTemplateArgumentType>());
-		CompatibleTypes.Add(ETypeCategory_SingleEnumValue, TArray<FRigVMTemplateArgumentType>());
-		CompatibleTypes.Add(ETypeCategory_ArrayEnumValue, TArray<FRigVMTemplateArgumentType>());
-		CompatibleTypes.Add(ETypeCategory_ArrayArrayEnumValue, TArray<FRigVMTemplateArgumentType>());
-		CompatibleTypes.Add(ETypeCategory_SingleObjectValue, TArray<FRigVMTemplateArgumentType>());
-		CompatibleTypes.Add(ETypeCategory_ArrayObjectValue, TArray<FRigVMTemplateArgumentType>());
-		CompatibleTypes.Add(ETypeCategory_ArrayArrayObjectValue, TArray<FRigVMTemplateArgumentType>());
-		
-		TArray<FRigVMTemplateArgumentType>& SingleAnyValueTypes = CompatibleTypes.FindChecked(ETypeCategory_SingleAnyValue);
-		TArray<FRigVMTemplateArgumentType>& ArrayAnyValueTypes = CompatibleTypes.FindChecked(ETypeCategory_ArrayAnyValue);
-		TArray<FRigVMTemplateArgumentType>& ArrayArrayAnyValueTypes = CompatibleTypes.FindChecked(ETypeCategory_ArrayArrayAnyValue);
-		TArray<FRigVMTemplateArgumentType>& SingleSimpleValueTypes = CompatibleTypes.FindChecked(ETypeCategory_SingleSimpleValue);
-		TArray<FRigVMTemplateArgumentType>& ArraySimpleValueTypes = CompatibleTypes.FindChecked(ETypeCategory_ArraySimpleValue);
-		TArray<FRigVMTemplateArgumentType>& ArrayArraySimpleValueTypes = CompatibleTypes.FindChecked(ETypeCategory_ArrayArraySimpleValue);
-		TArray<FRigVMTemplateArgumentType>& SingleMathStructValueTypes = CompatibleTypes.FindChecked(ETypeCategory_SingleMathStructValue);
-		TArray<FRigVMTemplateArgumentType>& ArrayMathStructValueTypes = CompatibleTypes.FindChecked(ETypeCategory_ArrayMathStructValue);
-		TArray<FRigVMTemplateArgumentType>& ArrayArrayMathStructValueTypes = CompatibleTypes.FindChecked(ETypeCategory_ArrayArrayMathStructValue);
-		TArray<FRigVMTemplateArgumentType>& SingleScriptStructValueTypes = CompatibleTypes.FindChecked(ETypeCategory_SingleScriptStructValue);
-		TArray<FRigVMTemplateArgumentType>& ArrayScriptStructValueTypes = CompatibleTypes.FindChecked(ETypeCategory_ArrayScriptStructValue);
-		TArray<FRigVMTemplateArgumentType>& ArrayArrayScriptStructValueTypes = CompatibleTypes.FindChecked(ETypeCategory_ArrayArrayScriptStructValue);
-		TArray<FRigVMTemplateArgumentType>& SingleEnumValueTypes = CompatibleTypes.FindChecked(ETypeCategory_SingleEnumValue);
-		TArray<FRigVMTemplateArgumentType>& ArrayEnumValueTypes = CompatibleTypes.FindChecked(ETypeCategory_ArrayEnumValue);
-		TArray<FRigVMTemplateArgumentType>& ArrayArrayEnumValueTypes = CompatibleTypes.FindChecked(ETypeCategory_ArrayArrayEnumValue);
-		TArray<FRigVMTemplateArgumentType>& SingleObjectValueTypes = CompatibleTypes.FindChecked(ETypeCategory_SingleObjectValue);
-		TArray<FRigVMTemplateArgumentType>& ArrayObjectValueTypes = CompatibleTypes.FindChecked(ETypeCategory_ArrayObjectValue);
-		TArray<FRigVMTemplateArgumentType>& ArrayArrayObjectValueTypes = CompatibleTypes.FindChecked(ETypeCategory_ArrayArrayObjectValue);
-		
-		SingleSimpleValueTypes.Add(FRigVMTemplateArgumentType(RigVMTypeUtils::BoolType));
-		SingleSimpleValueTypes.Add(FRigVMTemplateArgumentType(RigVMTypeUtils::Int32Type));
-		SingleSimpleValueTypes.Add(FRigVMTemplateArgumentType(RigVMTypeUtils::UInt8Type));
-		SingleSimpleValueTypes.Add(FRigVMTemplateArgumentType(RigVMTypeUtils::FloatType));
-		SingleSimpleValueTypes.Add(FRigVMTemplateArgumentType(RigVMTypeUtils::DoubleType));
-		SingleSimpleValueTypes.Add(FRigVMTemplateArgumentType(RigVMTypeUtils::FNameType));
-		SingleSimpleValueTypes.Add(FRigVMTemplateArgumentType(RigVMTypeUtils::FStringType));
-
-		for(const FRigVMTemplateArgumentType& Type : SingleSimpleValueTypes)
-		{
-			ArraySimpleValueTypes.Add(FRigVMTemplateArgumentType(RigVMTypeUtils::ArrayTypeFromBaseType(Type.CPPType)));
-			ArrayArraySimpleValueTypes.Add(FRigVMTemplateArgumentType(RigVMTypeUtils::ArrayTypeFromBaseType(ArraySimpleValueTypes.Last().CPPType)));
-		}
-
-		SingleAnyValueTypes = SingleSimpleValueTypes;
-		ArrayAnyValueTypes = ArraySimpleValueTypes;
-		ArrayArrayAnyValueTypes = ArrayArraySimpleValueTypes;
-
-		static const TArray<UScriptStruct*> MathTypes = { 
-			TBaseStructure<FRotator>::Get(),
-			TBaseStructure<FQuat>::Get(),
-			TBaseStructure<FTransform>::Get(),
-			TBaseStructure<FLinearColor>::Get(),
-			TBaseStructure<FColor>::Get(),
-			TBaseStructure<FPlane>::Get(),
-			TBaseStructure<FVector>::Get(),
-			TBaseStructure<FVector2D>::Get(),
-			TBaseStructure<FVector4>::Get(),
-			TBaseStructure<FBox2D>::Get()
-		};
-
-		for(UScriptStruct* MathType : MathTypes)
-		{
-			SingleMathStructValueTypes.Add(FRigVMTemplateArgumentType(MathType->GetStructCPPName(), MathType));
-			ArrayMathStructValueTypes.Add(FRigVMTemplateArgumentType(RigVMTypeUtils::ArrayTypeFromBaseType(MathType->GetStructCPPName()), MathType));
-			ArrayArrayMathStructValueTypes.Add(FRigVMTemplateArgumentType(RigVMTypeUtils::ArrayTypeFromBaseType(ArrayMathStructValueTypes.Last().CPPType), MathType));
-		}
-
-		struct FTypeTraverser
-		{
-			static EObjectFlags DisallowedFlags()
-			{
-				return RF_BeginDestroyed | RF_FinishDestroyed;
-			}
-
-			static EObjectFlags NeededFlags()
-			{
-				return RF_Public;
-			}
-			
-			static bool IsAllowedType(const FProperty* InProperty, bool bCheckFlags = true)
-			{
-				if(bCheckFlags)
-				{
-					if(!InProperty->HasAnyPropertyFlags(
-						CPF_BlueprintVisible |
-						CPF_BlueprintReadOnly |
-						CPF_Edit))
-					{
-						return false;
-					}
-				}
-
-				if(InProperty->IsA<FBoolProperty>() ||
-					InProperty->IsA<FUInt32Property>() ||
-					InProperty->IsA<FInt8Property>() ||
-					InProperty->IsA<FInt16Property>() ||
-					InProperty->IsA<FIntProperty>() ||
-					InProperty->IsA<FInt64Property>() ||
-					InProperty->IsA<FFloatProperty>() ||
-					InProperty->IsA<FDoubleProperty>() ||
-					InProperty->IsA<FNumericProperty>() ||
-					InProperty->IsA<FNameProperty>() ||
-					InProperty->IsA<FStrProperty>())
-				{
-					return true;
-				}
-
-				if(const FArrayProperty* ArrayProperty  = CastField<FArrayProperty>(InProperty))
-				{
-					return IsAllowedType(ArrayProperty->Inner, false);
-				}
-				if(const FStructProperty* StructProperty = CastField<FStructProperty>(InProperty))
-				{
-					return IsAllowedType(StructProperty->Struct);
-				}
-				if(const FObjectProperty* ObjectProperty = CastField<FObjectProperty>(InProperty))
-				{
-					return IsAllowedType(ObjectProperty->PropertyClass);
-				}
-				if(const FEnumProperty* EnumProperty = CastField<FEnumProperty>(InProperty))
-				{
-					return IsAllowedType(EnumProperty->GetEnum());
-				}
-				if(const FByteProperty* ByteProperty = CastField<FByteProperty>(InProperty))
-				{
-					if(const UEnum* Enum = ByteProperty->Enum)
-					{
-						return IsAllowedType(Enum);
-					}
-					return true;
-				}
-				return false;
-			}
-
-			static bool IsAllowedType(const UEnum* InEnum)
-			{
-				return !InEnum->HasAnyFlags(DisallowedFlags()) && InEnum->HasAllFlags(NeededFlags());
-			}
-
-			static bool IsAllowedType(const UStruct* InStruct)
-			{
-				if(InStruct->HasAnyFlags(DisallowedFlags()) || !InStruct->HasAllFlags(NeededFlags()))
-				{
-					return false;
-				}
-				if(InStruct->IsChildOf(FRigVMStruct::StaticStruct()))
-				{
-					return false;
-				}
-				if(InStruct->IsChildOf(FRigVMUnknownType::StaticStruct()))
-				{
-					return false;
-				}
-				for (TFieldIterator<FProperty> It(InStruct); It; ++It)
-				{
-					if(!IsAllowedType(*It))
-					{
-						return false;
-					}
-				}
-				return true;
-			}
-
-			static bool IsAllowedType(const UClass* InClass)
-			{
-				if(InClass->HasAnyClassFlags(CLASS_Hidden | CLASS_Abstract))
-				{
-					return false;
-				}
-
-				// note: currently we don't allow UObjects
-				return false;
-				//return IsAllowedType(Cast<UStruct>(InClass));
-			}
-		};
-
-		// add all structs
-		for (TObjectIterator<UScriptStruct> ScriptIt; ScriptIt; ++ScriptIt)
-		{
-			UScriptStruct* ScriptStruct = *ScriptIt;
-			if(!FTypeTraverser::IsAllowedType(ScriptStruct))
-			{
-				continue;
-			}
-
-			const FString CPPType = ScriptStruct->GetStructCPPName();
-
-			if(!ScriptStruct->IsChildOf(FRigVMExecuteContext::StaticStruct()))
-			{
-				SingleAnyValueTypes.Add(FRigVMTemplateArgumentType(CPPType, ScriptStruct));
-				ArrayAnyValueTypes.Add(FRigVMTemplateArgumentType(RigVMTypeUtils::ArrayTypeFromBaseType(CPPType), ScriptStruct));
-				ArrayArrayAnyValueTypes.Add(FRigVMTemplateArgumentType(RigVMTypeUtils::ArrayTypeFromBaseType(ArrayAnyValueTypes.Last().CPPType), ScriptStruct));
-				SingleScriptStructValueTypes.Add(FRigVMTemplateArgumentType(CPPType, ScriptStruct));
-				ArrayScriptStructValueTypes.Add(FRigVMTemplateArgumentType(RigVMTypeUtils::ArrayTypeFromBaseType(CPPType), ScriptStruct));
-				ArrayArrayScriptStructValueTypes.Add(FRigVMTemplateArgumentType(RigVMTypeUtils::ArrayTypeFromBaseType(ArrayScriptStructValueTypes.Last().CPPType), ScriptStruct));
-			}
-		}
-
-		// add all enums
-		for (TObjectIterator<UEnum> EnumIt; EnumIt; ++EnumIt)
-		{
-			UEnum* Enum = (*EnumIt);
-			if(!FTypeTraverser::IsAllowedType(Enum))
-			{
-				continue;
-			}
-			const FString CPPType = Enum->CppType.IsEmpty() ? Enum->GetName() : Enum->CppType;
-			SingleAnyValueTypes.Add(FRigVMTemplateArgumentType(CPPType, Enum));
-			ArrayAnyValueTypes.Add(FRigVMTemplateArgumentType(RigVMTypeUtils::ArrayTypeFromBaseType(CPPType), Enum));
-			ArrayArrayAnyValueTypes.Add(FRigVMTemplateArgumentType(RigVMTypeUtils::ArrayTypeFromBaseType(ArrayAnyValueTypes.Last().CPPType), Enum));
-			SingleEnumValueTypes.Add(FRigVMTemplateArgumentType(CPPType, Enum));
-			ArrayEnumValueTypes.Add(FRigVMTemplateArgumentType(RigVMTypeUtils::ArrayTypeFromBaseType(CPPType), Enum));
-			ArrayArrayEnumValueTypes.Add(FRigVMTemplateArgumentType(RigVMTypeUtils::ArrayTypeFromBaseType(ArrayEnumValueTypes.Last().CPPType), Enum));
-		}
-
-		// add all classes
-		for (TObjectIterator<UClass> ClassIt; ClassIt; ++ClassIt)
-		{
-			UClass* Class = *ClassIt;
-
-			// for now this always returns false
-			if(!FTypeTraverser::IsAllowedType(Class))
-			{
-				continue;
-			}
-
-			const FString CPPType = Class->GetPrefixCPP() + Class->GetName();
-			SingleAnyValueTypes.Add(FRigVMTemplateArgumentType(CPPType, Class));
-			ArrayAnyValueTypes.Add(FRigVMTemplateArgumentType(RigVMTypeUtils::ArrayTypeFromBaseType(CPPType), Class));
-			ArrayArrayAnyValueTypes.Add(FRigVMTemplateArgumentType(RigVMTypeUtils::ArrayTypeFromBaseType(ArrayAnyValueTypes.Last().CPPType), Class));
-			SingleObjectValueTypes.Add(FRigVMTemplateArgumentType(CPPType, Class));
-			ArrayObjectValueTypes.Add(FRigVMTemplateArgumentType(RigVMTypeUtils::ArrayTypeFromBaseType(CPPType), Class));
-			ArrayArrayObjectValueTypes.Add(FRigVMTemplateArgumentType(RigVMTypeUtils::ArrayTypeFromBaseType(ArrayObjectValueTypes.Last().CPPType), Class));
-		}
-
-		// check that matching categories have the same size
-		check(CompatibleTypes.FindChecked(ETypeCategory_SingleAnyValue).Num() == CompatibleTypes.FindChecked(ETypeCategory_ArrayAnyValue).Num());
-		check(CompatibleTypes.FindChecked(ETypeCategory_SingleAnyValue).Num() == CompatibleTypes.FindChecked(ETypeCategory_ArrayArrayAnyValue).Num());
-		check(CompatibleTypes.FindChecked(ETypeCategory_SingleSimpleValue).Num() == CompatibleTypes.FindChecked(ETypeCategory_ArraySimpleValue).Num());
-		check(CompatibleTypes.FindChecked(ETypeCategory_SingleSimpleValue).Num() == CompatibleTypes.FindChecked(ETypeCategory_ArrayArraySimpleValue).Num());
-		check(CompatibleTypes.FindChecked(ETypeCategory_SingleMathStructValue).Num() == CompatibleTypes.FindChecked(ETypeCategory_ArrayMathStructValue).Num());
-		check(CompatibleTypes.FindChecked(ETypeCategory_SingleMathStructValue).Num() == CompatibleTypes.FindChecked(ETypeCategory_ArrayArrayMathStructValue).Num());
-		check(CompatibleTypes.FindChecked(ETypeCategory_SingleScriptStructValue).Num() == CompatibleTypes.FindChecked(ETypeCategory_ArrayScriptStructValue).Num());
-		check(CompatibleTypes.FindChecked(ETypeCategory_SingleScriptStructValue).Num() == CompatibleTypes.FindChecked(ETypeCategory_ArrayArrayScriptStructValue).Num());
-		check(CompatibleTypes.FindChecked(ETypeCategory_SingleEnumValue).Num() == CompatibleTypes.FindChecked(ETypeCategory_ArrayEnumValue).Num());
-		check(CompatibleTypes.FindChecked(ETypeCategory_SingleEnumValue).Num() == CompatibleTypes.FindChecked(ETypeCategory_ArrayArrayEnumValue).Num());
-		check(CompatibleTypes.FindChecked(ETypeCategory_SingleObjectValue).Num() == CompatibleTypes.FindChecked(ETypeCategory_ArrayObjectValue).Num());
-		check(CompatibleTypes.FindChecked(ETypeCategory_SingleObjectValue).Num() == CompatibleTypes.FindChecked(ETypeCategory_ArrayArrayObjectValue).Num());
-	}
-
-	return CompatibleTypes.FindChecked(InCategory);
-}
-
-const TArray<FRigVMTemplateArgumentType>& FRigVMTemplateArgument::GetTypes() const
+const TArray<int32>& FRigVMTemplateArgument::GetTypeIndices() const
 {
-	return Types;
+	return TypeIndices;
 }
 
-TArray<FRigVMTemplateArgumentType> FRigVMTemplateArgument::GetSupportedTypes(const TArray<int32>& InPermutationIndices) const
+TArray<int32> FRigVMTemplateArgument::GetSupportedTypeIndices(const TArray<int32>& InPermutationIndices) const
 {
-	TArray<FRigVMTemplateArgumentType> SupportedTypes;
+	TArray<int32> SupportedTypes;
 	if(InPermutationIndices.IsEmpty())
 	{
-		for (const FRigVMTemplateArgumentType& Type : Types)
+		for (const int32& TypeIndex : TypeIndices)
 		{
-			SupportedTypes.AddUnique(Type);
+			SupportedTypes.AddUnique(TypeIndex);
 		}
 	}
 	else
 	{
 		for(const int32 PermutationIndex : InPermutationIndices)
 		{
-			SupportedTypes.AddUnique(Types[PermutationIndex]);
+			SupportedTypes.AddUnique(TypeIndices[PermutationIndex]);
 		}
 	}
 	return SupportedTypes;
 }
 
+#if WITH_EDITOR
+
 TArray<FString> FRigVMTemplateArgument::GetSupportedTypeStrings(const TArray<int32>& InPermutationIndices) const
 {
 	TArray<FString> SupportedTypes;
 	if(InPermutationIndices.IsEmpty())
 	{
-		for (const FRigVMTemplateArgumentType& Type : Types)
+		for (const int32& TypeIndex : TypeIndices)
 		{
-			SupportedTypes.AddUnique(Type.CPPType);
+			const FString TypeString = FRigVMRegistry::Get().GetType(TypeIndex).CPPType.ToString();
+			SupportedTypes.AddUnique(TypeString);
 		}
 	}
 	else
 	{
 		for(const int32 PermutationIndex : InPermutationIndices)
 		{
-			SupportedTypes.AddUnique(Types[PermutationIndex].CPPType);
+			const int32 TypeIndex = TypeIndices[PermutationIndex];
+			const FString TypeString = FRigVMRegistry::Get().GetType(TypeIndex).CPPType.ToString();
+			SupportedTypes.AddUnique(TypeString);
 		}
 	}
 	return SupportedTypes;
 }
 
+#endif
+
 ////////////////////////////////////////////////////////////////////////////////////////////////
 
 FRigVMTemplate::FRigVMTemplate()
@@ -805,7 +547,7 @@ FText FRigVMTemplate::GetDisplayNameForArgument(const FName& InArgumentName, con
 	{
 		FText ResolvedDisplayName;
 
-		auto VisitPermutation = [InArgumentName, &ResolvedDisplayName, this](int32 InPermutationIndex) -> bool
+		auto VisitPermutation = [InArgumentName, &ResolvedDisplayName, this](const int32 InPermutationIndex) -> bool
 		{
 			// if we can't find one permutation we are not going to find any, so it's ok to return false here
 			const FRigVMFunction* ResolvedFunction = GetPermutation(InPermutationIndex);
@@ -865,7 +607,7 @@ FString FRigVMTemplate::GetArgumentMetaData(const FName& InArgumentName, const F
 	{
 		FString ResolvedMetaData;
 
-		auto VisitPermutation = [InArgumentName, &ResolvedMetaData, InMetaDataKey, this](int32 InPermutationIndex) -> bool
+		auto VisitPermutation = [InArgumentName, &ResolvedMetaData, InMetaDataKey, this](const int32 InPermutationIndex) -> bool
 		{
 			// if we can't find one permutation we are not going to find any, so it's ok to return false here
 			const FRigVMFunction* ResolvedFunction = GetPermutation(InPermutationIndex);
@@ -949,8 +691,8 @@ bool FRigVMTemplate::Merge(const FRigVMTemplate& InOther)
 		int32 MatchingArguments = 0;
 		for(int32 ArgumentIndex = 0; ArgumentIndex < Arguments.Num(); ArgumentIndex++)
 		{
-			if(Arguments[ArgumentIndex].GetTypes()[PermutationIndex] ==
-				InOther.Arguments[ArgumentIndex].Types[0])
+			if(Arguments[ArgumentIndex].GetTypeIndices()[PermutationIndex] ==
+				InOther.Arguments[ArgumentIndex].TypeIndices[0])
 			{
 				MatchingArguments++;
 			}
@@ -970,7 +712,7 @@ bool FRigVMTemplate::Merge(const FRigVMTemplate& InOther)
 
 	for (int32 ArgumentIndex = 0; ArgumentIndex < Arguments.Num(); ArgumentIndex++)
 	{
-		if (InOther.Arguments[ArgumentIndex].Types.Num() != 1)
+		if (InOther.Arguments[ArgumentIndex].TypeIndices.Num() != 1)
 		{
 			return false;
 		}
@@ -979,18 +721,18 @@ bool FRigVMTemplate::Merge(const FRigVMTemplate& InOther)
 
 		// Add Other argument information into the TypeToPermutations map
 		{
-			const FString& OtherCPPType = InOther.Arguments[ArgumentIndex].Types[0].CPPType;
-			const int32 NewPermutationIndex = NewArgs[ArgumentIndex].Types.Num();
-			if (TArray<int32>* ArgTypePermutations = NewArgs[ArgumentIndex].TypeToPermutations.Find(OtherCPPType))
+			const int32 OtherTypeIndex = InOther.Arguments[ArgumentIndex].TypeIndices[0];
+			const int32 NewPermutationIndex = NewArgs[ArgumentIndex].TypeIndices.Num();
+			if (TArray<int32>* ArgTypePermutations = NewArgs[ArgumentIndex].TypeToPermutations.Find(OtherTypeIndex))
 			{
 				ArgTypePermutations->Add(NewPermutationIndex);
 			}
 			else
 			{
-				NewArgs[ArgumentIndex].TypeToPermutations.Add(OtherCPPType, {NewPermutationIndex});
+				NewArgs[ArgumentIndex].TypeToPermutations.Add(OtherTypeIndex, {NewPermutationIndex});
 			}
 		}
-		NewArgs[ArgumentIndex].Types.Add(InOther.Arguments[ArgumentIndex].Types[0]);
+		NewArgs[ArgumentIndex].TypeIndices.Add(InOther.Arguments[ArgumentIndex].TypeIndices[0]);
 	}
 
 	Arguments = NewArgs;
@@ -1007,12 +749,11 @@ const FRigVMTemplateArgument* FRigVMTemplate::FindArgument(const FName& InArgume
 	});
 }
 
-bool FRigVMTemplate::ArgumentSupportsType(const FName& InArgumentName, const FString& InCPPType, FRigVMTemplateArgumentType* OutType) const
+bool FRigVMTemplate::ArgumentSupportsTypeIndex(const FName& InArgumentName, int32 InTypeIndex, int32* OutTypeIndex) const
 {
 	if (const FRigVMTemplateArgument* Argument = FindArgument(InArgumentName))
 	{
-		return Argument->SupportsType(InCPPType, OutType);		
-		
+		return Argument->SupportsTypeIndex(InTypeIndex, OutTypeIndex);		
 	}
 	return false;
 }
@@ -1056,6 +797,8 @@ bool FRigVMTemplate::FullyResolve(FRigVMTemplate::FTypeMap& InOutTypes, int32& O
 
 bool FRigVMTemplate::Resolve(FTypeMap& InOutTypes, TArray<int32>& OutPermutationIndices, bool bAllowFloatingPointCasts) const
 {
+	const FRigVMRegistry& Registry = FRigVMRegistry::Get();
+	
 	FTypeMap InputTypes = InOutTypes;
 	InOutTypes.Reset();
 
@@ -1069,17 +812,17 @@ bool FRigVMTemplate::Resolve(FTypeMap& InOutTypes, TArray<int32>& OutPermutation
 	{
 		if (Argument.IsSingleton())
 		{
-			InOutTypes.Add(Argument.Name, Argument.Types[0]);
+			InOutTypes.Add(Argument.Name, Argument.TypeIndices[0]);
 			continue;
 		}
-		else if (const FRigVMTemplateArgumentType* InputType = InputTypes.Find(Argument.Name))
+		else if (const int32* InputType = InputTypes.Find(Argument.Name))
 		{
-			FRigVMTemplateArgumentType MatchedType = *InputType;
+			int32 MatchedType = *InputType;
 			bool bFoundMatch = false;
 			bool bFoundPerfectMatch = false;
-			for (int32 PermutationIndex = 0; PermutationIndex < Argument.Types.Num(); PermutationIndex++)
+			for (int32 PermutationIndex = 0; PermutationIndex < Argument.TypeIndices.Num(); PermutationIndex++)
 			{
-				if (!Argument.Types[PermutationIndex].Matches(InputType->CPPType, bAllowFloatingPointCasts))
+				if(!Registry.CanMatchTypes(Argument.TypeIndices[PermutationIndex], *InputType, bAllowFloatingPointCasts))
 				{
 					OutPermutationIndices.Remove(PermutationIndex);
 				}
@@ -1090,10 +833,10 @@ bool FRigVMTemplate::Resolve(FTypeMap& InOutTypes, TArray<int32>& OutPermutation
 					// if the type matches - but it's not the exact same
 					if(!bFoundPerfectMatch)
 					{
-						MatchedType = Argument.Types[PermutationIndex];
+						MatchedType = Argument.TypeIndices[PermutationIndex];
 
 						// if we found the perfect match - let's stop here
-						if(Argument.Types[PermutationIndex] == *InputType)
+						if(Argument.TypeIndices[PermutationIndex] == *InputType)
 						{
 							bFoundPerfectMatch = true;
 						}
@@ -1111,14 +854,14 @@ bool FRigVMTemplate::Resolve(FTypeMap& InOutTypes, TArray<int32>& OutPermutation
 					const TArray<int32> BeforePermutationIndices = OutPermutationIndices;
 					OutPermutationIndices.RemoveAll([Argument, MatchedType](int32 PermutationIndex) -> bool
 					{
-						return Argument.Types[PermutationIndex] != MatchedType;
+						return Argument.TypeIndices[PermutationIndex] != MatchedType;
 					});
 					if (OutPermutationIndices.IsEmpty() && bAllowFloatingPointCasts)
 					{
 						OutPermutationIndices = BeforePermutationIndices;
-						OutPermutationIndices.RemoveAll([Argument, MatchedType, InputType](int32 PermutationIndex) -> bool
+						OutPermutationIndices.RemoveAll([Argument, MatchedType, InputType, &Registry](int32 PermutationIndex) -> bool
 						{
-							return !Argument.Types[PermutationIndex].Matches(InputType->CPPType, true);
+							return !Registry.CanMatchTypes(Argument.TypeIndices[PermutationIndex], *InputType, true);
 						});
 					}
 				}
@@ -1129,23 +872,23 @@ bool FRigVMTemplate::Resolve(FTypeMap& InOutTypes, TArray<int32>& OutPermutation
 		const FRigVMTemplateArgument::EArrayType ArrayType = Argument.GetArrayType();
 		if(ArrayType == FRigVMTemplateArgument::EArrayType_Mixed)
 		{
-			InOutTypes.Add(Argument.Name, FRigVMTemplateArgumentType());
+			InOutTypes.Add(Argument.Name, RigVMTypeUtils::TypeIndex::WildCard);
 
-			if(const FRigVMTemplateArgumentType* InputType = InputTypes.Find(Argument.Name))
+			if(const int32* InputType = InputTypes.Find(Argument.Name))
 			{
-				if(InputType->IsArray())
+				if(Registry.IsArrayType(*InputType))
 				{
-					InOutTypes.FindChecked(Argument.Name) = FRigVMTemplateArgumentType::Array();
+					InOutTypes.FindChecked(Argument.Name) = RigVMTypeUtils::TypeIndex::WildCardArray;
 				}
 			}
 		}
 		else if(ArrayType == FRigVMTemplateArgument::EArrayType_ArrayValue)
 		{
-			InOutTypes.Add(Argument.Name, FRigVMTemplateArgumentType::Array());
+			InOutTypes.Add(Argument.Name, RigVMTypeUtils::TypeIndex::WildCardArray);
 		}
 		else
 		{
-			InOutTypes.Add(Argument.Name, FRigVMTemplateArgumentType());
+			InOutTypes.Add(Argument.Name, RigVMTypeUtils::TypeIndex::WildCard);
 		}
 	}
 
@@ -1154,7 +897,7 @@ bool FRigVMTemplate::Resolve(FTypeMap& InOutTypes, TArray<int32>& OutPermutation
 		InOutTypes.Reset();
 		for (const FRigVMTemplateArgument& Argument : Arguments)
 		{
-			InOutTypes.Add(Argument.Name, Argument.Types[OutPermutationIndices[0]]);
+			InOutTypes.Add(Argument.Name, Argument.TypeIndices[OutPermutationIndices[0]]);
 		}
 	}
 	else if (OutPermutationIndices.Num() > 1)
@@ -1163,7 +906,7 @@ bool FRigVMTemplate::Resolve(FTypeMap& InOutTypes, TArray<int32>& OutPermutation
 		{
 			if (Argument.IsSingleton(OutPermutationIndices))
 			{
-				InOutTypes.FindChecked(Argument.Name) = Argument.Types[OutPermutationIndices[0]];
+				InOutTypes.FindChecked(Argument.Name) = Argument.TypeIndices[OutPermutationIndices[0]];
 			}
 		}
 	}
@@ -1171,15 +914,17 @@ bool FRigVMTemplate::Resolve(FTypeMap& InOutTypes, TArray<int32>& OutPermutation
 	return !OutPermutationIndices.IsEmpty();
 }
 
-bool FRigVMTemplate::ResolveArgument(const FName& InArgumentName, const FRigVMTemplateArgumentType& InType,
+bool FRigVMTemplate::ResolveArgument(const FName& InArgumentName, const int32 InTypeIndex,
 	FTypeMap& InOutTypes) const
 {
-	auto RemoveWildCardTypes = [](const FTypeMap& InTypes)
+	const FRigVMRegistry& Registry = FRigVMRegistry::Get();
+
+	auto RemoveWildCardTypes = [&Registry](const FTypeMap& InTypes)
 	{
 		FTypeMap FilteredTypes;
 		for(const FTypePair& Pair: InTypes)
 		{
-			if(!Pair.Value.IsWildCard())
+			if(!Registry.IsWildCardType(Pair.Value))
 			{
 				FilteredTypes.Add(Pair);
 			}
@@ -1192,15 +937,15 @@ bool FRigVMTemplate::ResolveArgument(const FName& InArgumentName, const FRigVMTe
 
 	// first resolve with no types given except for the new argument type
 	FTypeMap ResolvedTypes;
-	ResolvedTypes.Add(InArgumentName, InType);
+	ResolvedTypes.Add(InArgumentName, InTypeIndex);
 	TArray<int32> PermutationIndices;
 	FTypeMap RemainingTypesToResolve;
 	
 	if(Resolve(ResolvedTypes, PermutationIndices, true))
 	{
 		// let's see if the input argument resolved into the expected type
-		const FRigVMTemplateArgumentType ResolvedInputType = ResolvedTypes.FindChecked(InArgumentName);
-		if(!ResolvedInputType.Matches(InType.CPPType, true))
+		const int32 ResolvedInputType = ResolvedTypes.FindChecked(InArgumentName);
+		if(!Registry.CanMatchTypes(ResolvedInputType, InTypeIndex, true))
 		{
 			return false;
 		}
diff --git a/Engine/Source/Runtime/RigVM/Public/RigVMCore/RigVMFunction.h b/Engine/Source/Runtime/RigVM/Public/RigVMCore/RigVMFunction.h
--- a/Engine/Source/Runtime/RigVM/Public/RigVMCore/RigVMFunction.h
+++ b/Engine/Source/Runtime/RigVM/Public/RigVMCore/RigVMFunction.h
@@ -62,6 +62,7 @@ struct RIGVM_API FRigVMFunction
 	int32 Index;
 	int32 TemplateIndex;
 	TArray<FRigVMFunctionArgument> Arguments;
+	mutable TArray<int32> ArgumentTypeIndices;
 
 	FRigVMFunction()
 		: Name(nullptr)
@@ -70,6 +71,7 @@ struct RIGVM_API FRigVMFunction
 		, Index(INDEX_NONE)
 		, TemplateIndex(INDEX_NONE)
 		, Arguments()
+		, ArgumentTypeIndices()
 	{
 	}
 
@@ -89,6 +91,7 @@ struct RIGVM_API FRigVMFunction
 	FString GetModuleName() const;
 	FString GetModuleRelativeHeaderPath() const;
 	const TArray<FRigVMFunctionArgument>& GetArguments() const { return Arguments; }
+	const TArray<int32>& GetArgumentTypeIndices() const;
 	bool IsAdditionalArgument(const FRigVMFunctionArgument& InArgument) const;
 	const FRigVMTemplate* GetTemplate() const;
 };
diff --git a/Engine/Source/Runtime/RigVM/Public/RigVMCore/RigVMRegistry.h b/Engine/Source/Runtime/RigVM/Public/RigVMCore/RigVMRegistry.h
--- a/Engine/Source/Runtime/RigVM/Public/RigVMCore/RigVMRegistry.h
+++ b/Engine/Source/Runtime/RigVM/Public/RigVMCore/RigVMRegistry.h
@@ -25,10 +25,85 @@ public:
 	// for example "FMyStruct::MyVirtualMethod"
 	void Register(const TCHAR* InName, FRigVMFunctionPtr InFunctionPtr, UScriptStruct* InStruct = nullptr, const TArray<FRigVMFunctionArgument>& InArguments = TArray<FRigVMFunctionArgument>());
 
+	// Initializes the registry by storing the defaults
+	void InitializeIfNeeded();
+	
 	// Refreshes the list and finds the function pointers
 	// based on the names.
 	void Refresh();
 
+	// Adds a type if it doesn't exist yet and returns its index.
+	// This function is not thead-safe
+	int32 FindOrAddType(const FRigVMTemplateArgumentType& InType);
+
+	// Returns the type index given a type
+	int32 GetTypeIndex(const FRigVMTemplateArgumentType& InType) const;
+
+	// Returns the type index given a cpp type and a type object
+	FORCEINLINE int32 GetTypeIndex(const FName& InCPPType, UObject* InCPPTypeObject) const
+	{
+		return GetTypeIndex(FRigVMTemplateArgumentType(InCPPType, InCPPTypeObject));
+	}
+
+	// Returns the type given its index
+	const FRigVMTemplateArgumentType& GetType(int32 InTypeIndex) const;
+
+	// Returns the number of types
+	FORCEINLINE int32 NumTypes() const { return Types.Num(); }
+
+	// Returns the type given only its cpp type
+	const FRigVMTemplateArgumentType& FindTypeFromCPPType(const FString& InCPPType) const;
+
+	// Returns the type index given only its cpp type
+	int32 GetTypeIndexFromCPPType(const FString& InCPPType) const;
+
+	// Returns true if the type is an array
+	bool IsArrayType(int32 InTypeIndex) const;
+
+	// Returns the dimensions of the array 
+	int32 GetArrayDimensionsForType(int32 InTypeIndex) const;
+
+	// Returns true if the type is a wildcard type
+	bool IsWildCardType(int32 InTypeIndex) const;
+
+	// Returns true if the types can be matched.
+	bool CanMatchTypes(int32 InTypeIndexA, int32 InTypeIndexB, bool bAllowFloatingPointCasts) const;
+
+	// Returns the list of compatible types for a given type
+	const TArray<int32>& GetCompatibleTypes(int32 InTypeIndex) const;
+
+	enum ETypeCategory
+	{
+		ETypeCategory_SingleAnyValue,
+		ETypeCategory_ArrayAnyValue,
+		ETypeCategory_ArrayArrayAnyValue,
+		ETypeCategory_SingleSimpleValue,
+		ETypeCategory_ArraySimpleValue,
+		ETypeCategory_ArrayArraySimpleValue,
+		ETypeCategory_SingleMathStructValue,
+		ETypeCategory_ArrayMathStructValue,
+		ETypeCategory_ArrayArrayMathStructValue,
+		ETypeCategory_SingleScriptStructValue,
+		ETypeCategory_ArrayScriptStructValue,
+		ETypeCategory_ArrayArrayScriptStructValue,
+		ETypeCategory_SingleEnumValue,
+		ETypeCategory_ArrayEnumValue,
+		ETypeCategory_ArrayArrayEnumValue,
+		ETypeCategory_SingleObjectValue,
+		ETypeCategory_ArrayObjectValue,
+		ETypeCategory_ArrayArrayObjectValue,
+		ETypeCategory_Invalid
+	};
+
+	// Returns all compatible types given a category
+	const TArray<int32>& GetTypesForCategory(ETypeCategory InCategory);
+
+	// Returns the type index of the array matching the given element type index
+	int32 GetArrayTypeFromBaseTypeIndex(int32 InTypeIndex) const;
+
+	// Returns the type index of the element matching the given array type index
+	int32 GetBaseTypeFromArrayTypeIndex(int32 InTypeIndex) const;
+
 	// Returns the function given its name (or nullptr)
 	const FRigVMFunction* FindFunction(const TCHAR* InName) const;
 
@@ -47,6 +122,20 @@ public:
 	// Defines and retrieves a template given its arguments
 	const FRigVMTemplate* GetOrAddTemplateFromArguments(const FName& InName, const TArray<FRigVMTemplateArgument>& InArguments);
 
+	// The list of base math types to automatically register 
+	inline static const TArray<UScriptStruct*> MathTypes = { 
+		TBaseStructure<FRotator>::Get(),
+		TBaseStructure<FQuat>::Get(),
+		TBaseStructure<FTransform>::Get(),
+		TBaseStructure<FLinearColor>::Get(),
+		TBaseStructure<FColor>::Get(),
+		TBaseStructure<FPlane>::Get(),
+		TBaseStructure<FVector>::Get(),
+		TBaseStructure<FVector2D>::Get(),
+		TBaseStructure<FVector4>::Get(),
+		TBaseStructure<FBox2D>::Get()
+	};
+
 private:
 
 	static const FName TemplateNameMetaName;
@@ -58,6 +147,41 @@ private:
 	// disable assignment operator
 	FRigVMRegistry& operator= (const FRigVMRegistry &InOther) = delete;
 
+	struct FTypeInfo
+	{
+		FTypeInfo()
+			: Type()
+			, BaseTypeIndex(INDEX_NONE)
+			, ArrayTypeIndex(INDEX_NONE)
+			, bIsArray(false)
+		{}
+		
+		FRigVMTemplateArgumentType Type;
+		int32 BaseTypeIndex;
+		int32 ArrayTypeIndex;
+		bool bIsArray;
+	};
+
+	FORCEINLINE static EObjectFlags DisallowedFlags()
+	{
+		return RF_BeginDestroyed | RF_FinishDestroyed;
+	}
+
+	FORCEINLINE static EObjectFlags NeededFlags()
+	{
+		return RF_Public;
+	}
+
+	static bool IsAllowedType(const FProperty* InProperty, bool bCheckFlags = true);
+	static bool IsAllowedType(const UEnum* InEnum);
+	static bool IsAllowedType(const UStruct* InStruct);
+	static bool IsAllowedType(const UClass* InClass);
+
+	// memory for all (known) types
+	// We use TChunkedArray because we need the memory locations to be stable, since we only ever add and never remove.
+	TArray<FTypeInfo> Types;
+	TMap<FRigVMTemplateArgumentType, int32> TypeToIndex;
+
 	// memory for all functions
 	// We use TChunkedArray because we need the memory locations to be stable, since we only ever add and never remove.
 	TChunkedArray<FRigVMFunction> Functions;
@@ -71,6 +195,9 @@ private:
 	// name lookup for templates
 	TMap<FName, int32> TemplateNotationToIndex;
 
+	// Maps storing the default types per type category
+	TMap<ETypeCategory, TArray<int32>> TypesPerCategory;
+
 	static FRigVMRegistry s_RigVMRegistry;
 	friend struct FRigVMStruct;
 };
diff --git a/Engine/Source/Runtime/RigVM/Public/RigVMCore/RigVMTemplate.h b/Engine/Source/Runtime/RigVM/Public/RigVMCore/RigVMTemplate.h
--- a/Engine/Source/Runtime/RigVM/Public/RigVMCore/RigVMTemplate.h
+++ b/Engine/Source/Runtime/RigVM/Public/RigVMCore/RigVMTemplate.h
@@ -14,29 +14,29 @@ struct RIGVM_API FRigVMTemplateArgumentType
 	GENERATED_BODY()
 
 	UPROPERTY()
-	FString CPPType;
+	FName CPPType;
 	
 	UPROPERTY()
 	TObjectPtr<UObject> CPPTypeObject; 
 
 	FRigVMTemplateArgumentType()
-		: CPPType()
+		: CPPType(NAME_None)
 		, CPPTypeObject(nullptr)
 	{
-		CPPType = RigVMTypeUtils::GetWildCardCPPType();
+		CPPType = RigVMTypeUtils::GetWildCardCPPTypeName();
 		CPPTypeObject = RigVMTypeUtils::GetWildCardCPPTypeObject();
 	}
 
-	FRigVMTemplateArgumentType(const FString& InCPPType, UObject* InCPPTypeObject = nullptr)
+	FRigVMTemplateArgumentType(const FName& InCPPType, UObject* InCPPTypeObject = nullptr)
 		: CPPType(InCPPType)
 		, CPPTypeObject(InCPPTypeObject)
 	{
-		check(!CPPType.IsEmpty());
+		check(!CPPType.IsNone());
 	}
 
 	static FRigVMTemplateArgumentType Array()
 	{
-		return FRigVMTemplateArgumentType(RigVMTypeUtils::GetWildCardArrayCPPType(), RigVMTypeUtils::GetWildCardCPPTypeObject());
+		return FRigVMTemplateArgumentType(RigVMTypeUtils::GetWildCardArrayCPPTypeName(), RigVMTypeUtils::GetWildCardCPPTypeObject());
 	}
 
 	bool operator == (const FRigVMTemplateArgumentType& InOther) const
@@ -54,55 +54,6 @@ struct RIGVM_API FRigVMTemplateArgumentType
 		return GetTypeHash(InType.CPPType);
 	}
 
-	FORCEINLINE bool Matches(const FString& InCPPType, bool bAllowFloatingPointCasts = true) const
-	{
-		if(CPPType == InCPPType)
-		{
-			return true;
-		}
-		if(bAllowFloatingPointCasts)
-		{
-			if(InCPPType == RigVMTypeUtils::FloatType && CPPType == RigVMTypeUtils::DoubleType)
-			{
-				return true;
-			}
-			if(InCPPType == RigVMTypeUtils::DoubleType && CPPType == RigVMTypeUtils::FloatType)
-			{
-				return true;
-			}
-			if(InCPPType == RigVMTypeUtils::FloatArrayType && CPPType == RigVMTypeUtils::DoubleArrayType)
-			{
-				return true;
-			}
-			if(InCPPType == RigVMTypeUtils::DoubleArrayType && CPPType == RigVMTypeUtils::FloatArrayType)
-			{
-				return true;
-			}
-		}
-		return false;
-	}
-
-	static TArray<FString> GetCompatibleTypes(const FString& InCPPType)
-	{
-		if(InCPPType == RigVMTypeUtils::FloatType)
-		{
-			return {RigVMTypeUtils::DoubleType};
-		}
-		if(InCPPType == RigVMTypeUtils::DoubleType)
-		{
-			return {RigVMTypeUtils::FloatType};
-		}
-		if(InCPPType == RigVMTypeUtils::FloatArrayType)
-		{
-			return {RigVMTypeUtils::DoubleArrayType};
-		}
-		if(InCPPType == RigVMTypeUtils::DoubleArrayType)
-		{
-			return {RigVMTypeUtils::FloatArrayType};
-		}
-		return {};
-	}
-
 	FName GetCPPTypeObjectPath() const
 	{
 		if(CPPTypeObject)
@@ -115,32 +66,32 @@ struct RIGVM_API FRigVMTemplateArgumentType
 	bool IsWildCard() const
 	{
 		return CPPTypeObject == RigVMTypeUtils::GetWildCardCPPTypeObject() ||
-			CPPType == RigVMTypeUtils::GetWildCardCPPType() ||
-			CPPType == RigVMTypeUtils::GetWildCardArrayCPPType();
+			CPPType == RigVMTypeUtils::GetWildCardCPPTypeName() ||
+			CPPType == RigVMTypeUtils::GetWildCardArrayCPPTypeName();
 	}
 
 	bool IsArray() const
 	{
-		return RigVMTypeUtils::IsArrayType(CPPType);
+		return RigVMTypeUtils::IsArrayType(CPPType.ToString());
 	}
 
 	FString GetBaseCPPType() const
 	{
 		if(IsArray())
 		{
-			return RigVMTypeUtils::BaseTypeFromArrayType(CPPType);
+			return RigVMTypeUtils::BaseTypeFromArrayType(CPPType.ToString());
 		}
-		return CPPType;
+		return CPPType.ToString();
 	}
 
 	void ConvertToArray() 
 	{
-		CPPType = RigVMTypeUtils::ArrayTypeFromBaseType(CPPType);
+		CPPType = *RigVMTypeUtils::ArrayTypeFromBaseType(CPPType.ToString());
 	}
 
 	void ConvertToBaseElement() 
 	{
-		CPPType = RigVMTypeUtils::BaseTypeFromArrayType(CPPType);
+		CPPType = *RigVMTypeUtils::BaseTypeFromArrayType(CPPType.ToString());
 	}
 };
 
@@ -162,34 +113,11 @@ struct RIGVM_API FRigVMTemplateArgument
 		EArrayType_Invalid
 	};
 
-	enum ETypeCategory
-	{
-		ETypeCategory_SingleAnyValue,
-		ETypeCategory_ArrayAnyValue,
-		ETypeCategory_ArrayArrayAnyValue,
-		ETypeCategory_SingleSimpleValue,
-		ETypeCategory_ArraySimpleValue,
-		ETypeCategory_ArrayArraySimpleValue,
-		ETypeCategory_SingleMathStructValue,
-		ETypeCategory_ArrayMathStructValue,
-		ETypeCategory_ArrayArrayMathStructValue,
-		ETypeCategory_SingleScriptStructValue,
-		ETypeCategory_ArrayScriptStructValue,
-		ETypeCategory_ArrayArrayScriptStructValue,
-		ETypeCategory_SingleEnumValue,
-		ETypeCategory_ArrayEnumValue,
-		ETypeCategory_ArrayArrayEnumValue,
-		ETypeCategory_SingleObjectValue,
-		ETypeCategory_ArrayObjectValue,
-		ETypeCategory_ArrayArrayObjectValue,
-		ETypeCategory_Invalid
-	};
-
 	// default constructor
 	FRigVMTemplateArgument();
 
-	FRigVMTemplateArgument(const FName& InName, ERigVMPinDirection InDirection, const FRigVMTemplateArgumentType& InType);
-	FRigVMTemplateArgument(const FName& InName, ERigVMPinDirection InDirection, const TArray<FRigVMTemplateArgumentType>& InTypes);
+	FRigVMTemplateArgument(const FName& InName, ERigVMPinDirection InDirection, int32 InType);
+	FRigVMTemplateArgument(const FName& InName, ERigVMPinDirection InDirection, const TArray<int32>& InTypeIndices);
 
 	// returns the name of the argument
 	const FName& GetName() const { return Name; }
@@ -198,26 +126,25 @@ struct RIGVM_API FRigVMTemplateArgument
 	ERigVMPinDirection GetDirection() const { return Direction; }
 
 	// returns true if this argument supports a given type across a set of permutations
-	bool SupportsType(const FString& InCPPType, FRigVMTemplateArgumentType* OutType = nullptr) const;
+	bool SupportsTypeIndex(int32 InTypeIndex, int32* OutTypeIndex = nullptr) const;
 
 	// returns the flat list of types (including duplicates) of this argument
-	const TArray<FRigVMTemplateArgumentType>& GetTypes() const;
+	const TArray<int32>& GetTypeIndices() const;
 
 	// returns an array of all of the supported types
-	TArray<FRigVMTemplateArgumentType> GetSupportedTypes(const TArray<int32>& InPermutationIndices = TArray<int32>()) const;
+	TArray<int32> GetSupportedTypeIndices(const TArray<int32>& InPermutationIndices = TArray<int32>()) const;
 
-	// returns an array of all supported types as strings
+#if WITH_EDITOR
+	// returns an array of all supported types as strings. this is used for automated testing only.
 	TArray<FString> GetSupportedTypeStrings(const TArray<int32>& InPermutationIndices = TArray<int32>()) const;
-
+#endif
+	
 	// returns true if an argument is singleton (same type for all variants)
 	bool IsSingleton(const TArray<int32>& InPermutationIndices = TArray<int32>()) const;
 
 	// returns true if the argument uses an array container
 	EArrayType GetArrayType() const;
 	
-	// Returns all compatible types given a category
-	static const TArray<FRigVMTemplateArgumentType>& GetCompatibleTypes(ETypeCategory InCategory);
-
 protected:
 
 	UPROPERTY()
@@ -230,9 +157,9 @@ protected:
 	ERigVMPinDirection Direction;
 
 	UPROPERTY()
-	TArray<FRigVMTemplateArgumentType> Types;
+	TArray<int32> TypeIndices;
 
-	TMap<FString, TArray<int32>> TypeToPermutations;
+	TMap<int32, TArray<int32>> TypeToPermutations;
 
 	// constructor from a property
 	FRigVMTemplateArgument(FProperty* InProperty);
@@ -254,8 +181,8 @@ struct RIGVM_API FRigVMTemplate
 	GENERATED_BODY()
 public:
 
-	typedef TMap<FName, FRigVMTemplateArgumentType> FTypeMap;
-	typedef TPair<FName, FRigVMTemplateArgumentType> FTypePair;
+	typedef TMap<FName, int32> FTypeMap;
+	typedef TPair<FName, int32> FTypePair;
 
 	// Default constructor
 	FRigVMTemplate();
@@ -285,7 +212,7 @@ public:
 	const FRigVMTemplateArgument* FindArgument(const FName& InArgumentName) const;
 
 	// returns true if a given arg supports a type
-	bool ArgumentSupportsType(const FName& InArgumentName, const FString& InCPPType, FRigVMTemplateArgumentType* OutType = nullptr) const;
+	bool ArgumentSupportsTypeIndex(const FName& InArgumentName, int32 InTypeIndex, int32* OutTypeIndex = nullptr) const;
 
 	// returns the number of permutations supported by this template
 	int32 NumPermutations() const { return Permutations.Num(); }
@@ -306,7 +233,7 @@ public:
 	bool Resolve(FTypeMap& InOutTypes, TArray<int32> & OutPermutationIndices, bool bAllowFloatingPointCasts) const;
 
 	// returns true if the template can resolve an argument to a new type
-	bool ResolveArgument(const FName& InArgumentName, const FRigVMTemplateArgumentType& InType, FTypeMap& InOutTypes) const;
+	bool ResolveArgument(const FName& InArgumentName, const int32 InTypeIndex, FTypeMap& InOutTypes) const;
 
 	// returns true if a given argument is valid for a template
 	static bool IsValidArgumentForTemplate(const FRigVMTemplateArgument& InArgument);
diff --git a/Engine/Source/Runtime/RigVM/Public/RigVMTypeUtils.h b/Engine/Source/Runtime/RigVM/Public/RigVMTypeUtils.h
--- a/Engine/Source/Runtime/RigVM/Public/RigVMTypeUtils.h
+++ b/Engine/Source/Runtime/RigVM/Public/RigVMTypeUtils.h
@@ -45,6 +45,27 @@ namespace RigVMTypeUtils
 	const FName FNameArrayTypeName = *FNameArrayType;
 	const FName FStringArrayTypeName = *FStringArrayType;
 
+	class RIGVM_API TypeIndex
+	{
+	public:
+		static inline int32 Bool = INDEX_NONE;	
+		static inline int32 Float = INDEX_NONE;	
+		static inline int32 Double = INDEX_NONE;	
+		static inline int32 Int32 = INDEX_NONE;	
+		static inline int32 UInt8 = INDEX_NONE;	
+		static inline int32 FName = INDEX_NONE;	
+		static inline int32 FString = INDEX_NONE;
+		static inline int32 WildCard = INDEX_NONE;	
+		static inline int32 BoolArray = INDEX_NONE;	
+		static inline int32 FloatArray = INDEX_NONE;	
+		static inline int32 DoubleArray = INDEX_NONE;	
+		static inline int32 Int32Array = INDEX_NONE;	
+		static inline int32 UInt8Array = INDEX_NONE;	
+		static inline int32 FNameArray = INDEX_NONE;	
+		static inline int32 FStringArray = INDEX_NONE;	
+		static inline int32 WildCardArray = INDEX_NONE;	
+	};
+
 	// Returns true if the type specified is an array
 	FORCEINLINE bool IsArrayType(const FString& InCPPType)
 	{
@@ -89,18 +110,30 @@ namespace RigVMTypeUtils
 		return WildCardTypeObject;
 	}
 
-	static const FString GetWildCardCPPType()
+	static const FString& GetWildCardCPPType()
 	{
 		static const FString WildCardCPPType = FRigVMUnknownType::StaticStruct()->GetStructCPPName(); 
 		return WildCardCPPType;
 	}
 
-	static const FString GetWildCardArrayCPPType()
+	static const FName& GetWildCardCPPTypeName()
+	{
+		static const FName WildCardCPPTypeName = *GetWildCardCPPType(); 
+		return WildCardCPPTypeName;
+	}
+
+	static const FString& GetWildCardArrayCPPType()
 	{
 		static const FString WildCardArrayCPPType = ArrayTypeFromBaseType(GetWildCardCPPType()); 
 		return WildCardArrayCPPType;
 	}
 
+	static const FName& GetWildCardArrayCPPTypeName()
+	{
+		static const FName WildCardArrayCPPTypeName = *GetWildCardArrayCPPType(); 
+		return WildCardArrayCPPTypeName;
+	}
+
 	FORCEINLINE FString PostProcessCPPType(const FString& InCPPType, UObject* InCPPTypeObject)
 	{
 		FString CPPType = InCPPType;
