RigVM: Fix static analysis pattern in RigVMSelectNode.cpp

diff --git a/Engine/Source/Developer/RigVMDeveloper/Private/RigVMModel/Nodes/RigVMSelectNode.cpp b/Engine/Source/Developer/RigVMDeveloper/Private/RigVMModel/Nodes/RigVMSelectNode.cpp
--- a/Engine/Source/Developer/RigVMDeveloper/Private/RigVMModel/Nodes/RigVMSelectNode.cpp
+++ b/Engine/Source/Developer/RigVMDeveloper/Private/RigVMModel/Nodes/RigVMSelectNode.cpp
@@ -77,7 +77,7 @@ const FRigVMTemplate* URigVMSelectNode::GetTemplate() const
 					ValueTypeIndex = InTypeIndex;
 					ResultTypeIndex = FRigVMRegistry::Get().GetBaseTypeFromArrayTypeIndex(InTypeIndex);
 				}
-				else if(InArgumentName == ValueFName)
+				else if(InArgumentName == ResultFName)
 				{
 					ValueTypeIndex = FRigVMRegistry::Get().GetArrayTypeFromBaseTypeIndex(InTypeIndex);;
 					ResultTypeIndex = InTypeIndex;
