RigVM: Use clone when copying an array into a variable

diff --git a/Engine/Source/Developer/RigVMDeveloper/Private/RigVMCompiler/RigVMCompiler.cpp b/Engine/Source/Developer/RigVMDeveloper/Private/RigVMCompiler/RigVMCompiler.cpp
--- a/Engine/Source/Developer/RigVMDeveloper/Private/RigVMCompiler/RigVMCompiler.cpp
+++ b/Engine/Source/Developer/RigVMDeveloper/Private/RigVMCompiler/RigVMCompiler.cpp
@@ -1843,8 +1843,23 @@ void URigVMCompiler::AddCopyOperator(const FRigVMCopyOp& InOp, const FRigVMAssig
 			AddCopyOperator(CopyOpInfo, WorkData, false);
 		}
 	}
+
+	bool bAddCopyOp = true;
+
+	// if we are copying into an array variable
+	if(const URigVMPin* Pin = InTargetExpr->GetPin())
+	{
+		if(Pin->IsArray() && Pin->GetNode()->IsA<URigVMVariableNode>())
+		{
+			WorkData.VM->GetByteCode().AddArrayCloneOp(InOp.Source, InOp.Target);
+			bAddCopyOp = false;
+		}
+	}
 	
-	WorkData.VM->GetByteCode().AddCopyOp(InOp);
+	if(bAddCopyOp)
+	{
+		WorkData.VM->GetByteCode().AddCopyOp(InOp);
+	}
 
 	int32 InstructionIndex = WorkData.VM->GetByteCode().GetNumInstructions() - 1;
 	if (Settings.SetupNodeInstructionIndex)
