RigVM: Only use array clone for top level properties

diff --git a/Engine/Source/Developer/RigVMDeveloper/Private/RigVMCompiler/RigVMCompiler.cpp b/Engine/Source/Developer/RigVMDeveloper/Private/RigVMCompiler/RigVMCompiler.cpp
--- a/Engine/Source/Developer/RigVMDeveloper/Private/RigVMCompiler/RigVMCompiler.cpp
+++ b/Engine/Source/Developer/RigVMDeveloper/Private/RigVMCompiler/RigVMCompiler.cpp
@@ -1851,8 +1851,12 @@ void URigVMCompiler::AddCopyOperator(const FRigVMCopyOp& InOp, const FRigVMAssig
 	{
 		if(Pin->IsArray() && Pin->GetNode()->IsA<URigVMVariableNode>())
 		{
-			WorkData.VM->GetByteCode().AddArrayCloneOp(InOp.Source, InOp.Target);
-			bAddCopyOp = false;
+			if(InOp.Source.GetRegisterOffset() == INDEX_NONE &&
+				InOp.Target.GetRegisterOffset() == INDEX_NONE)
+			{ 
+				WorkData.VM->GetByteCode().AddArrayCloneOp(InOp.Source, InOp.Target);
+				bAddCopyOp = false;
+			}
 		}
 	}
 	
