Control Rig: Fix crash when connecting a cast node after resolving to a different type

diff --git a/Engine/Source/Developer/RigVMDeveloper/Private/RigVMModel/RigVMController.cpp b/Engine/Source/Developer/RigVMDeveloper/Private/RigVMModel/RigVMController.cpp
--- a/Engine/Source/Developer/RigVMDeveloper/Private/RigVMModel/RigVMController.cpp
+++ b/Engine/Source/Developer/RigVMDeveloper/Private/RigVMModel/RigVMController.cpp
@@ -14959,7 +14959,7 @@ bool URigVMController::FullyResolveTemplateNode(URigVMTemplateNode* InNode, int3
 	const FRigVMTemplate* Template = InNode->GetTemplate();
 	int32 InputPermutation = InPermutationIndex;
 	
-	// Figure out the permutation index fromt the pin types
+	// Figure out the permutation index from the pin types
 	if (InPermutationIndex == INDEX_NONE)
 	{
 		FRigVMTemplate::FTypeMap TypeMap;
@@ -15461,16 +15461,13 @@ bool URigVMController::ResolveWildCardPin(URigVMPin* InPin, int32 InTypeIndex, b
 			return false;
 		}
 
-		const TArray<int32>& FilteredPermutations = TemplateNode->GetFilteredPermutationsIndices();
-		if (FilteredPermutations.Num() == 1)
+		TArray<FRigVMTemplatePreferredType> NewPreferredPermutationTypes = TemplateNode->PreferredPermutationPairs;
+		NewPreferredPermutationTypes.Add({FRigVMTemplatePreferredType(InPin->GetFName(), InTypeIndex)});
+		if (bSetupUndoRedo)
 		{
-			const TArray<FRigVMTemplatePreferredType> NewPreferredPermutationTypes = TemplateNode->GetArgumentTypesForPermutation(FilteredPermutations[0]);
-			if (bSetupUndoRedo)
-			{
-				ActionStack->AddAction(FRigVMSetPreferredTemplatePermutationsAction(TemplateNode, NewPreferredPermutationTypes));
-			}
-			TemplateNode->PreferredPermutationPairs = NewPreferredPermutationTypes;
+			ActionStack->AddAction(FRigVMSetPreferredTemplatePermutationsAction(TemplateNode, NewPreferredPermutationTypes));
 		}
+		TemplateNode->PreferredPermutationPairs = NewPreferredPermutationTypes;
 		if (bSetupUndoRedo)
 		{
 			ActionStack->EndAction(Action);
@@ -15556,95 +15553,105 @@ bool URigVMController::UpdateFilteredPermutations(URigVMPin* InPin, const TArray
 bool URigVMController::UpdateTemplateNodePinTypes(URigVMTemplateNode* InNode, bool bSetupUndoRedo)
 {
 	bool bAnyTypeChanged = false;
-	for(int32 PinIndex=0; PinIndex < InNode->GetPins().Num(); ++PinIndex)
+
+	// First unresolve any pins that need unresolving, then resolve everything else
+	for (int32 UnresolveResolve=0; UnresolveResolve<2; ++UnresolveResolve)
 	{
-		URigVMPin* Pin = InNode->GetPins()[PinIndex];
-		if (Pin->GetDirection() == ERigVMPinDirection::Hidden)
+		for(int32 PinIndex=0; PinIndex < InNode->GetPins().Num(); ++PinIndex)
 		{
-			continue;
-		}
+			URigVMPin* Pin = InNode->GetPins()[PinIndex];
+			if (UnresolveResolve == 0 && Pin->IsWildCard())
+			{
+				continue;
+			}
+			
+			if (Pin->GetDirection() == ERigVMPinDirection::Hidden)
+			{
+				continue;
+			}
 		
-		TArray<int32> Types = InNode->GetFilteredTypesForPin(Pin);
-		if (Types.IsEmpty())
-		{
-			continue;
-		}
+			TArray<int32> Types = InNode->GetFilteredTypesForPin(Pin);
+			if (Types.IsEmpty())
+			{
+				continue;
+			}
 
-		if (Types.Num() > 1)
-		{
-			
-			bool bCanReduceToSingleType = true;
-			for (int32 i=1; i<Types.Num(); ++i)
+			if (Types.Num() > 1)
 			{
-				if (!FRigVMRegistry::Get().CanMatchTypes(Types[0], Types[i], true))
+			
+				bool bCanReduceToSingleType = true;
+				for (int32 i=1; i<Types.Num(); ++i)
 				{
-					bCanReduceToSingleType = false;
-					break;
+					if (!FRigVMRegistry::Get().CanMatchTypes(Types[0], Types[i], true))
+					{
+						bCanReduceToSingleType = false;
+						break;
+					}
 				}
-			}
 			
-			// Select the same Pin type if it has one
-			if (bCanReduceToSingleType)
-			{
-				int32 PreferredIndex = 0;
-				if (!Pin->IsWildCard())
+				// Select the same Pin type if it has one
+				if (bCanReduceToSingleType)
 				{
-					for (int32 i=0; i<Types.Num(); ++i)
+					int32 PreferredIndex = 0;
+					if (!Pin->IsWildCard())
 					{
-						if (Types[i] == Pin->GetTypeIndex())
+						for (int32 i=0; i<Types.Num(); ++i)
 						{
-							PreferredIndex = i;
-							break;
+							if (Types[i] == Pin->GetTypeIndex())
+							{
+								PreferredIndex = i;
+								break;
+							}
 						}
 					}
+					Types = {Types[PreferredIndex]};
 				}
-				Types = {Types[PreferredIndex]};
 			}
-		}
 		
-		if (Types.Num() > 1)
-		{
-			// Unresolve
-			if (Pin->HasInjectedNodes())
+			if (Types.Num() > 1)
 			{
-				EjectNodeFromPin(Pin, bSetupUndoRedo);
-			}
+				// Unresolve
+				if (Pin->HasInjectedNodes())
+				{
+					EjectNodeFromPin(Pin, bSetupUndoRedo);
+				}
 			
-			const FRigVMTemplateArgument* Argument = InNode->GetTemplate()->FindArgument(*Pin->GetName());
-			const FRigVMTemplateArgument::EArrayType ArrayType = Argument->GetArrayType();
+				const FRigVMTemplateArgument* Argument = InNode->GetTemplate()->FindArgument(*Pin->GetName());
+				const FRigVMTemplateArgument::EArrayType ArrayType = Argument->GetArrayType();
 			
-			FString CPPType = RigVMTypeUtils::GetWildCardCPPType();
-			UObject* CPPObjectType = RigVMTypeUtils::GetWildCardCPPTypeObject();
+				FString CPPType = RigVMTypeUtils::GetWildCardCPPType();
+				UObject* CPPObjectType = RigVMTypeUtils::GetWildCardCPPTypeObject();
 
-			if (ArrayType == FRigVMTemplateArgument::EArrayType_ArrayValue)
-			{
-				CPPType = RigVMTypeUtils::GetWildCardArrayCPPType();
-			}
-			else if(ArrayType == FRigVMTemplateArgument::EArrayType_Mixed)
-			{				
-				CPPType = Pin->IsArray() ? RigVMTypeUtils::GetWildCardArrayCPPType() : RigVMTypeUtils::GetWildCardCPPType();
-			}
+				if (ArrayType == FRigVMTemplateArgument::EArrayType_ArrayValue)
+				{
+					CPPType = RigVMTypeUtils::GetWildCardArrayCPPType();
+				}
+				else if(ArrayType == FRigVMTemplateArgument::EArrayType_Mixed)
+				{				
+					CPPType = Pin->IsArray() ? RigVMTypeUtils::GetWildCardArrayCPPType() : RigVMTypeUtils::GetWildCardCPPType();
+				}
 			
-			if (Pin->GetCPPType() != CPPType || Pin->GetCPPTypeObject() != CPPObjectType)
+				if (Pin->GetCPPType() != CPPType || Pin->GetCPPTypeObject() != CPPObjectType)
+				{
+					bAnyTypeChanged = true;
+					ChangePinType(Pin, CPPType, CPPObjectType, bSetupUndoRedo, false, false, false);
+				}
+			}
+			else if (Types.Num() == 1)
 			{
-				bAnyTypeChanged = true;
-				ChangePinType(Pin, CPPType, CPPObjectType, bSetupUndoRedo, false, false, false);
+				// Resolve
+				if (Pin->GetTypeIndex() != Types[0])
+				{
+					bAnyTypeChanged = true;
+					ChangePinType(Pin, Types[0], bSetupUndoRedo, false, false, false);
+				}
 			}
-		}
-		else if (Types.Num() == 1)
-		{
-			// Resolve
-			if (Pin->GetTypeIndex() != Types[0])
+			else
 			{
-				bAnyTypeChanged = true;
-				ChangePinType(Pin, Types[0], bSetupUndoRedo, false, false, false);
+				ensure(false);
+				return false;
 			}
 		}
-		else
-		{
-			ensure(false);
-			return false;
-		}
 	}
 
 	return bAnyTypeChanged;
