RigVM: Dispatch Factory Support

* Moved TypeIndex from int32 to typedef for easier debugging
* Implemented Dispatch Factory
* Factories stored under registry
* FindFunction to invoke dispatch
* Runtime Type checking on the memory handle
* Implemented dispatch rigvm node
* Unit testing for numeric add test both for low level + controller layer

diff --git a/Engine/Plugins/Animation/ControlRig/Source/ControlRig/Private/Units/RigDispatchFactory.cpp b/Engine/Plugins/Animation/ControlRig/Source/ControlRig/Private/Units/RigDispatchFactory.cpp
new file mode 100644
--- /dev/null
+++ b/Engine/Plugins/Animation/ControlRig/Source/ControlRig/Private/Units/RigDispatchFactory.cpp
@@ -0,0 +1,71 @@
+// Copyright Epic Games, Inc. All Rights Reserved.
+
+#include "Units/RigDispatchFactory.h"
+
+#if WITH_EDITOR
+
+FString FRigDispatchFactory::GetArgumentDefaultValue(const FName& InArgumentName, TRigVMTypeIndex InTypeIndex) const
+{
+	if(FRigVMRegistry::Get().IsArrayType(InTypeIndex))
+	{
+		static const FString EmptyArrayString = TEXT("()");
+		return EmptyArrayString;
+	}
+	if(InTypeIndex == RigVMTypeUtils::TypeIndex::Bool)
+	{
+		static const FString FalseString = TEXT("False");
+		return FalseString;
+	}
+	if(InTypeIndex == RigVMTypeUtils::TypeIndex::Float || InTypeIndex == RigVMTypeUtils::TypeIndex::Double)
+	{
+		static const FString ZeroString = TEXT("0.000000");
+		return ZeroString;
+	}
+	if(InTypeIndex == RigVMTypeUtils::TypeIndex::Int32)
+	{
+		static const FString ZeroString = TEXT("0");
+		return ZeroString;
+	}
+	if(InTypeIndex == RigVMTypeUtils::TypeIndex::FName)
+	{
+		return FName(NAME_None).ToString();
+	}
+	if(InTypeIndex == FRigVMRegistry::Get().GetTypeIndex<FVector2D>())
+	{
+		static FString DefaultValueString = GetDefaultValueForStruct(FVector2D::ZeroVector);
+		return DefaultValueString;
+	}
+	if(InTypeIndex == FRigVMRegistry::Get().GetTypeIndex<FVector>())
+	{
+		static FString DefaultValueString = GetDefaultValueForStruct(FVector::ZeroVector); 
+		return DefaultValueString;
+	}
+	if(InTypeIndex == FRigVMRegistry::Get().GetTypeIndex<FRotator>())
+	{
+		static FString DefaultValueString = GetDefaultValueForStruct(FRotator::ZeroRotator); 
+		return DefaultValueString;
+	}
+	if(InTypeIndex == FRigVMRegistry::Get().GetTypeIndex<FQuat>())
+	{
+		static FString DefaultValueString = GetDefaultValueForStruct(FQuat::Identity); 
+		return DefaultValueString;
+	}
+	if(InTypeIndex == FRigVMRegistry::Get().GetTypeIndex<FTransform>())
+	{
+		static FString DefaultValueString = GetDefaultValueForStruct(FTransform::Identity); 
+		return DefaultValueString;
+	}
+	if(InTypeIndex == FRigVMRegistry::Get().GetTypeIndex<FEulerTransform>())
+	{
+		static FString DefaultValueString = GetDefaultValueForStruct(FEulerTransform::Identity); 
+		return DefaultValueString;
+	}
+	if(InTypeIndex == FRigVMRegistry::Get().GetTypeIndex<FRigElementKey>())
+	{
+		static FString DefaultValueString = GetDefaultValueForStruct(FRigElementKey(NAME_None, ERigElementType::Bone)); 
+		return DefaultValueString;
+	}
+	return FRigVMDispatchFactory::GetArgumentDefaultValue(InArgumentName, InTypeIndex);
+}
+
+#endif
\ No newline at end of file
diff --git a/Engine/Plugins/Animation/ControlRig/Source/ControlRig/Public/Units/RigDispatchFactory.h b/Engine/Plugins/Animation/ControlRig/Source/ControlRig/Public/Units/RigDispatchFactory.h
new file mode 100644
--- /dev/null
+++ b/Engine/Plugins/Animation/ControlRig/Source/ControlRig/Public/Units/RigDispatchFactory.h
@@ -0,0 +1,82 @@
+// Copyright Epic Games, Inc. All Rights Reserved.
+
+#pragma once
+
+#include "ControlRigDefines.h"
+#include "RigVMCore/RigVMDispatchFactory.h"
+#include "RigUnitContext.h"
+#include "ControlRig.h"
+#include "RigDispatchFactory.generated.h"
+
+/** Base class for all rig dispatch factories */
+USTRUCT(BlueprintType, meta=(Abstract, BlueprintInternalUseOnlyHierarchical))
+struct CONTROLRIG_API FRigDispatchFactory : public FRigVMDispatchFactory
+{
+	GENERATED_BODY()
+
+	FORCEINLINE virtual void RegisterDependencyTypes() const override
+	{
+		FRigVMRegistry::Get().FindOrAddType(FControlRigExecuteContext::StaticStruct());
+		FRigVMRegistry::Get().FindOrAddType(FRigElementKey::StaticStruct());
+    	FRigVMRegistry::Get().FindOrAddType(FCachedRigElement::StaticStruct());
+	}
+
+	FORCEINLINE virtual TArray<TPair<FName,FString>> GetOpaqueArguments() const override
+	{
+		static const TArray<TPair<FName,FString>> OpaqueArguments = {
+			TPair<FName,FString>(TEXT("Context"), TEXT("const FRigUnitContext&"))
+		};
+		return OpaqueArguments;
+	}
+
+#if WITH_EDITOR
+
+	virtual FString GetArgumentDefaultValue(const FName& InArgumentName, TRigVMTypeIndex InTypeIndex) const override;
+
+#endif
+
+	template <
+		typename T,
+		typename TEnableIf<TRigVMIsBaseStructure<T>::Value, T>::Type* = nullptr
+	>
+	FORCEINLINE static FString GetDefaultValueForStruct(const T& InValue)
+	{
+		static FString ValueString;
+		if(ValueString.IsEmpty())
+		{
+			TBaseStructure<T>::Get()->ExportText(ValueString, &InValue, &InValue, nullptr, PPF_None, nullptr);
+		}
+		return ValueString;
+	}
+
+	template <
+		typename T,
+		typename TEnableIf<TModels<CRigVMUStruct, T>::Value>::Type * = nullptr
+	>
+	FORCEINLINE static FString GetDefaultValueForStruct(const T& InValue)
+	{
+		static FString ValueString;
+		if(ValueString.IsEmpty())
+		{
+			T::StaticStruct()->ExportText(ValueString, &InValue, &InValue, nullptr, PPF_None, nullptr);
+		}
+		return ValueString;
+	}
+
+	FORCEINLINE static const FRigUnitContext& GetRigUnitContext(const FRigVMExtendedExecuteContext& InContext)
+	{
+		return *(const FRigUnitContext*)InContext.OpaqueArguments[0];
+	}
+
+#if WITH_EDITOR
+	FORCEINLINE bool CheckArgumentType(bool bCondition, const FName& InArgumentName) const
+	{
+		if(!bCondition)
+		{
+			UE_LOG(LogControlRig, Error, TEXT("Fatal: '%s' Argument '%s' has incorrect type."), *GetFactoryName().ToString(), *InArgumentName.ToString())
+		}
+		return bCondition;
+	}
+#endif
+};
+
diff --git a/Engine/Plugins/Animation/ControlRig/Source/ControlRigEditor/Private/ControlRigEditorModule.cpp b/Engine/Plugins/Animation/ControlRig/Source/ControlRigEditor/Private/ControlRigEditorModule.cpp
--- a/Engine/Plugins/Animation/ControlRig/Source/ControlRigEditor/Private/ControlRigEditorModule.cpp
+++ b/Engine/Plugins/Animation/ControlRig/Source/ControlRigEditor/Private/ControlRigEditorModule.cpp
@@ -131,6 +131,7 @@
 #include "RigVMModel/Nodes/RigVMAggregateNode.h"
 #include "RigVMUserWorkflowRegistry.h"
 #include "Units/ControlRigNodeWorkflow.h"
+#include "Units/RigDispatchFactory.h"
 
 #define LOCTEXT_NAMESPACE "ControlRigEditorModule"
 
@@ -997,10 +998,19 @@ void FControlRigEditorModule::GetTypeActions(UControlRigBlueprint* CRB, FBluepri
 		return;
 	}
 
+	FRigVMRegistry& Registry = FRigVMRegistry::Get();
+	Registry.Refresh();
+
 #if UE_RIGVM_ENABLE_TEMPLATE_NODES
 	
-	for (const FRigVMTemplate& Template : FRigVMRegistry::Get().GetTemplates())
+	for (const FRigVMTemplate& Template : Registry.GetTemplates())
 	{
+		// factories are registered below
+		if(Template.UsesDispatch())
+		{
+			continue;
+		}
+
 		// ignore templates that have only one permutation
 		if (Template.NumPermutations() <= 1)
 		{
@@ -1022,10 +1032,38 @@ void FControlRigEditorModule::GetTypeActions(UControlRigBlueprint* CRB, FBluepri
 		ActionRegistrar.AddBlueprintAction(ActionKey, NodeSpawner);
 	};
 
+	for (const FRigVMDispatchFactory* Factory : Registry.GetFactories())
+	{
+		// allow empty opaque arguments or the control rig notation
+		const TArray<TPair<FName,FString>>& OpaqueArguments = Factory->GetOpaqueArguments();
+		if(!OpaqueArguments.IsEmpty())
+		{
+			static const FRigDispatchFactory ControlRigFactory;
+			if(ControlRigFactory.GetOpaqueArguments() != OpaqueArguments)
+			{
+				continue;
+			}
+		}
+
+		const FRigVMTemplate* Template = Factory->GetTemplate();
+		if(Template == nullptr)
+		{
+			continue;
+		}
+
+		FText NodeCategory = FText::FromString(Factory->GetCategory());
+		FText MenuDesc = FText::FromString(Factory->GetNodeTitle(FRigVMTemplateTypeMap()));
+		FText ToolTip = Factory->GetNodeTooltip(FRigVMTemplateTypeMap());
+
+		UBlueprintNodeSpawner* NodeSpawner = UControlRigTemplateNodeSpawner::CreateFromNotation(Template->GetNotation(), MenuDesc, NodeCategory, ToolTip);
+		check(NodeSpawner != nullptr);
+		ActionRegistrar.AddBlueprintAction(ActionKey, NodeSpawner);
+	};
+
 #endif
 
 	// Add all rig units
-	for(const FRigVMFunction& Function : FRigVMRegistry::Get().GetFunctions())
+	for(const FRigVMFunction& Function : Registry.GetFunctions())
 	{
 		UScriptStruct* Struct = Function.Struct;
 		if (Struct == nullptr || !Struct->IsChildOf(FRigUnit::StaticStruct()))
diff --git a/Engine/Plugins/Animation/ControlRig/Source/ControlRigEditor/Private/Graph/SControlRigGraphChangePinType.cpp b/Engine/Plugins/Animation/ControlRig/Source/ControlRigEditor/Private/Graph/SControlRigGraphChangePinType.cpp
--- a/Engine/Plugins/Animation/ControlRig/Source/ControlRigEditor/Private/Graph/SControlRigGraphChangePinType.cpp
+++ b/Engine/Plugins/Animation/ControlRig/Source/ControlRigEditor/Private/Graph/SControlRigGraphChangePinType.cpp
@@ -192,8 +192,8 @@ void SControlRigChangePinType::FillPinTypeMenu(FMenuBuilder& MenuBuilder)
 		FRigVMTemplateArgumentType Type;
 	};
 
-	typedef TPair<int32, FArgumentInfo> FTypePair;
-	TMap<int32, FArgumentInfo> Types;
+	typedef TPair<TRigVMTypeIndex, FArgumentInfo> FTypePair;
+	TMap<TRigVMTypeIndex, FArgumentInfo> Types;
 	for(URigVMPin* ModelPin : ModelPins)
 	{
 		if(!ModelPin->IsRootPin())
@@ -208,12 +208,12 @@ void SControlRigChangePinType::FillPinTypeMenu(FMenuBuilder& MenuBuilder)
 				const FName ArgumentName = ModelPin->GetFName();
 				if(const FRigVMTemplateArgument* Argument = Template->FindArgument(ArgumentName))
 				{
-					const TArray<int32>& AllArgumentTypes = Argument->GetTypeIndices();
-					const TArray<int32>& FilteredArgumentTypes = Argument->GetSupportedTypeIndices(TemplateNode->GetFilteredPermutationsIndices());
+					const TArray<TRigVMTypeIndex>& AllArgumentTypes = Argument->GetTypeIndices();
+					const TArray<TRigVMTypeIndex>& FilteredArgumentTypes = Argument->GetSupportedTypeIndices(TemplateNode->GetFilteredPermutationsIndices());
 					
 					for(int32 PermutationIndex = 0; PermutationIndex < Template->NumPermutations(); PermutationIndex++)
 					{
-						int32 ArgumentTypeIndex = AllArgumentTypes[PermutationIndex];
+						TRigVMTypeIndex ArgumentTypeIndex = AllArgumentTypes[PermutationIndex];
 					
 						bool bIsFilteredOut = false;
 						if(!FilteredArgumentTypes.Contains(ArgumentTypeIndex))
@@ -242,14 +242,13 @@ void SControlRigChangePinType::FillPinTypeMenu(FMenuBuilder& MenuBuilder)
 								if(const FProperty* Property = FunctionStruct->FindPropertyByName(ArgumentName))
 								{
 									Types.Add(ArgumentTypeIndex, FArgumentInfo::Make(Property, bIsFilteredOut));
+									continue;
 								}
 							}
 						}
-						else
-						{
-							const FRigVMTemplateArgumentType ArgumentType = FRigVMRegistry::Get().GetType(ArgumentTypeIndex);
-							Types.Add(ArgumentTypeIndex, FArgumentInfo::Make(ArgumentType, bIsFilteredOut));
-						}
+						
+						const FRigVMTemplateArgumentType ArgumentType = FRigVMRegistry::Get().GetType(ArgumentTypeIndex);
+						Types.Add(ArgumentTypeIndex, FArgumentInfo::Make(ArgumentType, bIsFilteredOut));
 					}
 				}
 			}
@@ -308,7 +307,7 @@ void SControlRigChangePinType::FillPinTypeMenu(FMenuBuilder& MenuBuilder)
 		
 		for(int32 SortedIndex=0; SortedIndex < SortedTypes.Num(); SortedIndex++)
 		{
-			const int32& TypeIndex = SortedTypes[SortedIndex].Key;
+			const TRigVMTypeIndex& TypeIndex = SortedTypes[SortedIndex].Key;
 			const FRigVMTemplateArgumentType Type = FRigVMRegistry::Get().GetType(TypeIndex);
 			const bool bIsFilteredOut = SortedTypes[SortedIndex].Value.bIsFilteredOut;
 			if (Type.CPPTypeObject != nullptr && !IsValid(Type.CPPTypeObject))
@@ -328,7 +327,7 @@ void SControlRigChangePinType::FillPinTypeMenu(FMenuBuilder& MenuBuilder)
 				if(UScriptStruct* ScriptStruct = Cast<UScriptStruct>(Type.CPPTypeObject))
 				{
 					const FRigVMTemplateArgumentType MathType(*ScriptStruct->GetStructCPPName(), ScriptStruct);
-					const int32 MathTypeIndex = FRigVMRegistry::Get().GetTypeIndex(MathType);
+					const TRigVMTypeIndex MathTypeIndex = FRigVMRegistry::Get().GetTypeIndex(MathType);
 					
 					if(!FRigVMRegistry::Get().GetTypesForCategory(FRigVMTemplateArgument::ETypeCategory_SingleMathStructValue).Contains(MathTypeIndex))
 					{
@@ -400,7 +399,7 @@ void SControlRigChangePinType::HandlePinTypeChanged(FRigVMTemplateArgumentType I
 			continue;
 		}
 
-		const int32 TypeIndex = FRigVMRegistry::Get().GetTypeIndex(InType);
+		const TRigVMTypeIndex TypeIndex = FRigVMRegistry::Get().GetTypeIndex(InType);
 		Controller->ResolveWildCardPin(ModelPin->GetPinPath(), TypeIndex, true, true);
 	}
 }
diff --git a/Engine/Source/Developer/RigVMDeveloper/Private/RigVMCompiler/RigVMCompiler.cpp b/Engine/Source/Developer/RigVMDeveloper/Private/RigVMCompiler/RigVMCompiler.cpp
--- a/Engine/Source/Developer/RigVMDeveloper/Private/RigVMCompiler/RigVMCompiler.cpp
+++ b/Engine/Source/Developer/RigVMDeveloper/Private/RigVMCompiler/RigVMCompiler.cpp
@@ -2,6 +2,7 @@
 
 #include "RigVMCompiler/RigVMCompiler.h"
 #include "RigVMModel/RigVMController.h"
+#include "RigVMModel/Nodes/RigVMDispatchNode.h"
 #include "RigVMCore/RigVMExecuteContext.h"
 #include "RigVMCore/RigVMNativized.h"
 #include "RigVMDeveloperModule.h"
@@ -722,9 +723,7 @@ void URigVMCompiler::TraverseEntry(const FRigVMEntryExprAST* InExpr, FRigVMCompi
 	if (WorkData.bSetupMemory)
 	{
 		TSharedPtr<FStructOnScope> DefaultStruct = UnitNode->ConstructStructInstance();
-		WorkData.DefaultStructs.Add(DefaultStruct);
 		TraverseChildren(InExpr, WorkData);
-		WorkData.DefaultStructs.Pop();
 	}
 	else
 	{
@@ -768,32 +767,52 @@ void URigVMCompiler::TraverseEntry(const FRigVMEntryExprAST* InExpr, FRigVMCompi
 
 int32 URigVMCompiler::TraverseCallExtern(const FRigVMCallExternExprAST* InExpr, FRigVMCompilerWorkData& WorkData)
 {
-	URigVMUnitNode* UnitNode = Cast<URigVMUnitNode>(InExpr->GetNode());
-	if(!ValidateNode(UnitNode))
+	URigVMNode* Node = InExpr->GetNode();
+	URigVMUnitNode* UnitNode = Cast<URigVMUnitNode>(Node);
+	URigVMDispatchNode* DispatchNode = Cast<URigVMDispatchNode>(Node);
+	if(!ValidateNode(UnitNode, false) && !ValidateNode(DispatchNode, false))
 	{
 		return INDEX_NONE;
 	}
 
-	if(UnitNode->GetScriptStruct() == nullptr)
+	if(UnitNode)
 	{
-		return INDEX_NONE;
+		if(UnitNode->GetScriptStruct() == nullptr)
+		{
+			return INDEX_NONE;
+		}
+	}
+
+	if(DispatchNode)
+	{
+		if(DispatchNode->GetFactory() == nullptr)
+		{
+			return INDEX_NONE;
+		}
 	}
 
 	int32 InstructionIndex = INDEX_NONE;
 
 	if (WorkData.bSetupMemory)
 	{
-		TSharedPtr<FStructOnScope> DefaultStruct = UnitNode->ConstructStructInstance();
-		WorkData.DefaultStructs.Add(DefaultStruct);
 		TraverseChildren(InExpr, WorkData);
-		WorkData.DefaultStructs.Pop();
 	}
 	else
 	{
 		TArray<FRigVMOperand> Operands;
 
 		// iterate over the child expressions in the order of the arguments on the function
-		const FRigVMFunction* Function = FRigVMRegistry::Get().FindFunction(UnitNode->GetScriptStruct(), *UnitNode->GetMethodName().ToString());
+		const FRigVMFunction* Function = nullptr;
+
+		if(UnitNode)
+		{
+			Function = FRigVMRegistry::Get().FindFunction(UnitNode->GetScriptStruct(), *UnitNode->GetMethodName().ToString());
+		}
+		else if(DispatchNode)
+		{
+			Function = DispatchNode->GetResolvedFunction();
+		}
+			
 		check(Function);
 		
 		for(const FRigVMFunctionArgument& Argument : Function->GetArguments())
@@ -822,14 +841,15 @@ int32 URigVMCompiler::TraverseCallExtern(const FRigVMCallExternExprAST* InExpr,
 		TraverseChildren(InExpr, WorkData);
 
 		// setup the instruction
-		int32 FunctionIndex = WorkData.VM->AddRigVMFunction(UnitNode->GetScriptStruct(), UnitNode->GetMethodName());
+		const int32 FunctionIndex = WorkData.VM->AddRigVMFunction(Function->GetName());
+		check(FunctionIndex != INDEX_NONE);
 		WorkData.VM->GetByteCode().AddExecuteOp(FunctionIndex, Operands);
 		InstructionIndex = WorkData.VM->GetByteCode().GetNumInstructions() - 1;
 
 #if WITH_EDITORONLY_DATA
 		TArray<FRigVMOperand> InputsOperands, OutputOperands;
 
-		for(const URigVMPin* InputPin : UnitNode->GetPins())
+		for(const URigVMPin* InputPin : Node->GetPins())
 		{
 			if(InputPin->IsExecuteContext())
 			{
@@ -2591,13 +2611,26 @@ FRigVMOperand URigVMCompiler::FindOrAddRegister(const FRigVMVarExprAST* InVarExp
 			CPPType = RigVMTypeUtils::ArrayTypeFromBaseType(CPPType);
 			JoinedDefaultValue = URigVMPin::GetDefaultValueForArray({ JoinedDefaultValue });
 		}
-		else if(Pin->GetDirection() == ERigVMPinDirection::Hidden && Pin->GetNode()->IsA<URigVMUnitNode>())
+		else if(Pin->GetDirection() == ERigVMPinDirection::Hidden)
 		{
-			UScriptStruct* UnitStruct = Cast<URigVMUnitNode>(Pin->GetNode())->GetScriptStruct();
-			const FProperty* Property = UnitStruct->FindPropertyByName(Pin->GetFName());
-			check(Property);
+			bool bValidHiddenPin = false;
+			if(Pin->GetNode()->IsA<URigVMUnitNode>())
+			{
+				UScriptStruct* UnitStruct = Cast<URigVMUnitNode>(Pin->GetNode())->GetScriptStruct();
+				const FProperty* Property = UnitStruct->FindPropertyByName(Pin->GetFName());
+				check(Property);
 
-			if (!Property->HasMetaData(FRigVMStruct::SingletonMetaName))
+				if (!Property->HasMetaData(FRigVMStruct::SingletonMetaName))
+				{
+					bValidHiddenPin = true;
+				}
+			}
+			else if(Pin->GetNode()->IsA<URigVMDispatchNode>())
+			{
+				bValidHiddenPin = true;
+			}
+
+			if(bValidHiddenPin)
 			{
 				CPPType = RigVMTypeUtils::ArrayTypeFromBaseType(CPPType);
 				JoinedDefaultValue = URigVMPin::GetDefaultValueForArray({ JoinedDefaultValue });
@@ -2734,18 +2767,23 @@ const FRigVMCompilerWorkData::FRigVMASTProxyArray& URigVMCompiler::FindProxiesWi
 	return WorkData.CachedProxiesWithSharedOperand.Add(InProxy, PinProxies);
 }
 
-bool URigVMCompiler::ValidateNode(URigVMNode* InNode)
+bool URigVMCompiler::ValidateNode(URigVMNode* InNode, bool bCheck)
 {
-	check(InNode);
-
-	if(InNode->HasWildCardPin())
+	if(bCheck)
 	{
-		static const FString UnknownTypeMessage = TEXT("Node @@ has unresolved pins of wildcard type.");
-		Settings.Report(EMessageSeverity::Error, InNode, UnknownTypeMessage);
-		return false;
+		check(InNode)
 	}
-
-	return true;
+	if(InNode)
+	{
+		if(InNode->HasWildCardPin())
+		{
+			static const FString UnknownTypeMessage = TEXT("Node @@ has unresolved pins of wildcard type.");
+			Settings.Report(EMessageSeverity::Error, InNode, UnknownTypeMessage);
+			return false;
+		}
+		return true;
+	}
+	return false;
 }
 
 void URigVMCompiler::ReportInfo(const FString& InMessage)
diff --git a/Engine/Source/Developer/RigVMDeveloper/Private/RigVMModel/Nodes/RigVMDispatchNode.cpp b/Engine/Source/Developer/RigVMDeveloper/Private/RigVMModel/Nodes/RigVMDispatchNode.cpp
new file mode 100644
--- /dev/null
+++ b/Engine/Source/Developer/RigVMDeveloper/Private/RigVMModel/Nodes/RigVMDispatchNode.cpp
@@ -0,0 +1,173 @@
+// Copyright Epic Games, Inc. All Rights Reserved.
+
+#include "RigVMModel/Nodes/RigVMDispatchNode.h"
+#include "RigVMUserWorkflowRegistry.h"
+
+FString URigVMDispatchNode::GetNodeTitle() const
+{
+	if (const FRigVMDispatchFactory* Factory = GetFactory())
+	{
+		return Factory->GetNodeTitle(GetFilteredTypes());
+	}
+	return Super::GetNodeTitle();
+}
+
+FText URigVMDispatchNode::GetToolTipText() const
+{
+	if (const FRigVMDispatchFactory* Factory = GetFactory())
+	{
+		return Factory->GetNodeTooltip(GetFilteredTypes());
+	}
+	return URigVMNode::GetToolTipText();
+}
+
+FLinearColor URigVMDispatchNode::GetNodeColor() const
+{
+	if (const FRigVMDispatchFactory* Factory = GetFactory())
+	{
+		return Factory->GetNodeColor();
+	}
+	return URigVMNode::GetNodeColor();
+}
+
+bool URigVMDispatchNode::IsDefinedAsConstant() const
+{
+	if (const UScriptStruct* Struct = GetFactoryStruct())
+	{
+		return Struct->HasMetaData(FRigVMStruct::ConstantMetaName);
+	}
+	return false;
+}
+
+bool URigVMDispatchNode::IsDefinedAsVarying() const
+{
+	if (const UScriptStruct* Struct = GetFactoryStruct())
+	{
+		return Struct->HasMetaData(FRigVMStruct::VaryingMetaName);
+	}
+	return false;
+}
+
+FText URigVMDispatchNode::GetToolTipTextForPin(const URigVMPin* InPin) const
+{
+	if (const FRigVMDispatchFactory* Factory = GetFactory())
+	{
+		const URigVMPin* RootPin = InPin->GetRootPin();
+		return Factory->GetArgumentTooltip(RootPin->GetFName(), RootPin->GetTypeIndex());
+	}
+	return URigVMNode::GetToolTipTextForPin(InPin);
+}
+
+bool URigVMDispatchNode::IsDeprecated() const
+{
+	return !GetDeprecatedMetadata().IsEmpty();
+}
+
+FString URigVMDispatchNode::GetDeprecatedMetadata() const
+{
+	if (const UScriptStruct* Struct = GetFactoryStruct())
+	{
+		FString DeprecatedMetadata;
+		if(Struct->GetStringMetaDataHierarchical(FRigVMStruct::DeprecatedMetaName, &DeprecatedMetadata))
+		{
+			return DeprecatedMetadata;
+		}
+	}
+	return FString();
+}
+
+TArray<URigVMPin*> URigVMDispatchNode::GetAggregateInputs() const
+{
+	TArray<URigVMPin*> AggregateInputs;
+#if UE_RIGVM_AGGREGATE_NODES_ENABLED
+	if(const FRigVMDispatchFactory* Factory = GetFactory())
+	{
+		const TArray<FName> ArgumentNames = Factory->GetAggregateInputArguments();
+		for(const FName& ArgumentName : ArgumentNames)
+		{
+			AggregateInputs.Add(FindPin(ArgumentName.ToString()));
+		}
+	}
+#endif
+	return AggregateInputs;
+}
+
+TArray<URigVMPin*> URigVMDispatchNode::GetAggregateOutputs() const
+{
+	TArray<URigVMPin*> AggregateOutputs;
+#if UE_RIGVM_AGGREGATE_NODES_ENABLED
+	if(const FRigVMDispatchFactory* Factory = GetFactory())
+	{
+		const TArray<FName> ArgumentNames = Factory->GetAggregateOutputArguments();
+		for(const FName& ArgumentName : ArgumentNames)
+		{
+			AggregateOutputs.Add(FindPin(ArgumentName.ToString()));
+		}
+	}
+#endif
+	return AggregateOutputs;
+}
+
+FName URigVMDispatchNode::GetNextAggregateName(const FName& InLastAggregatePinName) const
+{
+#if UE_RIGVM_AGGREGATE_NODES_ENABLED
+	if(const FRigVMDispatchFactory* Factory = GetFactory())
+	{
+		return Factory->GetNextAggregateName(InLastAggregatePinName);
+	}
+#endif
+	return FName();
+}
+
+const FRigVMDispatchFactory* URigVMDispatchNode::GetFactory() const
+{
+	if(CachedFactory)
+	{
+		return CachedFactory;
+	}
+	
+	if(const FRigVMTemplate* Template = GetTemplate())
+	{
+		CachedFactory = Template->GetDispatchFactory(); 
+		return CachedFactory;
+	}
+	return nullptr;
+}
+
+const UScriptStruct* URigVMDispatchNode::GetFactoryStruct() const
+{
+	if(const FRigVMDispatchFactory* Factory = GetFactory())
+	{
+		return Factory->GetScriptStruct();
+	}
+	return nullptr;
+}
+
+const FRigVMTemplateTypeMap& URigVMDispatchNode::GetFilteredTypes() const
+{
+	const TArray<int32> PermutationIndices = GetFilteredPermutationsIndices();
+	if(PermutationIndices.Num() == 1)
+	{
+		TypesFromPins = GetTypesForPermutation(PermutationIndices[0]);
+	}
+	else
+	{
+		TypesFromPins.Reset();
+	}
+	return TypesFromPins;
+}
+
+void URigVMDispatchNode::InvalidateCache()
+{
+	Super::InvalidateCache();
+	CachedFactory = nullptr;
+}
+
+FRigVMStructUpgradeInfo URigVMDispatchNode::GetUpgradeInfo() const
+{
+	if(const FRigVMDispatchFactory* Factory = GetFactory())
+	{
+		return Factory->GetUpgradeInfo(GetFilteredTypes());
+	}
+	return FRigVMStructUpgradeInfo();
+}
diff --git a/Engine/Source/Developer/RigVMDeveloper/Private/RigVMModel/Nodes/RigVMTemplateNode.cpp b/Engine/Source/Developer/RigVMDeveloper/Private/RigVMModel/Nodes/RigVMTemplateNode.cpp
--- a/Engine/Source/Developer/RigVMDeveloper/Private/RigVMModel/Nodes/RigVMTemplateNode.cpp
+++ b/Engine/Source/Developer/RigVMDeveloper/Private/RigVMModel/Nodes/RigVMTemplateNode.cpp
@@ -100,10 +100,28 @@ void URigVMTemplateNode::ConvertPreferredTypesToString()
 
 void URigVMTemplateNode::ConvertPreferredTypesToTypeIndex()
 {
+	const FRigVMTemplate* Template = GetTemplate();
+	
 	// we rely on the strings being serialized - so on load we need to update the type index again 
 	for(FRigVMTemplatePreferredType& PreferredType : PreferredPermutationPairs)
 	{
 		PreferredType.UpdateIndexFromString();
+
+		if(Template)
+		{
+			if(const FRigVMTemplateArgument* Argument = Template->FindArgument(PreferredType.GetArgument()))
+			{
+				TRigVMTypeIndex ResolvedTypeIndex = INDEX_NONE;
+				if(Argument->SupportsTypeIndex(PreferredType.GetTypeIndex(), &ResolvedTypeIndex))
+				{
+					if(ResolvedTypeIndex != PreferredType.GetTypeIndex())
+					{
+						PreferredType.TypeIndex = ResolvedTypeIndex;
+						PreferredType.UpdateStringFromIndex();
+					}
+				}
+			}
+		}
 	}
 }
 
@@ -182,7 +200,7 @@ FText URigVMTemplateNode::GetToolTipTextForPin(const URigVMPin* InPin) const
 					FString SupportedTypesJoined;
 					for (int32 Index : PermutationIndices)
 					{
-						const int32 TypeIndex = Arg->GetTypeIndices()[Index];
+						const TRigVMTypeIndex TypeIndex = Arg->GetTypeIndices()[Index];
 						FString Type = FRigVMRegistry::Get().GetType(TypeIndex).CPPType.ToString();
 						if (!FilteredPermutations.Contains(Index))
 						{
@@ -215,7 +233,7 @@ bool URigVMTemplateNode::IsSingleton() const
 	return GetTemplate() == nullptr;
 }
 
-bool URigVMTemplateNode::SupportsType(const URigVMPin* InPin, int32 InTypeIndex, int32* OutTypeIndex)
+bool URigVMTemplateNode::SupportsType(const URigVMPin* InPin, TRigVMTypeIndex InTypeIndex, TRigVMTypeIndex* OutTypeIndex)
 {
 	const FRigVMRegistry& Registry = FRigVMRegistry::Get();
 	static const FString WildCardCPPType = RigVMTypeUtils::GetWildCardCPPType();
@@ -223,7 +241,7 @@ bool URigVMTemplateNode::SupportsType(const URigVMPin* InPin, int32 InTypeIndex,
 
 	const URigVMPin* RootPin = InPin->GetRootPin();
 
-	int32 TypeIndex = InTypeIndex;
+	TRigVMTypeIndex TypeIndex = InTypeIndex;
 
 	if(InPin->GetParentPin() == RootPin && RootPin->IsArray())
 	{
@@ -270,8 +288,8 @@ bool URigVMTemplateNode::SupportsType(const URigVMPin* InPin, int32 InTypeIndex,
 	
 	if (const FRigVMTemplate* Template = GetTemplate())
 	{
-		const TPair<FName,int32> CacheKey(RootPin->GetFName(), TypeIndex);
-		if (const TPair<bool, int32>* CachedResult = SupportedTypesCache.Find(CacheKey))
+		const TPair<FName,TRigVMTypeIndex> CacheKey(RootPin->GetFName(), TypeIndex);
+		if (const TPair<bool, TRigVMTypeIndex>* CachedResult = SupportedTypesCache.Find(CacheKey))
 		{
 			if(OutTypeIndex)
 			{
@@ -280,7 +298,7 @@ bool URigVMTemplateNode::SupportsType(const URigVMPin* InPin, int32 InTypeIndex,
 			return CachedResult->Key;
 		}
 
-		int32 SupportedTypeIndex = INDEX_NONE;
+		TRigVMTypeIndex SupportedTypeIndex = INDEX_NONE;
 		if (Template->ArgumentSupportsTypeIndex(RootPin->GetFName(), TypeIndex, &SupportedTypeIndex))
 		{
 			if (OutTypeIndex)
@@ -303,7 +321,7 @@ bool URigVMTemplateNode::SupportsType(const URigVMPin* InPin, int32 InTypeIndex,
 	return false;
 }
 
-bool URigVMTemplateNode::FilteredSupportsType(const URigVMPin* InPin, int32 InTypeIndex, int32* OutTypeIndex, bool bAllowFloatingPointCasts)
+bool URigVMTemplateNode::FilteredSupportsType(const URigVMPin* InPin, TRigVMTypeIndex InTypeIndex, TRigVMTypeIndex* OutTypeIndex, bool bAllowFloatingPointCasts)
 {
 	if (OutTypeIndex)
 	{
@@ -347,7 +365,7 @@ bool URigVMTemplateNode::FilteredSupportsType(const URigVMPin* InPin, int32 InTy
 
 	const FRigVMRegistry& Registry = FRigVMRegistry::Get();
 
-	int32 RootTypeIndex = InTypeIndex;
+	TRigVMTypeIndex RootTypeIndex = InTypeIndex;
 	if (bIsArrayElement)
 	{
 		RootTypeIndex = Registry.GetArrayTypeFromBaseTypeIndex(RootTypeIndex);
@@ -358,10 +376,10 @@ bool URigVMTemplateNode::FilteredSupportsType(const URigVMPin* InPin, int32 InTy
 		return GetTemplate()->ArgumentSupportsTypeIndex(RootPin->GetFName(), InTypeIndex, OutTypeIndex);
 	}
 
-	const TArray<int32>& TypeIndices = Argument->GetTypeIndices();
+	const TArray<TRigVMTypeIndex>& TypeIndices = Argument->GetTypeIndices();
 	for (const int32& PermutationIndex : FilteredPermutations)
 	{
-		const int32& FilteredTypeIndex = TypeIndices[PermutationIndex];
+		const TRigVMTypeIndex& FilteredTypeIndex = TypeIndices[PermutationIndex];
 		if (Registry.CanMatchTypes(FilteredTypeIndex, RootTypeIndex, bAllowFloatingPointCasts))
 		{
 			return true;
@@ -439,16 +457,21 @@ FString URigVMTemplateNode::GetInitialDefaultValueForPin(const FName& InRootPinN
 	}
 
 	FString DefaultValue;
+	const FRigVMTemplate* Template = GetTemplate();
+	const FRigVMDispatchFactory* Factory = Template->GetDispatchFactory(); 
 
 	for(const int32 PermutationIndex : PermutationIndices)
 	{
 		FString NewDefaultValue;
 
-		const FRigVMTemplate* Template = GetTemplate();
 		const FRigVMTemplateArgument* Argument = Template->FindArgument(InRootPinName);
-		const int32 TypeIndex = Argument->GetTypeIndices()[PermutationIndex];
+		const TRigVMTypeIndex TypeIndex = Argument->GetTypeIndices()[PermutationIndex];
 
-		if(const FRigVMFunction* Permutation = Template->GetPermutation(PermutationIndex))
+		if(Factory)
+		{
+			NewDefaultValue = Factory->GetArgumentDefaultValue(Argument->GetName(), TypeIndex);
+		}
+		else if(const FRigVMFunction* Permutation = Template->GetPermutation(PermutationIndex))
 		{
 			const TSharedPtr<FStructOnScope> StructOnScope = MakeShareable(new FStructOnScope(Permutation->Struct));
 			const FRigVMStruct* DefaultStruct = (const FRigVMStruct*)StructOnScope->GetStructMemory();
@@ -563,10 +586,10 @@ const TArray<int32>& URigVMTemplateNode::GetFilteredPermutationsIndices() const
 	return FilteredPermutations;
 }
 
-TArray<int32> URigVMTemplateNode::GetFilteredTypesForPin(URigVMPin* InPin) const
+TArray<TRigVMTypeIndex> URigVMTemplateNode::GetFilteredTypesForPin(URigVMPin* InPin) const
 {
 	ensureMsgf(InPin->GetNode() == this, TEXT("GetFilteredTypesForPin of %s with pin from another node %s"), *GetNodePath(), *InPin->GetPinPath(true));
-	TArray<int32> FilteredTypes;
+	TArray<TRigVMTypeIndex> FilteredTypes;
 
 	if (FilteredPermutations.IsEmpty())
 	{
@@ -577,12 +600,12 @@ TArray<int32> URigVMTemplateNode::GetFilteredTypesForPin(URigVMPin* InPin) const
 	{
 		for (const FRigVMTemplatePreferredType& PreferredType : PreferredPermutationPairs)
 		{
-			if (InPin->GetFName() == PreferredType.Argument)
+			if (InPin->GetFName() == PreferredType.GetArgument())
 			{
-				const FRigVMTemplateArgument* Argument = GetTemplate()->FindArgument(PreferredType.Argument);
-				for (const int32& TypeIndex : Argument->GetTypeIndices())
+				const FRigVMTemplateArgument* Argument = GetTemplate()->FindArgument(PreferredType.GetArgument());
+				for (const TRigVMTypeIndex& TypeIndex : Argument->GetTypeIndices())
 				{
-					if (TypeIndex == PreferredType.TypeIndex)
+					if (TypeIndex == PreferredType.GetTypeIndex())
 					{
 						return {TypeIndex};						
 					}
@@ -603,7 +626,7 @@ TArray<int32> URigVMTemplateNode::GetFilteredTypesForPin(URigVMPin* InPin) const
 		FilteredTypes = Argument->GetSupportedTypeIndices(FilteredPermutations);
 		if (bIsArrayElement)
 		{
-			for (int32& ArrayType : FilteredTypes)
+			for (TRigVMTypeIndex& ArrayType : FilteredTypes)
 			{
 				ArrayType = FRigVMRegistry::Get().GetBaseTypeFromArrayTypeIndex(ArrayType);
 			}
@@ -656,7 +679,7 @@ TArray<int32> URigVMTemplateNode::GetNewFilteredPermutations(URigVMPin* InPin, U
 			{
 				for(int32 PermutationIndex : PermutationsToTry)
 				{
-					int32 TypeIndex = Argument->GetTypeIndices()[PermutationIndex];
+					TRigVMTypeIndex TypeIndex = Argument->GetTypeIndices()[PermutationIndex];
 					if (bIsArrayElement)
 					{
 						TypeIndex = FRigVMRegistry::Get().GetBaseTypeFromArrayTypeIndex(TypeIndex);
@@ -674,10 +697,10 @@ TArray<int32> URigVMTemplateNode::GetNewFilteredPermutations(URigVMPin* InPin, U
 	{
 		if (const FRigVMTemplateArgument* Argument = GetTemplate()->FindArgument(RootPin->GetFName()))
 		{
-			const int32 LinkedTypeIndex = LinkedPin->GetTypeIndex();
+			const TRigVMTypeIndex LinkedTypeIndex = LinkedPin->GetTypeIndex();
 			for(int32 PermutationIndex : PermutationsToTry)
 			{
-				int32 TypeIndex = Argument->GetTypeIndices()[PermutationIndex];
+				TRigVMTypeIndex TypeIndex = Argument->GetTypeIndices()[PermutationIndex];
 				if (bIsArrayElement)
 				{
 					TypeIndex = FRigVMRegistry::Get().GetBaseTypeFromArrayTypeIndex(TypeIndex);
@@ -693,17 +716,17 @@ TArray<int32> URigVMTemplateNode::GetNewFilteredPermutations(URigVMPin* InPin, U
 	return NewFilteredPermutations;
 }
 
-TArray<int32> URigVMTemplateNode::GetNewFilteredPermutations(URigVMPin* InPin, const TArray<int32>& InTypeIndices)
+TArray<int32> URigVMTemplateNode::GetNewFilteredPermutations(URigVMPin* InPin, const TArray<TRigVMTypeIndex>& InTypeIndices)
 {
 	TArray<int32> NewFilteredPermutations;
 	NewFilteredPermutations.Reserve(FilteredPermutations.Num());
 
 	URigVMPin* RootPin = InPin;
-	TArray<int32> RootTypes = InTypeIndices;
+	TArray<TRigVMTypeIndex> RootTypes = InTypeIndices;
 	if (URigVMPin* ParentPin = InPin->GetParentPin())
 	{
 		RootPin = ParentPin;
-		for (int32& TypeIndex : RootTypes)
+		for (TRigVMTypeIndex& TypeIndex : RootTypes)
 		{
 			TypeIndex = FRigVMRegistry::Get().GetArrayTypeFromBaseTypeIndex(TypeIndex);
 		}
@@ -717,10 +740,10 @@ TArray<int32> URigVMTemplateNode::GetNewFilteredPermutations(URigVMPin* InPin, c
 	
 	if (const FRigVMTemplateArgument* Argument = GetTemplate()->FindArgument(RootPin->GetFName()))
 	{
-		const TArray<int32>& TypeIndices = Argument->GetTypeIndices();
+		const TArray<TRigVMTypeIndex>& TypeIndices = Argument->GetTypeIndices();
 		for(int32 PermutationIndex : PermutationsToTry)
 		{
-			for (const int32& RootTypeIndex : RootTypes)
+			for (const TRigVMTypeIndex& RootTypeIndex : RootTypes)
 			{				
 				if (FRigVMRegistry::Get().CanMatchTypes(TypeIndices[PermutationIndex], RootTypeIndex, true))
 				{
@@ -741,13 +764,13 @@ TArray<int32> URigVMTemplateNode::FindPermutationsForTypes(const TArray<FRigVMTe
 		const FRigVMRegistry& Registry = FRigVMRegistry::Get();
 		
 		TArray<const FRigVMTemplateArgument*> Args;
-		TArray<int32> TypeIndices;
+		TArray<TRigVMTypeIndex> TypeIndices;
 		for (const FRigVMTemplatePreferredType& ArgumentType : ArgumentTypes)
 		{
-			if (const FRigVMTemplateArgument* Argument = Template->FindArgument(ArgumentType.Argument))
+			if (const FRigVMTemplateArgument* Argument = Template->FindArgument(ArgumentType.GetArgument()))
 			{
 				Args.Add(Argument);
-				TypeIndices.Add(ArgumentType.TypeIndex);
+				TypeIndices.Add(ArgumentType.GetTypeIndex());
 			}
 			else
 			{
@@ -780,7 +803,7 @@ TArray<int32> URigVMTemplateNode::FindPermutationsForTypes(const TArray<FRigVMTe
 	return Permutations;
 }
 
-TArray<FRigVMTemplatePreferredType> URigVMTemplateNode::GetArgumentTypesForPermutation(const int32 InPermutationIndex) const
+TArray<FRigVMTemplatePreferredType> URigVMTemplateNode::GetPreferredTypesForPermutation(const int32 InPermutationIndex) const
 {
 	TArray<FRigVMTemplatePreferredType> ArgTypes;
 	if (const FRigVMTemplate* Template = GetTemplate())
@@ -803,7 +826,17 @@ TArray<FRigVMTemplatePreferredType> URigVMTemplateNode::GetArgumentTypesForPermu
 	return ArgTypes;
 }
 
-bool URigVMTemplateNode::PinNeedsFilteredTypesUpdate(URigVMPin* InPin, const TArray<int32>& InTypeIndices)
+FRigVMTemplateTypeMap URigVMTemplateNode::GetTypesForPermutation(const int32 InPermutationIndex) const
+{
+	if (const FRigVMTemplate* Template = GetTemplate())
+	{
+		return Template->GetTypesForPermutation(InPermutationIndex);
+	}
+
+	return FRigVMTemplateTypeMap();
+}
+
+bool URigVMTemplateNode::PinNeedsFilteredTypesUpdate(URigVMPin* InPin, const TArray<TRigVMTypeIndex>& InTypeIndices)
 {
 	const TArray<int32> NewFilteredPermutations = GetNewFilteredPermutations(InPin, InTypeIndices);
 	if (NewFilteredPermutations.Num() == FilteredPermutations.Num())
@@ -844,7 +877,7 @@ bool URigVMTemplateNode::UpdateFilteredPermutations(URigVMPin* InPin, URigVMPin*
 	return true;
 }
 
-bool URigVMTemplateNode::UpdateFilteredPermutations(URigVMPin* InPin, const TArray<int32>& InTypeIndices)
+bool URigVMTemplateNode::UpdateFilteredPermutations(URigVMPin* InPin, const TArray<TRigVMTypeIndex>& InTypeIndices)
 {
 	const TArray<int32> NewFilteredPermutations = GetNewFilteredPermutations(InPin, InTypeIndices);
 	if (NewFilteredPermutations.Num() == FilteredPermutations.Num())
diff --git a/Engine/Source/Developer/RigVMDeveloper/Private/RigVMModel/Nodes/RigVMUnitNode.cpp b/Engine/Source/Developer/RigVMDeveloper/Private/RigVMModel/Nodes/RigVMUnitNode.cpp
--- a/Engine/Source/Developer/RigVMDeveloper/Private/RigVMModel/Nodes/RigVMUnitNode.cpp
+++ b/Engine/Source/Developer/RigVMDeveloper/Private/RigVMModel/Nodes/RigVMUnitNode.cpp
@@ -165,89 +165,6 @@ TArray<FRigVMUserWorkflow> URigVMUnitNode::GetSupportedWorkflows(ERigVMUserWorkf
 	return Workflows;
 }
 
-bool URigVMUnitNode::IsAggregate() const
-{
-#if UE_RIGVM_AGGREGATE_NODES_ENABLED
-	const TArray<URigVMPin*> AggregateInputs = GetAggregateInputs();
-	const TArray<URigVMPin*> AggregateOutputs = GetAggregateOutputs();
-
-	if ((AggregateInputs.Num() == 2 && AggregateOutputs.Num() == 1) ||
-		(AggregateInputs.Num() == 1 && AggregateOutputs.Num() == 2))
-	{
-		TArray<URigVMPin*> AggregateAll = AggregateInputs;
-		AggregateAll.Append(AggregateOutputs);
-		for (int32 i = 1; i < 3; ++i)
-		{
-			if (AggregateAll[0]->GetCPPType() != AggregateAll[i]->GetCPPType() ||
-				AggregateAll[0]->GetCPPTypeObject() != AggregateAll[i]->GetCPPTypeObject())
-			{
-				return false;
-			}
-		}
-		
-		return true;
-	}
-#endif
-	
-	return false;
-}
-
-URigVMPin* URigVMUnitNode::GetFirstAggregatePin() const
-{
-#if UE_RIGVM_AGGREGATE_NODES_ENABLED
-	const TArray<URigVMPin*> Inputs = GetAggregateInputs();
-	const TArray<URigVMPin*> Outputs = GetAggregateOutputs();
-	if (Inputs.Num() == 2 && Outputs.Num() == 1)
-	{
-		return Inputs[0];
-	}
-	if (Inputs.Num() == 1 && Outputs.Num() == 2)
-	{
-		return Outputs[0];
-	}
-#endif
-	return nullptr;
-}
-
-URigVMPin* URigVMUnitNode::GetSecondAggregatePin() const
-{
-#if UE_RIGVM_AGGREGATE_NODES_ENABLED
-	const TArray<URigVMPin*> Inputs = GetAggregateInputs();
-	const TArray<URigVMPin*> Outputs = GetAggregateOutputs();
-	if (Inputs.Num() == 2 && Outputs.Num() == 1)
-	{
-		return Inputs[1];
-	}
-	if (Inputs.Num() == 1 && Outputs.Num() == 2)
-	{
-		return Outputs[1];
-	}
-#endif
-	return nullptr;
-}
-
-URigVMPin* URigVMUnitNode::GetOppositeAggregatePin() const
-{
-#if UE_RIGVM_AGGREGATE_NODES_ENABLED
-	const TArray<URigVMPin*> Inputs = GetAggregateInputs();
-	const TArray<URigVMPin*> Outputs = GetAggregateOutputs();
-	if (Inputs.Num() == 2 && Outputs.Num() == 1)
-	{
-		return Outputs[0];
-	}
-	if (Inputs.Num() == 1 && Outputs.Num() == 2)
-	{
-		return Inputs[0];
-	}
-#endif
-	return nullptr;
-}
-
-bool URigVMUnitNode::IsInputAggregate() const
-{
-	return GetAggregateInputs().Num() == 2;
-}
-
 TArray<URigVMPin*> URigVMUnitNode::GetAggregateInputs() const
 {
 	TArray<URigVMPin*> AggregateInputs;
diff --git a/Engine/Source/Developer/RigVMDeveloper/Private/RigVMModel/RigVMController.cpp b/Engine/Source/Developer/RigVMDeveloper/Private/RigVMModel/RigVMController.cpp
--- a/Engine/Source/Developer/RigVMDeveloper/Private/RigVMModel/RigVMController.cpp
+++ b/Engine/Source/Developer/RigVMDeveloper/Private/RigVMModel/RigVMController.cpp
@@ -6,6 +6,7 @@
 #include "RigVMModel/Nodes/RigVMFunctionReturnNode.h"
 #include "RigVMModel/Nodes/RigVMFunctionReferenceNode.h"
 #include "RigVMModel/Nodes/RigVMAggregateNode.h"
+#include "RigVMModel/Nodes/RigVMDispatchNode.h"
 #include "RigVMCore/RigVMRegistry.h"
 #include "RigVMCore/RigVMExecuteContext.h"
 #include "RigVMCore/RigVMUnknownType.h"
@@ -407,6 +408,43 @@ TArray<FString> URigVMController::GetAddNodePythonCommands(URigVMNode* Node) con
 	const FString GraphName = GetSanitizedGraphName(GetGraph()->GetGraphName());
 	const FString NodeName = GetSanitizedNodeName(Node->GetName());
 
+	auto GetResolveWildcardPinsPythonCommands = [](const FString& InGraphName, const URigVMTemplateNode* InNode, const FRigVMTemplate* InTemplate)
+	{
+		TArray<FString> Commands;
+
+		// Lets minimize the number of commands by stopping when the number of permutations left is 1 (or less)
+		TArray<int32> Permutations;
+		Permutations.SetNumUninitialized(InTemplate->NumPermutations());
+		FRigVMTemplate::FTypeMap TypeMap;
+		
+		for (int32 ArgIndex = 0; ArgIndex < InTemplate->NumArguments(); ++ArgIndex)
+		{
+			if (Permutations.Num() < 2)
+			{
+				break;
+			}
+			
+			const FRigVMTemplateArgument* Argument = InTemplate->GetArgument(ArgIndex);
+			if (!Argument->IsSingleton())
+			{
+				URigVMPin* Pin = InNode->FindPin(Argument->GetName().ToString());
+				if (!Pin->IsWildCard())
+				{
+					Commands.Add(FString::Printf(TEXT("blueprint.get_controller_by_name('%s').resolve_wild_card_pin('%s', '%s', '%s')"),
+							*InGraphName,
+							*Pin->GetPinPath(),
+							*Pin->GetCPPType(),
+							*Pin->GetCPPTypeObject()->GetPathName()));
+
+					TypeMap.Add(Argument->GetName(), Pin->GetTypeIndex());
+					InTemplate->Resolve(TypeMap, Permutations, false);
+				}
+			}
+		}
+
+		return Commands;
+	};
+
 	if (const URigVMUnitNode* UnitNode = Cast<URigVMUnitNode>(Node))
 	{
 		if (const URigVMInjectionInfo* InjectionInfo = Cast<URigVMInjectionInfo>(UnitNode->GetOuter()))
@@ -449,38 +487,25 @@ TArray<FString> URigVMController::GetAddNodePythonCommands(URigVMNode* Node) con
 			// Try to resolve wildcard pins			
 			if (const FRigVMTemplate* Template = UnitNode->GetTemplate())
 			{
-				// Lets minimize the number of commands by stopping when the number of permutations left is 1 (or less)
-				TArray<int32> Permutations;
-				Permutations.SetNumUninitialized(Template->NumPermutations());
-				FRigVMTemplate::FTypeMap TypeMap;
-		
-				for (int32 ArgIndex = 0; ArgIndex < Template->NumArguments(); ++ArgIndex)
-				{
-					if (Permutations.Num() < 2)
-					{
-						break;
-					}
-			
-					const FRigVMTemplateArgument* Argument = Template->GetArgument(ArgIndex);
-					if (!Argument->IsSingleton())
-					{
-						URigVMPin* Pin = UnitNode->FindPin(Argument->GetName().ToString());
-						if (!Pin->IsWildCard())
-						{
-							Commands.Add(FString::Printf(TEXT("blueprint.get_controller_by_name('%s').resolve_wild_card_pin('%s', '%s', '%s')"),
-									*GraphName,
-									*Pin->GetPinPath(),
-									*Pin->GetCPPType(),
-									*Pin->GetCPPTypeObject()->GetPathName()));
-
-							TypeMap.Add(Argument->GetName(), Pin->GetTypeIndex());
-							Template->Resolve(TypeMap, Permutations, false);
-						}
-					}
-				}
+				Commands.Append(GetResolveWildcardPinsPythonCommands(GraphName, UnitNode, Template));
 			}			
 		}		
 	}
+	else if (const URigVMDispatchNode* DispatchNode = Cast<URigVMDispatchNode>(Node))
+	{
+		// add_template_node(notation, position=[0.0, 0.0], node_name='', undo=True)
+		Commands.Add(FString::Printf(TEXT("blueprint.get_controller_by_name('%s').add_template_node('%s', %s, '%s')"),
+						*GraphName,
+						*DispatchNode->GetNotation().ToString(),
+						*RigVMPythonUtils::Vector2DToPythonString(DispatchNode->GetPosition()),
+						*NodeName));
+
+		// Try to resolve wildcard pins			
+		if (const FRigVMTemplate* Template = DispatchNode->GetTemplate())
+		{
+			Commands.Append(GetResolveWildcardPinsPythonCommands(GraphName, DispatchNode, Template));
+		}			
+	}
 	else if (const URigVMAggregateNode* AggregateNode = Cast<URigVMAggregateNode>(Node))
 	{
 		TArray<FString> InnerNodeCommands = GetAddNodePythonCommands(AggregateNode->GetFirstInnerNode());
@@ -798,7 +823,7 @@ URigVMUnitNode* URigVMController::AddUnitNode(UScriptStruct* InScriptStruct, con
 		TArray<int32> OldPermutations = TemplateNode->FilteredPermutations;
 		int32 PermutationIndex = Template->FindPermutation(Function);
 		TemplateNode->FilteredPermutations = {PermutationIndex};
-		const TArray<FRigVMTemplatePreferredType> NewPreferredPermutationTypes = TemplateNode->GetArgumentTypesForPermutation(PermutationIndex);
+		const TArray<FRigVMTemplatePreferredType> NewPreferredPermutationTypes = TemplateNode->GetPreferredTypesForPermutation(PermutationIndex);
 		if (bSetupUndoRedo)
 		{
 			FRigVMSetTemplateFilteredPermutationsAction Action(TemplateNode, nullptr, OldPermutations);
@@ -1824,7 +1849,7 @@ bool URigVMController::UnresolveTemplateNodes(const TArray<URigVMTemplateNode*>&
 	const FRigVMRegistry& Registry = FRigVMRegistry::Get();
 	
 	// Find pins outside our set of nodes which were resolved to a type
-	TMap<URigVMPin*, int32> PinsToResolve;
+	TMap<URigVMPin*, TRigVMTypeIndex> PinsToResolve;
 	for (TPair<FString, FString>& LinkPath : LinkedPaths)
 	{
 		const FString& SourcePath = LinkPath.Key;
@@ -1923,7 +1948,7 @@ bool URigVMController::UnresolveTemplateNodes(const TArray<URigVMTemplateNode*>&
 		// if the types' containers don't match the current pin containers try again
 		bool bTypesWereAdjusted = false;
 		FRigVMTemplate::FTypeMap AdjustedTypes = Types;
-		for(TPair<FName, int32>& TypePair : AdjustedTypes)
+		for(TPair<FName, TRigVMTypeIndex>& TypePair : AdjustedTypes)
 		{
 			if(URigVMPin* Pin = Node->FindPin(TypePair.Key.ToString()))
 			{
@@ -1951,10 +1976,10 @@ bool URigVMController::UnresolveTemplateNodes(const TArray<URigVMTemplateNode*>&
 			}
 		}
 
-		for(const TPair<FName, int32>& TypePair : Types)
+		for(const TPair<FName, TRigVMTypeIndex>& TypePair : Types)
 		{
 			const FName& PinName = TypePair.Key;
-			const int32& ExpectedTypeIndex = TypePair.Value;
+			const TRigVMTypeIndex& ExpectedTypeIndex = TypePair.Value;
 			
 			URigVMPin* Pin = Node->FindPin(PinName.ToString());
 			if(Pin == nullptr)
@@ -1973,7 +1998,7 @@ bool URigVMController::UnresolveTemplateNodes(const TArray<URigVMTemplateNode*>&
 
 				if(Pin->IsArray())
 				{
-					const int32 ElementTypeIndex = FRigVMRegistry::Get().GetBaseTypeFromArrayTypeIndex(ExpectedTypeIndex);
+					const TRigVMTypeIndex ElementTypeIndex = FRigVMRegistry::Get().GetBaseTypeFromArrayTypeIndex(ExpectedTypeIndex);
 					
 					TArray<URigVMPin*> SubPins = Pin->GetSubPins();
 					for (URigVMPin* SubPin : SubPins)
@@ -1993,7 +2018,7 @@ bool URigVMController::UnresolveTemplateNodes(const TArray<URigVMTemplateNode*>&
 			return A->IsWildCard() == B->IsWildCard();
 		}), bSetupUndoRedo);
 	
-	for (TPair<URigVMPin*, int32>& Pin : PinsToResolve)
+	for (TPair<URigVMPin*, TRigVMTypeIndex>& Pin : PinsToResolve)
 	{
 		TArray<URigVMLink*> LinksToBreak;
 		for (URigVMLink* Link : Pin.Key->GetLinks())
@@ -3717,7 +3742,7 @@ TArray<FName> URigVMController::ImportNodesFromText(const FString& InText, bool
 							{
 								if (!FirstTemplateNode->IsSingleton())
 								{
-									TArray<int32> InputTypes = GetWildcardFilteredTypeIndices(SourcePin);
+									TArray<TRigVMTypeIndex> InputTypes = GetWildcardFilteredTypeIndices(SourcePin);
 									if (InputTypes.Num() > 0)
 									{
 										PrepareTemplatePinForType(TargetPin, InputTypes, bSetupUndoRedo);										
@@ -3728,7 +3753,7 @@ TArray<FName> URigVMController::ImportNodesFromText(const FString& InText, bool
 							{
 								if (!SecondTemplateNode->IsSingleton())
 								{
-									TArray<int32> OutTypes = GetWildcardFilteredTypeIndices(TargetPin);
+									TArray<TRigVMTypeIndex> OutTypes = GetWildcardFilteredTypeIndices(TargetPin);
 									if (OutTypes.Num() > 0)
 									{
 										PrepareTemplatePinForType(SourcePin, OutTypes, bSetupUndoRedo);										
@@ -9271,7 +9296,7 @@ bool URigVMController::AddLink(URigVMPin* OutputPin, URigVMPin* InputPin, bool b
 		{
 			if (!FirstTemplateNode->IsSingleton())
 			{
-				TArray<int32> InputTypes = GetWildcardFilteredTypeIndices(SecondToResolve);
+				TArray<TRigVMTypeIndex> InputTypes = GetWildcardFilteredTypeIndices(SecondToResolve);
 				if (InputTypes.Num() > 0)
 				{
 					if (!PrepareTemplatePinForType(FirstToResolve, InputTypes, bSetupUndoRedo))
@@ -9286,7 +9311,7 @@ bool URigVMController::AddLink(URigVMPin* OutputPin, URigVMPin* InputPin, bool b
 		{
 			if (!SecondTemplateNode->IsSingleton())
 			{
-				TArray<int32> OutTypes = GetWildcardFilteredTypeIndices(FirstToResolve);
+				TArray<TRigVMTypeIndex> OutTypes = GetWildcardFilteredTypeIndices(FirstToResolve);
 				if (OutTypes.Num() > 0)
 				{
 					if (!PrepareTemplatePinForType(SecondToResolve, OutTypes, bSetupUndoRedo))
@@ -11664,7 +11689,7 @@ URigVMRerouteNode* URigVMController::AddFreeRerouteNode(bool bShowAsFullNode, co
 		TArray<int32> FilterPermutations = Node->GetFilteredPermutationsIndices();
 		if (FilterPermutations.Num() == 1)
 		{
-			const TArray<FRigVMTemplatePreferredType> NewPreferredPermutationTypes = Node->GetArgumentTypesForPermutation(FilterPermutations[0]);
+			const TArray<FRigVMTemplatePreferredType> NewPreferredPermutationTypes = Node->GetPreferredTypesForPermutation(FilterPermutations[0]);
 			if (bSetupUndoRedo)
 			{
 				ActionStack->AddAction(FRigVMSetPreferredTemplatePermutationsAction(Node, NewPreferredPermutationTypes));
@@ -11856,7 +11881,7 @@ URigVMIfNode* URigVMController::AddIfNode(const FString& InCPPType, const FName&
 		TArray<int32> FilterPermutations = Node->GetFilteredPermutationsIndices();
 		if (FilterPermutations.Num() == 1)
 		{
-			const TArray<FRigVMTemplatePreferredType> NewPreferredPermutationTypes = Node->GetArgumentTypesForPermutation(FilterPermutations[0]);
+			const TArray<FRigVMTemplatePreferredType> NewPreferredPermutationTypes = Node->GetPreferredTypesForPermutation(FilterPermutations[0]);
 			if (bSetupUndoRedo)
 			{
 				ActionStack->AddAction(FRigVMSetPreferredTemplatePermutationsAction(Node, NewPreferredPermutationTypes));
@@ -11977,7 +12002,7 @@ URigVMSelectNode* URigVMController::AddSelectNode(const FString& InCPPType, cons
 		TArray<int32> FilterPermutations = Node->GetFilteredPermutationsIndices();
 		if (FilterPermutations.Num() == 1)
 		{
-			const TArray<FRigVMTemplatePreferredType> NewPreferredPermutationTypes = Node->GetArgumentTypesForPermutation(FilterPermutations[0]);
+			const TArray<FRigVMTemplatePreferredType> NewPreferredPermutationTypes = Node->GetPreferredTypesForPermutation(FilterPermutations[0]);
 			if (bSetupUndoRedo)
 			{
 				ActionStack->AddAction(FRigVMSetPreferredTemplatePermutationsAction(Node, NewPreferredPermutationTypes));
@@ -12038,7 +12063,11 @@ URigVMTemplateNode* URigVMController::AddTemplateNode(const FName& InNotation, c
 	URigVMTemplateNode* Node = nullptr;
 
 	// determine what kind of node we need to create
-	if(const FRigVMFunction* FirstFunction = Template->GetPermutation(0))
+	if(Template->UsesDispatch())
+	{
+		Node = NewObject<URigVMDispatchNode>(Graph, *Name);
+	}
+	else if(const FRigVMFunction* FirstFunction = Template->GetPermutation(0))
 	{
 		const UScriptStruct* PotentialUnitStruct = FirstFunction->Struct;
 		if(PotentialUnitStruct && PotentialUnitStruct->IsChildOf(FRigVMStruct::StaticStruct()))
@@ -12075,7 +12104,7 @@ URigVMTemplateNode* URigVMController::AddTemplateNode(const FName& InNotation, c
 		const FRigVMTemplateArgument* Arg = Template->GetArgument(ArgIndex);
 
 		URigVMPin* Pin = NewObject<URigVMPin>(Node, Arg->GetName());
-		const int32& TypeIndex = Types.FindChecked(Arg->GetName());
+		const TRigVMTypeIndex& TypeIndex = Types.FindChecked(Arg->GetName());
 		const FRigVMTemplateArgumentType Type = FRigVMRegistry::Get().GetType(TypeIndex);
 		Pin->CPPType = Type.CPPType.ToString();
 		Pin->CPPTypeObject = Type.CPPTypeObject;
@@ -12094,7 +12123,7 @@ URigVMTemplateNode* URigVMController::AddTemplateNode(const FName& InNotation, c
 			{
 				AddPinsForStruct(ScriptStruct, Pin->GetNode(), Pin, Pin->Direction, DefaultValue, false, false);
 			}
-			else
+			else if(!DefaultValue.IsEmpty())
 			{
 				SetPinDefaultValue(Pin, DefaultValue, true, false, false, false);
 			}
@@ -12179,9 +12208,12 @@ TArray<UScriptStruct*> URigVMController::GetUnitStructsForTemplate(const FName&
 	const FRigVMTemplate* Template = FRigVMRegistry::Get().FindTemplate(InNotation);
 	if(Template)
 	{
-		for(int32 PermutationIndex = 0; PermutationIndex < Template->NumPermutations(); PermutationIndex++)
+		if(!Template->UsesDispatch())
 		{
-			UnitStructs.Add(Template->GetPermutation(PermutationIndex)->Struct);
+			for(int32 PermutationIndex = 0; PermutationIndex < Template->NumPermutations(); PermutationIndex++)
+			{
+				UnitStructs.Add(Template->GetPermutation(PermutationIndex)->Struct);
+			}
 		}
 	}
 	
@@ -12563,7 +12595,7 @@ URigVMArrayNode* URigVMController::AddArrayNode(ERigVMOpCode InOpCode, const FSt
 		TArray<int32> FilterPermutations = Node->GetFilteredPermutationsIndices();
 		if (FilterPermutations.Num() == 1)
 		{
-			const TArray<FRigVMTemplatePreferredType> NewPreferredPermutationTypes = Node->GetArgumentTypesForPermutation(FilterPermutations[0]);
+			const TArray<FRigVMTemplatePreferredType> NewPreferredPermutationTypes = Node->GetPreferredTypesForPermutation(FilterPermutations[0]);
 			if (bSetupUndoRedo)
 			{
 				ActionStack->AddAction(FRigVMSetPreferredTemplatePermutationsAction(Node, NewPreferredPermutationTypes));
@@ -14667,7 +14699,12 @@ void URigVMController::ApplyPinState(URigVMPin* InPin, const FPinState& InPinSta
 
 	if (!InPinState.DefaultValue.IsEmpty())
 	{
-		SetPinDefaultValue(InPin, InPinState.DefaultValue, true, bSetupUndoRedo, false);
+		FString DefaultValue = InPinState.DefaultValue;
+		PostProcessDefaultValue(InPin, DefaultValue);
+		if(!DefaultValue.IsEmpty())
+		{
+			SetPinDefaultValue(InPin, DefaultValue, true, bSetupUndoRedo, false);
+		}
 	}
 
 	SetPinExpansion(InPin, InPinState.bIsExpanded, bSetupUndoRedo);
@@ -14876,6 +14913,16 @@ void URigVMController::CreateDefaultValueForStructIfRequired(UScriptStruct* InSt
 
 void URigVMController::PostProcessDefaultValue(URigVMPin* Pin, FString& OutDefaultValue)
 {
+	static const FString NoneString = FName(NAME_None).ToString();
+	static const FString QuotedNoneString = FString::Printf(TEXT("\"%s\""), *NoneString);
+	if(OutDefaultValue == NoneString || OutDefaultValue == QuotedNoneString)
+	{
+		if(!Pin->IsStringType())
+		{
+			OutDefaultValue.Reset();
+		}
+	}
+
 	if (Pin->IsArray() && OutDefaultValue.IsEmpty())
 	{
 		OutDefaultValue = TEXT("()");
@@ -14933,14 +14980,17 @@ void URigVMController::ResolveTemplateNodeMetaData(URigVMTemplateNode* InNode, b
 	{
 		for(URigVMPin* Pin : InNode->GetPins())
 		{
-			if(Pin->IsWildCard() || Pin->ContainsWildCardSubPin())
+			if(Pin->IsWildCard() || Pin->ContainsWildCardSubPin() || Pin->GetDirection() == ERigVMPinDirection::Hidden)
 			{
 				continue;
 			}
 			if(!Pin->IsValidDefaultValue(Pin->GetDefaultValue()))
 			{
 				const FString NewDefaultValue = InNode->GetInitialDefaultValueForPin(Pin->GetFName(), FilteredPermutationIndices);
-				SetPinDefaultValue(Pin, NewDefaultValue, true, bSetupUndoRedo, false, true);
+				if(!NewDefaultValue.IsEmpty())
+				{
+					SetPinDefaultValue(Pin, NewDefaultValue, true, bSetupUndoRedo, false, true);
+				}
 			}
 		}
 	}
@@ -14957,6 +15007,8 @@ bool URigVMController::FullyResolveTemplateNode(URigVMTemplateNode* InNode, int3
 	check(InNode);
 
 	const FRigVMTemplate* Template = InNode->GetTemplate();
+	const FRigVMDispatchFactory* Factory = Template->GetDispatchFactory();
+	
 	int32 InputPermutation = InPermutationIndex;
 	
 	// Figure out the permutation index from the pin types
@@ -14973,6 +15025,16 @@ bool URigVMController::FullyResolveTemplateNode(URigVMTemplateNode* InNode, int3
 		Template->Resolve(TypeMap, Permutations, true);
 		check(!Permutations.IsEmpty());
 		InputPermutation = Permutations[0];
+
+		// make sure the permutation exists
+		if(Factory)
+		{
+			const FRigVMFunctionPtr DispatchFunction = Factory->GetDispatchFunction(TypeMap);
+			const FRigVMFunction* ResolvedFunction = Template->GetPermutation(InputPermutation);
+			check(DispatchFunction);
+			check(ResolvedFunction);
+			check(ResolvedFunction->FunctionPtr == DispatchFunction);
+		}
 	}
 	
 
@@ -14983,11 +15045,11 @@ bool URigVMController::FullyResolveTemplateNode(URigVMTemplateNode* InNode, int3
 	// find all existing pins that we may need to change
 	TArray<FRigVMTemplateArgument> MissingPins;
 	TArray<URigVMPin*> PinsToRemove;
-	TMap<URigVMPin*, int32> PinTypesToChange;
+	TMap<URigVMPin*, TRigVMTypeIndex> PinTypesToChange;
 	for(int32 ArgIndex = 0; ArgIndex < Template->NumArguments(); ArgIndex++)
 	{
 		const FRigVMTemplateArgument* Argument = Template->GetArgument(ArgIndex);
-		const int32 ResolvedTypeIndex = Argument->GetSupportedTypeIndices(PermutationIndices)[0];
+		const TRigVMTypeIndex ResolvedTypeIndex = Argument->GetSupportedTypeIndices(PermutationIndices)[0];
 
 		URigVMPin* Pin = InNode->FindPin(Argument->GetName().ToString());
 		if(Pin == nullptr)
@@ -15008,15 +15070,13 @@ bool URigVMController::FullyResolveTemplateNode(URigVMTemplateNode* InNode, int3
 	// find all missing pins which are not arguments for the template
 	if(ResolvedFunction)
 	{
-		TArray<UStruct*> StructsToVisit = FRigVMTemplate::GetSuperStructs(ResolvedFunction->Struct, true);
-		for(UStruct* StructToVisit : StructsToVisit)
+		if(ResolvedFunction->Struct == nullptr)
 		{
-			for (TFieldIterator<FProperty> It(StructToVisit, EFieldIterationFlags::None); It; ++It)
+			const TArray<FRigVMTemplateArgument>& Arguments = Template->Arguments;
+			for(const FRigVMTemplateArgument& Argument : Arguments)
 			{
-				const FRigVMTemplateArgument ExpectedArgument(*It);
-				const int32 ExpectedTypeIndex = ExpectedArgument.GetSupportedTypeIndices()[0];
-
-				if(URigVMPin* Pin = InNode->FindPin(It->GetFName().ToString()))
+				const TRigVMTypeIndex ExpectedTypeIndex = Argument.TypeIndices[InputPermutation];
+				if(URigVMPin* Pin = InNode->FindPin(Argument.GetName().ToString()))
 				{
 					if(Pin->GetTypeIndex() != ExpectedTypeIndex)
 					{
@@ -15025,7 +15085,31 @@ bool URigVMController::FullyResolveTemplateNode(URigVMTemplateNode* InNode, int3
 				}
 				else
 				{
-					MissingPins.Add(ExpectedArgument);
+					MissingPins.Add(Argument);
+				}
+			}
+		}
+		else
+		{
+			TArray<UStruct*> StructsToVisit = FRigVMTemplate::GetSuperStructs(ResolvedFunction->Struct, true);
+			for(UStruct* StructToVisit : StructsToVisit)
+			{
+				for (TFieldIterator<FProperty> It(StructToVisit, EFieldIterationFlags::None); It; ++It)
+				{
+					const FRigVMTemplateArgument ExpectedArgument(*It);
+					const TRigVMTypeIndex ExpectedTypeIndex = ExpectedArgument.GetSupportedTypeIndices()[0];
+
+					if(URigVMPin* Pin = InNode->FindPin(It->GetFName().ToString()))
+					{
+						if(Pin->GetTypeIndex() != ExpectedTypeIndex)
+						{
+							PinTypesToChange.Add(Pin, ExpectedTypeIndex);
+						}
+					}
+					else
+					{
+						MissingPins.Add(ExpectedArgument);
+					}
 				}
 			}
 		}
@@ -15034,11 +15118,29 @@ bool URigVMController::FullyResolveTemplateNode(URigVMTemplateNode* InNode, int3
 	// find all pins which don't have a matching arg on the function
 	if(ResolvedFunction)
 	{
-		for(URigVMPin* Pin : InNode->GetPins())
+		if(ResolvedFunction->Struct == nullptr)
+		{
+			const TArray<FRigVMTemplateArgument>& Arguments = Template->Arguments;
+			for(URigVMPin* Pin : InNode->GetPins())
+			{
+				const bool bFound = Arguments.ContainsByPredicate([Pin](const FRigVMTemplateArgument& Argument)
+				{
+					return Pin->GetFName() == Argument.GetName();
+				});
+				if(!bFound)
+				{
+					PinsToRemove.Add(Pin);
+				}
+			}
+		}
+		else
 		{
-			if(ResolvedFunction->Struct->FindPropertyByName(Pin->GetFName()) == nullptr)
+			for(URigVMPin* Pin : InNode->GetPins())
 			{
-				PinsToRemove.Add(Pin);
+				if(ResolvedFunction->Struct->FindPropertyByName(Pin->GetFName()) == nullptr)
+				{
+					PinsToRemove.Add(Pin);
+				}
 			}
 		}
 
@@ -15060,10 +15162,10 @@ bool URigVMController::FullyResolveTemplateNode(URigVMTemplateNode* InNode, int3
 
 	// update the incorrectly typed pins
 	bool bNeedsTemplateUpdate = false;
-	for(const TPair<URigVMPin*, int32>& Pair : PinTypesToChange)
+	for(const TPair<URigVMPin*, TRigVMTypeIndex>& Pair : PinTypesToChange)
 	{
 		URigVMPin* Pin = Pair.Key;
-		const int32& ExpectedTypeIndex = Pair.Value;
+		const TRigVMTypeIndex& ExpectedTypeIndex = Pair.Value;
 		
 		if(!Pin->IsWildCard())
 		{
@@ -15106,20 +15208,55 @@ bool URigVMController::FullyResolveTemplateNode(URigVMTemplateNode* InNode, int3
 	// add missing pins
 	if(ResolvedFunction)
 	{
-		for(const FRigVMTemplateArgument& MissingPin : MissingPins)
+		if(ResolvedFunction->Struct == nullptr)
 		{
-			check(MissingPin.GetDirection() == ERigVMPinDirection::Hidden);
-			
-			FProperty* Property = ResolvedFunction->Struct->FindPropertyByName(MissingPin.GetName());
-			check(Property);
+			for(const FRigVMTemplateArgument& MissingPin : MissingPins)
+			{
+				check(MissingPin.GetDirection() == ERigVMPinDirection::Hidden);
+				
+				URigVMPin* Pin = NewObject<URigVMPin>(Cast<UObject>(InNode), MissingPin.GetName());
 
-			URigVMPin* Pin = NewObject<URigVMPin>(Cast<UObject>(InNode), MissingPin.GetName());
-			ConfigurePinFromProperty(Property, Pin, MissingPin.GetDirection());
+				const TRigVMTypeIndex TypeIndex = MissingPin.TypeIndices[InputPermutation];
+				const FRigVMTemplateArgumentType& Type = FRigVMRegistry::Get().GetType(TypeIndex); 
+				
+				Pin->Direction = MissingPin.GetDirection();
+				Pin->CPPType = Pin->LastKnownCPPType = Type.CPPType.ToString();
+				Pin->CPPTypeObject = Type.CPPTypeObject;
+				Pin->CPPTypeObjectPath = Type.GetCPPTypeObjectPath();
+				Pin->LastKnownTypeIndex = TypeIndex;
 
-			AddNodePin(InNode, Pin);
-			Notify(ERigVMGraphNotifType::PinAdded, Pin);
+				if(Factory)
+				{
+					const FText DisplayNameText = Factory->GetDisplayNameForArgument(MissingPin.GetName());
+					if(!DisplayNameText.IsEmpty())
+					{
+						Pin->DisplayName = *DisplayNameText.ToString();
+					}
+				}
+				
+				AddNodePin(InNode, Pin);
+				Notify(ERigVMGraphNotifType::PinAdded, Pin);
+
+				// we don't need to set the default value here since the pin is hidden
+			}
+		}
+		else
+		{
+			for(const FRigVMTemplateArgument& MissingPin : MissingPins)
+			{
+				check(MissingPin.GetDirection() == ERigVMPinDirection::Hidden);
+				
+				FProperty* Property = ResolvedFunction->Struct->FindPropertyByName(MissingPin.GetName());
+				check(Property);
 
-			// we don't need to set the default value here since the pin is hidden
+				URigVMPin* Pin = NewObject<URigVMPin>(Cast<UObject>(InNode), MissingPin.GetName());
+				ConfigurePinFromProperty(Property, Pin, MissingPin.GetDirection());
+
+				AddNodePin(InNode, Pin);
+				Notify(ERigVMGraphNotifType::PinAdded, Pin);
+
+				// we don't need to set the default value here since the pin is hidden
+			}
 		}
 	}
 
@@ -15136,7 +15273,7 @@ bool URigVMController::FullyResolveTemplateNode(URigVMTemplateNode* InNode, int3
 	return true;
 }
 
-bool URigVMController::PrepareTemplatePinForType(URigVMPin* InPin, const TArray<int32>& InTypeIndices, bool bSetupUndoRedo)
+bool URigVMController::PrepareTemplatePinForType(URigVMPin* InPin, const TArray<TRigVMTypeIndex>& InTypeIndices, bool bSetupUndoRedo)
 {
 	if (!IsValidGraph())
 	{
@@ -15161,11 +15298,11 @@ bool URigVMController::PrepareTemplatePinForType(URigVMPin* InPin, const TArray<
 
 	// If the pin is an element of an array, get the array pin and convert the types
 	URigVMPin* RootPin = InPin;
-	TArray<int32> Types = InTypeIndices;
+	TArray<TRigVMTypeIndex> Types = InTypeIndices;
 	if (InPin->IsArrayElement())
 	{
 		RootPin = InPin->GetParentPin();
-		for (int32& Type : Types)
+		for (TRigVMTypeIndex& Type : Types)
 		{
 			Type = FRigVMRegistry::Get().GetArrayTypeFromBaseTypeIndex(Type);
 		}
@@ -15183,7 +15320,7 @@ bool URigVMController::PrepareTemplatePinForType(URigVMPin* InPin, const TArray<
 	}
 
 	bool bFilteredSupportsType = false;
-	for (const int32& Type : Types)
+	for (const TRigVMTypeIndex& Type : Types)
 	{
 		if (TemplateNode->FilteredSupportsType(RootPin, Type))
 		{
@@ -15204,7 +15341,7 @@ bool URigVMController::PrepareTemplatePinForType(URigVMPin* InPin, const TArray<
 			TemplateNode->PreferredPermutationPairs = {};
 		}
 		
-		for (const int32& Type : Types)
+		for (const TRigVMTypeIndex& Type : Types)
 		{
 			if (TemplateNode->SupportsType(RootPin, Type))
 			{
@@ -15327,9 +15464,9 @@ bool URigVMController::PrepareTemplatePinForType(URigVMPin* InPin, const TArray<
 	return false;
 }
 
-TArray<int32> URigVMController::GetWildcardFilteredTypeIndices(URigVMPin* InPin)
+TArray<TRigVMTypeIndex> URigVMController::GetWildcardFilteredTypeIndices(URigVMPin* InPin)
 {
-	TArray<int32> Types;
+	TArray<TRigVMTypeIndex> Types;
 	if (!InPin->IsWildCard())
 	{
 		Types.Add(InPin->GetTypeIndex());
@@ -15414,7 +15551,7 @@ bool URigVMController::ResolveWildCardPin(URigVMPin* InPin, const FRigVMTemplate
 	return ResolveWildCardPin(InPin, FRigVMRegistry::Get().GetTypeIndex(InType), bSetupUndoRedo, bPrintPythonCommand);
 }
 
-bool URigVMController::ResolveWildCardPin(const FString& InPinPath, int32 InTypeIndex, bool bSetupUndoRedo,
+bool URigVMController::ResolveWildCardPin(const FString& InPinPath, TRigVMTypeIndex InTypeIndex, bool bSetupUndoRedo,
 	bool bPrintPythonCommand)
 {
 	if (!IsValidGraph())
@@ -15437,7 +15574,7 @@ bool URigVMController::ResolveWildCardPin(const FString& InPinPath, int32 InType
 	return false;
 }
 
-bool URigVMController::ResolveWildCardPin(URigVMPin* InPin, int32 InTypeIndex, bool bSetupUndoRedo,
+bool URigVMController::ResolveWildCardPin(URigVMPin* InPin, TRigVMTypeIndex InTypeIndex, bool bSetupUndoRedo,
 	bool bPrintPythonCommand)
 {
 	if(InPin->IsWildCard())
@@ -15520,7 +15657,7 @@ bool URigVMController::UpdateFilteredPermutations(URigVMPin* InPin, URigVMPin* I
 	return false;
 }
 
-bool URigVMController::UpdateFilteredPermutations(URigVMPin* InPin, const TArray<int32>& InTypeIndices, bool bSetupUndoRedo)
+bool URigVMController::UpdateFilteredPermutations(URigVMPin* InPin, const TArray<TRigVMTypeIndex>& InTypeIndices, bool bSetupUndoRedo)
 {
 	URigVMTemplateNode* Node = Cast<URigVMTemplateNode>(InPin->GetNode());
 	if (!Node)
@@ -15570,7 +15707,7 @@ bool URigVMController::UpdateTemplateNodePinTypes(URigVMTemplateNode* InNode, bo
 				continue;
 			}
 		
-			TArray<int32> Types = InNode->GetFilteredTypesForPin(Pin);
+			TArray<TRigVMTypeIndex> Types = InNode->GetFilteredTypesForPin(Pin);
 			if (Types.IsEmpty())
 			{
 				continue;
@@ -15836,7 +15973,7 @@ bool URigVMController::RecomputeAllTemplateFilteredTypes(bool bSetupUndoRedo)
 				if (OutputPin->IsStructMember())
 				{
 					URigVMPin* RootPin = OutputPin->GetRootPin();
-					const int32 TypeIndex = TypesBeforeRecomputing.FindChecked(RootPin);
+					const TRigVMTypeIndex TypeIndex = TypesBeforeRecomputing.FindChecked(RootPin);
 					if (UpdateFilteredPermutations(RootPin, {TypeIndex}, bSetupUndoRedo))
 					{
 						PropagateTemplateFilteredTypes(OutputNode, bSetupUndoRedo);
@@ -15853,7 +15990,7 @@ bool URigVMController::RecomputeAllTemplateFilteredTypes(bool bSetupUndoRedo)
 				if (InputPin->IsStructMember() && InputNode->PreferredPermutationPairs.IsEmpty())
 				{
 					URigVMPin* RootPin = InputPin->GetRootPin();
-					const int32 TypeIndex = TypesBeforeRecomputing.FindChecked(RootPin);
+					const TRigVMTypeIndex TypeIndex = TypesBeforeRecomputing.FindChecked(RootPin);
 					if (UpdateFilteredPermutations(RootPin, {TypeIndex}, bSetupUndoRedo))
 					{
 						PropagateTemplateFilteredTypes(InputNode, bSetupUndoRedo);
@@ -15902,7 +16039,10 @@ bool URigVMController::RecomputeAllTemplateFilteredTypes(bool bSetupUndoRedo)
 				continue;
 			}
 
-			bAnyTypeChanged |= UpdateTemplateNodePinTypes(TemplateNode, bSetupUndoRedo);
+			if(UpdateTemplateNodePinTypes(TemplateNode, bSetupUndoRedo))
+			{
+				bAnyTypeChanged = true;
+			}
 		}
 	}
 
@@ -16044,11 +16184,11 @@ bool URigVMController::ChangePinType(URigVMPin* InPin, const FString& InCPPType,
 	}
 
 	const FRigVMTemplateArgumentType Type(*InCPPType, InCPPTypeObject);
-	const int32 TypeIndex = FRigVMRegistry::Get().GetTypeIndex(Type);
+	const TRigVMTypeIndex TypeIndex = FRigVMRegistry::Get().GetTypeIndex(Type);
 	return ChangePinType(InPin, TypeIndex, bSetupUndoRedo, bSetupOrphanPins, bBreakLinks, bRemoveSubPins);
 }
 
-bool URigVMController::ChangePinType(URigVMPin* InPin, int32 InTypeIndex, bool bSetupUndoRedo, bool bSetupOrphanPins, bool bBreakLinks, bool bRemoveSubPins)
+bool URigVMController::ChangePinType(URigVMPin* InPin, TRigVMTypeIndex InTypeIndex, bool bSetupUndoRedo, bool bSetupOrphanPins, bool bBreakLinks, bool bRemoveSubPins)
 {
 	if (!bIsTransacting && !IsGraphEditable())
 	{
@@ -16061,7 +16201,7 @@ bool URigVMController::ChangePinType(URigVMPin* InPin, int32 InTypeIndex, bool b
 	}
 
 	// only allow valid pin cpp types on template nodes
-	int32 TypeIndex = InTypeIndex;
+	TRigVMTypeIndex TypeIndex = InTypeIndex;
 	if(URigVMTemplateNode* TemplateNode = Cast<URigVMTemplateNode>(InPin->GetNode()))
 	{
 		if(!TemplateNode->SupportsType(InPin, InTypeIndex, &TypeIndex))
@@ -16079,7 +16219,7 @@ bool URigVMController::ChangePinType(URigVMPin* InPin, int32 InTypeIndex, bool b
 		}
 		else
 		{
-			const int32 BaseTypeIndex = bIsArrayType ? FRigVMRegistry::Get().GetBaseTypeFromArrayTypeIndex(TypeIndex) : TypeIndex;
+			const TRigVMTypeIndex BaseTypeIndex = bIsArrayType ? FRigVMRegistry::Get().GetBaseTypeFromArrayTypeIndex(TypeIndex) : TypeIndex;
 			TypeIndex = InPin->IsArray() ? FRigVMRegistry::Get().GetArrayTypeFromBaseTypeIndex(BaseTypeIndex) : BaseTypeIndex;
 		}
 	}
@@ -16182,7 +16322,7 @@ bool URigVMController::ChangePinType(URigVMPin* InPin, int32 InTypeIndex, bool b
 
 	if (InPin->IsArray())
 	{
-		const int32 BaseTypeIndex = FRigVMRegistry::Get().GetBaseTypeFromArrayTypeIndex(TypeIndex);
+		const TRigVMTypeIndex BaseTypeIndex = FRigVMRegistry::Get().GetBaseTypeFromArrayTypeIndex(TypeIndex);
 		for (int32 i=0; i<InPin->GetSubPins().Num(); ++i)
 		{
 			URigVMPin* SubPin = InPin->GetSubPins()[i];
diff --git a/Engine/Source/Developer/RigVMDeveloper/Private/RigVMModel/RigVMNode.cpp b/Engine/Source/Developer/RigVMDeveloper/Private/RigVMModel/RigVMNode.cpp
--- a/Engine/Source/Developer/RigVMDeveloper/Private/RigVMModel/RigVMNode.cpp
+++ b/Engine/Source/Developer/RigVMDeveloper/Private/RigVMModel/RigVMNode.cpp
@@ -549,6 +549,89 @@ TArray<FRigVMUserWorkflow> URigVMNode::GetSupportedWorkflows(ERigVMUserWorkflowT
 	return URigVMUserWorkflowRegistry::Get()->GetWorkflows(InType, Struct, InSubject);
 }
 
+bool URigVMNode::IsAggregate() const
+{
+#if UE_RIGVM_AGGREGATE_NODES_ENABLED
+	const TArray<URigVMPin*> AggregateInputs = GetAggregateInputs();
+	const TArray<URigVMPin*> AggregateOutputs = GetAggregateOutputs();
+
+	if ((AggregateInputs.Num() == 2 && AggregateOutputs.Num() == 1) ||
+		(AggregateInputs.Num() == 1 && AggregateOutputs.Num() == 2))
+	{
+		TArray<URigVMPin*> AggregateAll = AggregateInputs;
+		AggregateAll.Append(AggregateOutputs);
+		for (int32 i = 1; i < 3; ++i)
+		{
+			if (AggregateAll[0]->GetCPPType() != AggregateAll[i]->GetCPPType() ||
+				AggregateAll[0]->GetCPPTypeObject() != AggregateAll[i]->GetCPPTypeObject())
+			{
+				return false;
+			}
+		}
+		
+		return true;
+	}
+#endif
+	
+	return false;
+}
+
+URigVMPin* URigVMNode::GetFirstAggregatePin() const
+{
+#if UE_RIGVM_AGGREGATE_NODES_ENABLED
+	const TArray<URigVMPin*> Inputs = GetAggregateInputs();
+	const TArray<URigVMPin*> Outputs = GetAggregateOutputs();
+	if (Inputs.Num() == 2 && Outputs.Num() == 1)
+	{
+		return Inputs[0];
+	}
+	if (Inputs.Num() == 1 && Outputs.Num() == 2)
+	{
+		return Outputs[0];
+	}
+#endif
+	return nullptr;
+}
+
+URigVMPin* URigVMNode::GetSecondAggregatePin() const
+{
+#if UE_RIGVM_AGGREGATE_NODES_ENABLED
+	const TArray<URigVMPin*> Inputs = GetAggregateInputs();
+	const TArray<URigVMPin*> Outputs = GetAggregateOutputs();
+	if (Inputs.Num() == 2 && Outputs.Num() == 1)
+	{
+		return Inputs[1];
+	}
+	if (Inputs.Num() == 1 && Outputs.Num() == 2)
+	{
+		return Outputs[1];
+	}
+#endif
+	return nullptr;
+}
+
+URigVMPin* URigVMNode::GetOppositeAggregatePin() const
+{
+#if UE_RIGVM_AGGREGATE_NODES_ENABLED
+	const TArray<URigVMPin*> Inputs = GetAggregateInputs();
+	const TArray<URigVMPin*> Outputs = GetAggregateOutputs();
+	if (Inputs.Num() == 2 && Outputs.Num() == 1)
+	{
+		return Outputs[0];
+	}
+	if (Inputs.Num() == 1 && Outputs.Num() == 2)
+	{
+		return Inputs[0];
+	}
+#endif
+	return nullptr;
+}
+
+bool URigVMNode::IsInputAggregate() const
+{
+	return GetAggregateInputs().Num() == 2;
+}
+
 #if WITH_EDITOR
 
 const URigVMNode::FProfilingCache* URigVMNode::UpdateProfilingCacheIfNeeded(URigVM* InVM, const FRigVMASTProxy& InProxy) const
diff --git a/Engine/Source/Developer/RigVMDeveloper/Private/RigVMModel/RigVMPin.cpp b/Engine/Source/Developer/RigVMDeveloper/Private/RigVMModel/RigVMPin.cpp
--- a/Engine/Source/Developer/RigVMDeveloper/Private/RigVMModel/RigVMPin.cpp
+++ b/Engine/Source/Developer/RigVMDeveloper/Private/RigVMModel/RigVMPin.cpp
@@ -616,7 +616,7 @@ FRigVMTemplateArgumentType URigVMPin::GetTemplateArgumentType() const
 	return FRigVMRegistry::Get().GetType(GetTypeIndex());
 }
 
-int32 URigVMPin::GetTypeIndex() const
+TRigVMTypeIndex URigVMPin::GetTypeIndex() const
 {
 	if(LastKnownCPPType != GetCPPType())
 	{
diff --git a/Engine/Source/Developer/RigVMDeveloper/Public/RigVMCompiler/RigVMCompiler.h b/Engine/Source/Developer/RigVMDeveloper/Public/RigVMCompiler/RigVMCompiler.h
--- a/Engine/Source/Developer/RigVMDeveloper/Public/RigVMCompiler/RigVMCompiler.h
+++ b/Engine/Source/Developer/RigVMDeveloper/Public/RigVMCompiler/RigVMCompiler.h
@@ -88,7 +88,6 @@ public:
 	TMap<const FRigVMExprAST*, bool> ExprComplete;
 	TArray<const FRigVMExprAST*> ExprToSkip;
 	TMap<FString, int32> ProcessedLinks;
-	TArray<TSharedPtr<FStructOnScope>> DefaultStructs;
 	TMap<int32, FRigVMOperand> IntegerLiterals;
 	FRigVMOperand ComparisonOperand;
 
@@ -199,7 +198,7 @@ private:
 	FRigVMOperand FindOrAddRegister(const FRigVMVarExprAST* InVarExpr, FRigVMCompilerWorkData& WorkData, bool bIsDebugValue = false);
 	const FRigVMCompilerWorkData::FRigVMASTProxyArray& FindProxiesWithSharedOperand(const FRigVMVarExprAST* InVarExpr, FRigVMCompilerWorkData& WorkData);
 
-	bool ValidateNode(URigVMNode* InNode);
+	bool ValidateNode(URigVMNode* InNode, bool bCheck = true);
 	
 	void ReportInfo(const FString& InMessage);
 	void ReportWarning(const FString& InMessage);
diff --git a/Engine/Source/Developer/RigVMDeveloper/Public/RigVMModel/Nodes/RigVMDispatchNode.h b/Engine/Source/Developer/RigVMDeveloper/Public/RigVMModel/Nodes/RigVMDispatchNode.h
new file mode 100644
--- /dev/null
+++ b/Engine/Source/Developer/RigVMDeveloper/Public/RigVMModel/Nodes/RigVMDispatchNode.h
@@ -0,0 +1,50 @@
+// Copyright Epic Games, Inc. All Rights Reserved.
+
+#pragma once
+
+#include "RigVMModel/Nodes/RigVMTemplateNode.h"
+#include "RigVMCore/RigVMDispatchFactory.h"
+#include "RigVMDispatchNode.generated.h"
+
+/**
+ * The Struct Node represents a Function Invocation of a RIGVM_METHOD
+ * declared on a USTRUCT. Struct Nodes have input / output pins for all
+ * struct UPROPERTY members.
+ */
+UCLASS(BlueprintType)
+class RIGVMDEVELOPER_API URigVMDispatchNode : public URigVMTemplateNode
+{
+	GENERATED_BODY()
+
+public:
+
+	// Override node functions
+	virtual FString GetNodeTitle() const override;
+	virtual FText GetToolTipText() const override;
+	virtual FLinearColor GetNodeColor() const override;
+	virtual bool IsDefinedAsConstant() const override;
+	virtual bool IsDefinedAsVarying() const override;
+	virtual TArray<URigVMPin*> GetAggregateInputs() const override;
+	virtual TArray<URigVMPin*> GetAggregateOutputs() const override;
+	virtual FName GetNextAggregateName(const FName& InLastAggregatePinName) const override;
+	virtual FRigVMStructUpgradeInfo GetUpgradeInfo() const override;
+
+	const FRigVMDispatchFactory* GetFactory() const;
+	bool IsDeprecated() const;
+	FString GetDeprecatedMetadata() const;
+
+protected:
+
+	virtual FText GetToolTipTextForPin(const URigVMPin* InPin) const override;
+	const UScriptStruct* GetFactoryStruct() const;
+	const FRigVMTemplateTypeMap& GetFilteredTypes() const;
+	virtual void InvalidateCache() override;
+
+private:
+
+	mutable FRigVMTemplateTypeMap TypesFromPins;
+	mutable const FRigVMDispatchFactory* CachedFactory = nullptr;
+
+	friend class URigVMController;
+};
+
diff --git a/Engine/Source/Developer/RigVMDeveloper/Public/RigVMModel/Nodes/RigVMTemplateNode.h b/Engine/Source/Developer/RigVMDeveloper/Public/RigVMModel/Nodes/RigVMTemplateNode.h
--- a/Engine/Source/Developer/RigVMDeveloper/Public/RigVMModel/Nodes/RigVMTemplateNode.h
+++ b/Engine/Source/Developer/RigVMDeveloper/Public/RigVMModel/Nodes/RigVMTemplateNode.h
@@ -10,20 +10,39 @@ USTRUCT()
 struct RIGVMDEVELOPER_API FRigVMTemplatePreferredType
 {
 	GENERATED_BODY()
+
+public:
 	
 	FRigVMTemplatePreferredType()
 		: Argument(NAME_None)
 		, TypeIndex(INDEX_NONE)
 	{}
 	
-	FRigVMTemplatePreferredType(const FName& InArgument, int32 InTypeIndex)
+	FRigVMTemplatePreferredType(const FName& InArgument, TRigVMTypeIndex InTypeIndex)
 		: Argument(InArgument)
 		, TypeIndex(InTypeIndex)
-	{}
+	{
+#if UE_RIGVM_DEBUG_TYPEINDEX
+		UpdateStringFromIndex();
+#endif
+	}
+
+	FORCEINLINE const FName& GetArgument() const { return Argument; }
+	FORCEINLINE TRigVMTypeIndex GetTypeIndex() const
+	{
+#if UE_RIGVM_DEBUG_TYPEINDEX
+		const FRigVMTemplateArgumentType& Type = FRigVMRegistry::Get().GetType(TypeIndex);
+		return FRigVMRegistry::Get().GetTypeIndex(Type);
+#else
+		return TypeIndex;
+#endif
+	}
 	
 	void UpdateStringFromIndex();
 	void UpdateIndexFromString();
 
+protected:
+	
 	UPROPERTY()
 	FName Argument;
 
@@ -32,6 +51,8 @@ struct RIGVMDEVELOPER_API FRigVMTemplatePreferredType
 
 	UPROPERTY()
 	FString TypeString;
+
+	friend class URigVMTemplateNode;
 };
 
 /**
@@ -72,10 +93,10 @@ public:
 	virtual bool IsSingleton() const;
 
 	// returns true if a pin supports a given type
-	bool SupportsType(const URigVMPin* InPin, int32 InTypeIndex, int32* OutTypeIndex = nullptr);
+	bool SupportsType(const URigVMPin* InPin, TRigVMTypeIndex InTypeIndex, TRigVMTypeIndex* OutTypeIndex = nullptr);
 
 	// returns true if a pin supports a given type after filtering
-	bool FilteredSupportsType(const URigVMPin* InPin, int32 InTypeIndex, int32* OutTypeIndex = nullptr, bool bAllowFloatingPointCasts = true);
+	bool FilteredSupportsType(const URigVMPin* InPin, TRigVMTypeIndex InTypeIndex, TRigVMTypeIndex* OutTypeIndex = nullptr, bool bAllowFloatingPointCasts = true);
 
 	// returns the resolved functions for the template
 	TArray<const FRigVMFunction*> GetResolvedPermutations() const;
@@ -104,14 +125,14 @@ public:
 	const TArray<int32>& GetFilteredPermutationsIndices() const;
 
 	// returns the filtered types of this pin
-	TArray<int32> GetFilteredTypesForPin(URigVMPin* InPin) const;
+	TArray<TRigVMTypeIndex> GetFilteredTypesForPin(URigVMPin* InPin) const;
 
 	// returns true if updating pin filters with InTypes would result in different filters 
-	bool PinNeedsFilteredTypesUpdate(URigVMPin* InPin, const TArray<int32>& InTypeIndices);
+	bool PinNeedsFilteredTypesUpdate(URigVMPin* InPin, const TArray<TRigVMTypeIndex>& InTypeIndices);
 	bool PinNeedsFilteredTypesUpdate(URigVMPin* InPin, URigVMPin* LinkedPin);
 
 	// updates the filtered permutations given a link or the types for a pin
-	bool UpdateFilteredPermutations(URigVMPin* InPin, const TArray<int32>& InTypeIndices);
+	bool UpdateFilteredPermutations(URigVMPin* InPin, const TArray<TRigVMTypeIndex>& InTypeIndices);
 	bool UpdateFilteredPermutations(URigVMPin* InPin, URigVMPin* LinkedPin);
 
 	// initializes the filtered permutations to all possible permutations
@@ -131,10 +152,11 @@ protected:
 	virtual void InvalidateCache() override;
 	
 	TArray<int32> GetNewFilteredPermutations(URigVMPin* InPin, URigVMPin* LinkedPin);
-	TArray<int32> GetNewFilteredPermutations(URigVMPin* InPin, const TArray<int32>& InTypeIndices);
+	TArray<int32> GetNewFilteredPermutations(URigVMPin* InPin, const TArray<TRigVMTypeIndex>& InTypeIndices);
 
 	TArray<int32> FindPermutationsForTypes(const TArray<FRigVMTemplatePreferredType>& ArgumentTypes, bool bAllowCasting = false) const;
-	TArray<FRigVMTemplatePreferredType> GetArgumentTypesForPermutation(const int32 InPermutationIndex) const;
+	TArray<FRigVMTemplatePreferredType> GetPreferredTypesForPermutation(const int32 InPermutationIndex) const;
+	FRigVMTemplateTypeMap GetTypesForPermutation(const int32 InPermutationIndex) const;
 
 	UPROPERTY()
 	FName TemplateNotation;
@@ -151,7 +173,7 @@ protected:
 	TArray<FRigVMTemplatePreferredType> PreferredPermutationPairs;
 
 	TArray<int32> FilteredPermutations;
-	TMap<TPair<FName,int32>, TPair<bool, int32>> SupportedTypesCache;
+	TMap<TPair<FName,TRigVMTypeIndex>, TPair<bool, TRigVMTypeIndex>> SupportedTypesCache;
 
 	mutable const FRigVMTemplate* CachedTemplate;
 	mutable const FRigVMFunction* CachedFunction;
diff --git a/Engine/Source/Developer/RigVMDeveloper/Public/RigVMModel/Nodes/RigVMUnitNode.h b/Engine/Source/Developer/RigVMDeveloper/Public/RigVMModel/Nodes/RigVMUnitNode.h
--- a/Engine/Source/Developer/RigVMDeveloper/Public/RigVMModel/Nodes/RigVMUnitNode.h
+++ b/Engine/Source/Developer/RigVMDeveloper/Public/RigVMModel/Nodes/RigVMUnitNode.h
@@ -30,13 +30,7 @@ public:
 	virtual FName GetEventName() const override;
 	virtual bool CanOnlyExistOnce() const override;
 	virtual bool IsLoopNode() const override;
-	virtual bool CanBeUpgraded() const override { return GetUpgradeInfo().IsValid(); }
 	virtual TArray<FRigVMUserWorkflow> GetSupportedWorkflows(ERigVMUserWorkflowType InType, const UObject* InSubject) const override;
-	virtual bool IsAggregate() const override;
-	virtual URigVMPin* GetFirstAggregatePin() const override;
-	virtual URigVMPin* GetSecondAggregatePin() const override;
-	virtual URigVMPin* GetOppositeAggregatePin() const override;
-	virtual bool IsInputAggregate() const override;
 	virtual TArray<URigVMPin*> GetAggregateInputs() const override;
 	virtual TArray<URigVMPin*> GetAggregateOutputs() const override;
 	virtual FName GetNextAggregateName(const FName& InLastAggregatePinName) const override;
@@ -77,7 +71,7 @@ public:
 		return InstanceRef;
 	}
 
-	FRigVMStructUpgradeInfo GetUpgradeInfo() const;
+	virtual FRigVMStructUpgradeInfo GetUpgradeInfo() const override;
 
 protected:
 
diff --git a/Engine/Source/Developer/RigVMDeveloper/Public/RigVMModel/RigVMController.h b/Engine/Source/Developer/RigVMDeveloper/Public/RigVMModel/RigVMController.h
--- a/Engine/Source/Developer/RigVMDeveloper/Public/RigVMModel/RigVMController.h
+++ b/Engine/Source/Developer/RigVMDeveloper/Public/RigVMModel/RigVMController.h
@@ -379,8 +379,8 @@ public:
 	UFUNCTION(BlueprintCallable, Category = RigVMController)
 	bool ResolveWildCardPin(const FString& InPinPath, const FString& InCPPType, const FName& InCPPTypeObjectPath, bool bSetupUndoRedo = true, bool bPrintPythonCommand = false);
 	bool ResolveWildCardPin(URigVMPin* InPin, const FRigVMTemplateArgumentType& InType, bool bSetupUndoRedo = true, bool bPrintPythonCommand = false);
-	bool ResolveWildCardPin(const FString& InPinPath, int32 InTypeIndex, bool bSetupUndoRedo = true, bool bPrintPythonCommand = false);
-	bool ResolveWildCardPin(URigVMPin* InPin, int32 InTypeIndex, bool bSetupUndoRedo = true, bool bPrintPythonCommand = false);
+	bool ResolveWildCardPin(const FString& InPinPath, TRigVMTypeIndex InTypeIndex, bool bSetupUndoRedo = true, bool bPrintPythonCommand = false);
+	bool ResolveWildCardPin(URigVMPin* InPin, TRigVMTypeIndex InTypeIndex, bool bSetupUndoRedo = true, bool bPrintPythonCommand = false);
 
 	// Adds a Function / Struct Node to the edited Graph as an injected node
 	// UnitNode represent a RIGVM_METHOD declaration on a USTRUCT.
@@ -1052,15 +1052,15 @@ private:
 	// If any of the types is supported (without breaking any links), then the filtered permutations will be updated and the change will
 	// propagate to other nodes in the graph.
 	// If it is not supported, we will attempt to find and break any links that do not support this change
-	bool PrepareTemplatePinForType(URigVMPin* InPin, const TArray<int32>& InTypeIndices, bool bSetupUndoRedo);
+	bool PrepareTemplatePinForType(URigVMPin* InPin, const TArray<TRigVMTypeIndex>& InTypeIndices, bool bSetupUndoRedo);
 
 	// Get filtered types for a wildcard node. If template node, that means just returning its filtered wildcard types, but if it's another type of node (select, if, rereoute), iterate
 	// its connections to figure out the filtered types
-	TArray<int32> GetWildcardFilteredTypeIndices(URigVMPin* InPin);
+	TArray<TRigVMTypeIndex> GetWildcardFilteredTypeIndices(URigVMPin* InPin);
 	
 	// Updates the permutations allowed without having to break any links
 	bool UpdateFilteredPermutations(URigVMPin* InPin, URigVMPin* InLinkedPin, bool bSetupUndoRedo);
-	bool UpdateFilteredPermutations(URigVMPin* InPin, const TArray<int32>& InTypeIndices, bool bSetupUndoRedo);
+	bool UpdateFilteredPermutations(URigVMPin* InPin, const TArray<TRigVMTypeIndex>& InTypeIndices, bool bSetupUndoRedo);
 
 	// Changes Pin types if filtered types of a pin are unique
 	bool UpdateTemplateNodePinTypes(URigVMTemplateNode* InNode, bool bSetupUndoRedo);
@@ -1072,7 +1072,7 @@ private:
 	bool ChangePinType(const FString& InPinPath, const FString& InCPPType, const FName& InCPPTypeObjectPath, bool bSetupUndoRedo, bool bSetupOrphanPins = true, bool bBreakLinks = true, bool bRemoveSubPins = true);
 	bool ChangePinType(URigVMPin* InPin, const FString& InCPPType, const FName& InCPPTypeObjectPath, bool bSetupUndoRedo, bool bSetupOrphanPins = true, bool bBreakLinks = true, bool bRemoveSubPins = true);
 	bool ChangePinType(URigVMPin* InPin, const FString& InCPPType,UObject* InCPPTypeObject, bool bSetupUndoRedo, bool bSetupOrphanPins = true, bool bBreakLinks = true, bool bRemoveSubPins = true);
-	bool ChangePinType(URigVMPin* InPin, int32 InTypeIndex, bool bSetupUndoRedo, bool bSetupOrphanPins = true, bool bBreakLinks = true, bool bRemoveSubPins = true);
+	bool ChangePinType(URigVMPin* InPin, TRigVMTypeIndex InTypeIndex, bool bSetupUndoRedo, bool bSetupOrphanPins = true, bool bBreakLinks = true, bool bRemoveSubPins = true);
 
 #if WITH_EDITOR
 	void RewireLinks(URigVMPin* OldPin, URigVMPin* NewPin, bool bAsInput, bool bSetupUndoRedo, TArray<URigVMLink*> InLinks = TArray<URigVMLink*>());
diff --git a/Engine/Source/Developer/RigVMDeveloper/Public/RigVMModel/RigVMNode.h b/Engine/Source/Developer/RigVMDeveloper/Public/RigVMModel/RigVMNode.h
--- a/Engine/Source/Developer/RigVMDeveloper/Public/RigVMModel/RigVMNode.h
+++ b/Engine/Source/Developer/RigVMDeveloper/Public/RigVMModel/RigVMNode.h
@@ -4,6 +4,7 @@
 
 #include "RigVMPin.h"
 #include "RigVMCore/RigVM.h"
+#include "RigVMCore/RigVMStruct.h"
 #include "RigVMCore/RigVMUserWorkflow.h"
 #include "RigVMNode.generated.h"
 
@@ -211,7 +212,7 @@ public:
 
 	// returns true if the node can be upgraded
 	UFUNCTION(BlueprintPure, Category = RigVMNode)
-	virtual bool CanBeUpgraded() const { return false; }
+	virtual bool CanBeUpgraded() const { return GetUpgradeInfo().IsValid(); }
 
 	// returns all supported workflows of the node
 	UFUNCTION(BlueprintCallable, Category = RigVMNode)
@@ -230,19 +231,19 @@ public:
 	void SetExecutionIsHaltedAtThisNode(const bool bValue) { bHaltedAtThisNode = bValue; }
 
 	UFUNCTION(BlueprintPure, Category = RigVMNode)
-	virtual bool IsAggregate() const { return false; }
+	virtual bool IsAggregate() const;
 
 	UFUNCTION(BlueprintPure, Category = RigVMNode)
-	virtual URigVMPin* GetFirstAggregatePin() const { return nullptr; }
+	virtual URigVMPin* GetFirstAggregatePin() const;
 
 	UFUNCTION(BlueprintPure, Category = RigVMNode)
-	virtual URigVMPin* GetSecondAggregatePin() const { return nullptr; }
+	virtual URigVMPin* GetSecondAggregatePin() const;
 
 	UFUNCTION(BlueprintPure, Category = RigVMNode)
-	virtual URigVMPin* GetOppositeAggregatePin() const { return nullptr; }
+	virtual URigVMPin* GetOppositeAggregatePin() const;
 
 	UFUNCTION(BlueprintPure, Category = RigVMNode)
-	virtual bool IsInputAggregate() const { return false; }
+	virtual bool IsInputAggregate() const;
 
 	UFUNCTION(BlueprintPure, Category = RigVMNode)
 	virtual TArray<URigVMPin*> GetAggregateInputs() const { return {}; }
@@ -253,6 +254,8 @@ public:
 	UFUNCTION(BlueprintPure, Category = RigVMNode)
 	virtual FName GetNextAggregateName(const FName& InLastAggregatePinName) const { return NAME_None; }
 
+	virtual FRigVMStructUpgradeInfo GetUpgradeInfo() const { return FRigVMStructUpgradeInfo(); }
+
 private:
 
 	static const FString NodeColorName;
diff --git a/Engine/Source/Developer/RigVMDeveloper/Public/RigVMModel/RigVMPin.h b/Engine/Source/Developer/RigVMDeveloper/Public/RigVMModel/RigVMPin.h
--- a/Engine/Source/Developer/RigVMDeveloper/Public/RigVMModel/RigVMPin.h
+++ b/Engine/Source/Developer/RigVMDeveloper/Public/RigVMModel/RigVMPin.h
@@ -243,7 +243,7 @@ public:
 	FRigVMTemplateArgumentType GetTemplateArgumentType() const;
 
 	// Returns the argument type index this pin would represent within a template
-	int32 GetTypeIndex() const;
+	TRigVMTypeIndex GetTypeIndex() const;
 
 	// Returns true if the C++ data type is FString or FName
 	UFUNCTION(BlueprintCallable, Category = RigVMPin)
@@ -513,7 +513,7 @@ private:
 	FString BoundVariablePath_DEPRECATED;
 
 	mutable FString LastKnownCPPType;
-	mutable int32 LastKnownTypeIndex;
+	mutable TRigVMTypeIndex LastKnownTypeIndex;
 
 	static const FString OrphanPinPrefix;
 
diff --git a/Engine/Source/Runtime/RigVM/Private/RigVMCore/RigVM.cpp b/Engine/Source/Runtime/RigVM/Private/RigVMCore/RigVM.cpp
--- a/Engine/Source/Runtime/RigVM/Private/RigVMCore/RigVM.cpp
+++ b/Engine/Source/Runtime/RigVM/Private/RigVMCore/RigVM.cpp
@@ -78,6 +78,7 @@ URigVM::URigVM()
 #endif
     , FunctionNamesPtr(&FunctionNamesStorage)
     , FunctionsPtr(&FunctionsStorage)
+    , FactoriesPtr(&FactoriesStorage)
 #if WITH_EDITOR
 	, FirstEntryEventInQueue(NAME_None)
 #endif
@@ -241,6 +242,7 @@ void URigVM::PostLoad()
 	{
 		Instructions.Reset();
 		FunctionsStorage.Reset();
+		FactoriesStorage.Reset();
 		ParametersNameMap.Reset();
 
 		for (int32 Index = 0; Index < Parameters.Num(); Index++)
@@ -542,6 +544,7 @@ void URigVM::Reset(bool IsIgnoringArchetypeRef)
 	{
 		FunctionNamesStorage.Reset();
 		FunctionsStorage.Reset();
+		FactoriesStorage.Reset();
 		ExternalPropertyPathDescriptions.Reset();
 		ExternalPropertyPaths.Reset();
 		ByteCodeStorage.Reset();
@@ -555,6 +558,7 @@ void URigVM::Reset(bool IsIgnoringArchetypeRef)
 	{
 		FunctionNamesPtr = &FunctionNamesStorage;
 		FunctionsPtr = &FunctionsStorage;
+		FactoriesPtr = &FactoriesStorage;
 		ByteCodePtr = &ByteCodeStorage;
 	}
 
@@ -568,6 +572,7 @@ void URigVM::Empty()
 {
 	FunctionNamesStorage.Empty();
 	FunctionsStorage.Empty();
+	FactoriesStorage.Empty();
 	ExternalPropertyPathDescriptions.Empty();
 	ExternalPropertyPaths.Empty();
 	ByteCodeStorage.Empty();
@@ -652,6 +657,16 @@ void URigVM::CopyFrom(URigVM* InVM, bool bDeferCopy, bool bReferenceLiteralMemor
 		FunctionsPtr = InVM->FunctionsPtr;
 	}
 	
+	if(InVM->FactoriesPtr == &InVM->FactoriesStorage && !bReferenceByteCode)
+	{
+		FactoriesStorage = InVM->FactoriesStorage;
+		FactoriesPtr = &FactoriesStorage;
+	}
+	else
+	{
+		FactoriesPtr = InVM->FactoriesPtr;
+	}
+
 	if(InVM->ByteCodePtr == &InVM->ByteCodeStorage && !bReferenceByteCode)
 	{
 		ByteCodeStorage = InVM->ByteCodeStorage;
@@ -677,20 +692,27 @@ void URigVM::CopyFrom(URigVM* InVM, bool bDeferCopy, bool bReferenceLiteralMemor
 int32 URigVM::AddRigVMFunction(UScriptStruct* InRigVMStruct, const FName& InMethodName)
 {
 	check(InRigVMStruct);
-	FString FunctionKey = FString::Printf(TEXT("F%s::%s"), *InRigVMStruct->GetName(), *InMethodName.ToString());
-	int32 FunctionIndex = GetFunctionNames().Find(*FunctionKey);
+	const FString FunctionKey = FString::Printf(TEXT("F%s::%s"), *InRigVMStruct->GetName(), *InMethodName.ToString());
+	return AddRigVMFunction(FunctionKey);
+}
+
+int32 URigVM::AddRigVMFunction(const FString& InFunctionName)
+{
+	const FName FunctionName = *InFunctionName;
+	const int32 FunctionIndex = GetFunctionNames().Find(FunctionName);
 	if (FunctionIndex != INDEX_NONE)
 	{
 		return FunctionIndex;
 	}
 
-	const FRigVMFunction* Function = FRigVMRegistry::Get().FindFunction(*FunctionKey);
+	const FRigVMFunction* Function = FRigVMRegistry::Get().FindFunction(*InFunctionName);
 	if (Function == nullptr)
 	{
 		return INDEX_NONE;
 	}
 
-	GetFunctionNames().Add(*FunctionKey);
+	GetFunctionNames().Add(FunctionName);
+	GetFactories().Add(Function->Factory);
 	return GetFunctions().Add(Function->FunctionPtr);
 }
 
@@ -958,13 +980,17 @@ void URigVM::ResolveFunctionsIfRequired()
 	{
 		GetFunctions().Reset();
 		GetFunctions().SetNumZeroed(GetFunctionNames().Num());
+		GetFactories().Reset();
+		GetFactories().SetNumZeroed(GetFunctionNames().Num());
 
 		for (int32 FunctionIndex = 0; FunctionIndex < GetFunctionNames().Num(); FunctionIndex++)
 		{
 			GetFunctions()[FunctionIndex] = nullptr;
+			GetFactories()[FunctionIndex] = nullptr;
 			if(const FRigVMFunction* Function = FRigVMRegistry::Get().FindFunction(*GetFunctionNames()[FunctionIndex].ToString()))
 			{
 				GetFunctions()[FunctionIndex] = Function->FunctionPtr;
+				GetFactories()[FunctionIndex] = Function->Factory;
 			}
 			ensureMsgf(GetFunctions()[FunctionIndex], TEXT("Function %s is not valid"), *GetFunctionNames()[FunctionIndex].ToString());
 		}
@@ -1468,6 +1494,7 @@ bool URigVM::Initialize(TArrayView<URigVMMemoryStorage*> Memory, TArrayView<void
 	
 	FRigVMByteCode& ByteCode = GetByteCode();
 	TArray<FRigVMFunctionPtr>& Functions = GetFunctions();
+	TArray<const FRigVMDispatchFactory*>& Factories = GetFactories();
 
 #if WITH_EDITOR
 	TArray<FName>& FunctionNames = GetFunctionNames();
@@ -1557,6 +1584,7 @@ bool URigVM::Initialize(TArrayView<URigVMMemoryStorage*> Memory, TArrayView<void
 #if WITH_EDITOR
 				Context.PublicData.FunctionName = FunctionNames[Op.FunctionIndex];
 #endif
+				Context.Factory = Factories[Op.FunctionIndex];
 
 				// find out the largest slice count
 				int32 MaxSliceCount = 1;
@@ -1667,6 +1695,7 @@ bool URigVM::Execute(TArrayView<URigVMMemoryStorage*> Memory, TArrayView<void*>
 	
 	FRigVMByteCode& ByteCode = GetByteCode();
 	TArray<FRigVMFunctionPtr>& Functions = GetFunctions();
+	TArray<const FRigVMDispatchFactory*>& Factories = GetFactories();
 
 #if WITH_EDITOR
 	TArray<FName>& FunctionNames = GetFunctionNames();
@@ -1861,6 +1890,7 @@ bool URigVM::Execute(TArrayView<URigVMMemoryStorage*> Memory, TArrayView<void*>
 #if WITH_EDITOR
 				Context.PublicData.FunctionName = FunctionNames[Op.FunctionIndex];
 #endif
+				Context.Factory = Factories[Op.FunctionIndex];
 				(*Functions[Op.FunctionIndex])(Context, Handles);
 
 #if WITH_EDITOR
diff --git a/Engine/Source/Runtime/RigVM/Private/RigVMCore/RigVMDispatchFactory.cpp b/Engine/Source/Runtime/RigVM/Private/RigVMCore/RigVMDispatchFactory.cpp
new file mode 100644
--- /dev/null
+++ b/Engine/Source/Runtime/RigVM/Private/RigVMCore/RigVMDispatchFactory.cpp
@@ -0,0 +1,133 @@
+// Copyright Epic Games, Inc. All Rights Reserved.
+
+#include "RigVMCore/RigVMDispatchFactory.h"
+#include "RigVMCore/RigVMRegistry.h"
+#include "RigVMCore/RigVMStruct.h"
+
+const FString FRigVMDispatchFactory::DispatchPrefix = TEXT("DISPATCH_");
+
+FName FRigVMDispatchFactory::GetFactoryName() const
+{
+	static constexpr TCHAR Format[] = TEXT("%s%s");
+	return *FString::Printf(Format, *DispatchPrefix, *GetScriptStruct()->GetName());
+}
+
+#if WITH_EDITOR
+
+FName FRigVMDispatchFactory::GetNextAggregateName(const FName& InLastAggregatePinName) const
+{
+	return FRigVMStruct().GetNextAggregateName(InLastAggregatePinName);
+}
+
+FLinearColor FRigVMDispatchFactory::GetNodeColor() const
+{
+	if(const UScriptStruct* ScriptStruct = GetScriptStruct())
+	{
+		FString NodeColor;
+		if (ScriptStruct->GetStringMetaDataHierarchical(FRigVMStruct::NodeColorMetaName, &NodeColor))
+		{
+			return FRigVMTemplate::GetColorFromMetadata(NodeColor);
+		}
+	}
+	return FLinearColor::White;
+}
+
+FText FRigVMDispatchFactory::GetNodeTooltip(const FRigVMTemplateTypeMap& InTypes) const
+{
+	return GetScriptStruct()->GetToolTipText();
+}
+
+FString FRigVMDispatchFactory::GetCategory() const
+{
+	if(const UScriptStruct* ScriptStruct = GetScriptStruct())
+	{
+		FString Category;
+		if (ScriptStruct->GetStringMetaDataHierarchical(FRigVMStruct::CategoryMetaName, &Category))
+		{
+			return Category;
+		}
+	}
+	return FString();
+}
+
+FString FRigVMDispatchFactory::GetKeywords() const
+{
+	if(const UScriptStruct* ScriptStruct = GetScriptStruct())
+	{
+		FString Keywords;
+		if (ScriptStruct->GetStringMetaDataHierarchical(FRigVMStruct::KeywordsMetaName, &Keywords))
+		{
+			return Keywords;
+		}
+	}
+	return FString();
+}
+
+#endif
+
+FRigVMFunctionPtr FRigVMDispatchFactory::GetDispatchFunction(const FRigVMTemplateTypeMap& InTypes) const
+{
+	const FString PermutationName = GetPermutationName(InTypes);
+	if(const FRigVMFunction* ExistingFunction = FRigVMRegistry::Get().FindFunction(*PermutationName))
+	{
+		return ExistingFunction->FunctionPtr;
+	}
+	return GetDispatchFunctionImpl(InTypes);
+}
+
+FString FRigVMDispatchFactory::GetPermutationName(const FRigVMTemplateTypeMap& InTypes) const
+{
+	const FRigVMRegistry& Registry = FRigVMRegistry::Get();
+	TArray<FString> TypePairStrings;
+	for(const TPair<FName,TRigVMTypeIndex>& Pair : InTypes)
+	{
+		const FRigVMTemplateArgumentType& Type = Registry.GetType(Pair.Value);
+		static constexpr TCHAR Format[] = TEXT("%s:%s");
+		TypePairStrings.Add(FString::Printf(Format, *Pair.Key.ToString(), *Type.CPPType.ToString()));
+	}
+			
+	static constexpr TCHAR Format[] = TEXT("%s::%s");
+	return FString::Printf(Format, *GetFactoryName().ToString(), *FString::Join(TypePairStrings, TEXT(",")));
+}
+
+const FRigVMTemplate* FRigVMDispatchFactory::GetTemplate() const
+{
+	static bool bIsDispatchingTemplate = false;
+	if(bIsDispatchingTemplate)
+	{
+			return nullptr;
+	}
+	TGuardValue<bool> ReEntryGuard(bIsDispatchingTemplate, true);
+	
+	if(CachedTemplate)
+	{
+		return CachedTemplate;
+	}
+	
+	FRigVMRegistry& Registry = FRigVMRegistry::Get();
+
+	const FName FactoryName = GetFactoryName();
+	const TArray<FRigVMTemplateArgument> Arguments = GetArguments();
+	
+	FRigVMTemplateDelegates Delegates;
+	Delegates.NewArgumentTypeDelegate = FRigVMTemplate_NewArgumentTypeDelegate::CreateLambda(
+		[this](const FRigVMTemplate*, const FName& InArgumentName, int32 InTypeIndex)
+		{
+			return OnNewArgumentType(InArgumentName, InTypeIndex);
+		});
+
+	Delegates.GetDispatchFactoryDelegate = FRigVMTemplate_GetDispatchFactoryDelegate::CreateLambda(
+[FactoryName]()
+	{
+		return FRigVMRegistry::Get().FindDispatchFactory(FactoryName);
+	});
+
+	Delegates.RequestDispatchFunctionDelegate = FRigVMTemplate_RequestDispatchFunctionDelegate::CreateLambda(
+	[this](const FRigVMTemplate*, const FRigVMTemplateTypeMap& InTypes)
+	{
+		return GetDispatchFunction(InTypes);
+	});
+
+	CachedTemplate = Registry.GetOrAddTemplateFromArguments(GetFactoryName(), Arguments, Delegates); 
+	return CachedTemplate;
+}
diff --git a/Engine/Source/Runtime/RigVM/Private/RigVMCore/RigVMRegistry.cpp b/Engine/Source/Runtime/RigVM/Private/RigVMCore/RigVMRegistry.cpp
--- a/Engine/Source/Runtime/RigVM/Private/RigVMCore/RigVMRegistry.cpp
+++ b/Engine/Source/Runtime/RigVM/Private/RigVMCore/RigVMRegistry.cpp
@@ -9,6 +9,19 @@
 FRigVMRegistry FRigVMRegistry::s_RigVMRegistry;
 const FName FRigVMRegistry::TemplateNameMetaName = TEXT("TemplateName");
 
+FRigVMRegistry::~FRigVMRegistry()
+{
+	for(FRigVMDispatchFactory* Factory : Factories)
+	{
+		if(const UScriptStruct* ScriptStruct = Factory->GetScriptStruct())
+		{
+			ScriptStruct->DestroyStruct(Factory, 1);
+		}
+		FMemory::Free(Factory);
+	}
+	Factories.Reset();
+}
+
 FRigVMRegistry& FRigVMRegistry::Get()
 {
 	s_RigVMRegistry.InitializeIfNeeded();
@@ -47,25 +60,25 @@ void FRigVMRegistry::InitializeIfNeeded()
 	TypesPerCategory.Reserve(19);
 	ArgumentsPerCategory.Reserve(19);
 	
-	TypesPerCategory.Add(FRigVMTemplateArgument::ETypeCategory_Execute, TArray<int32>()).Reserve(8);
-	TypesPerCategory.Add(FRigVMTemplateArgument::ETypeCategory_SingleAnyValue, TArray<int32>()).Reserve(256);
-	TypesPerCategory.Add(FRigVMTemplateArgument::ETypeCategory_ArrayAnyValue, TArray<int32>()).Reserve(256);
-	TypesPerCategory.Add(FRigVMTemplateArgument::ETypeCategory_ArrayArrayAnyValue, TArray<int32>()).Reserve(256);
-	TypesPerCategory.Add(FRigVMTemplateArgument::ETypeCategory_SingleSimpleValue, TArray<int32>()).Reserve(8);
-	TypesPerCategory.Add(FRigVMTemplateArgument::ETypeCategory_ArraySimpleValue, TArray<int32>()).Reserve(8);
-	TypesPerCategory.Add(FRigVMTemplateArgument::ETypeCategory_ArrayArraySimpleValue, TArray<int32>()).Reserve(8);
-	TypesPerCategory.Add(FRigVMTemplateArgument::ETypeCategory_SingleMathStructValue, TArray<int32>()).Reserve(GetMathTypes().Num());
-	TypesPerCategory.Add(FRigVMTemplateArgument::ETypeCategory_ArrayMathStructValue, TArray<int32>()).Reserve(GetMathTypes().Num());
-	TypesPerCategory.Add(FRigVMTemplateArgument::ETypeCategory_ArrayArrayMathStructValue, TArray<int32>()).Reserve(GetMathTypes().Num());
-	TypesPerCategory.Add(FRigVMTemplateArgument::ETypeCategory_SingleScriptStructValue, TArray<int32>()).Reserve(128);
-	TypesPerCategory.Add(FRigVMTemplateArgument::ETypeCategory_ArrayScriptStructValue, TArray<int32>()).Reserve(128);
-	TypesPerCategory.Add(FRigVMTemplateArgument::ETypeCategory_ArrayArrayScriptStructValue, TArray<int32>()).Reserve(128);
-	TypesPerCategory.Add(FRigVMTemplateArgument::ETypeCategory_SingleEnumValue, TArray<int32>()).Reserve(128);
-	TypesPerCategory.Add(FRigVMTemplateArgument::ETypeCategory_ArrayEnumValue, TArray<int32>()).Reserve(128);
-	TypesPerCategory.Add(FRigVMTemplateArgument::ETypeCategory_ArrayArrayEnumValue, TArray<int32>()).Reserve(128);
-	TypesPerCategory.Add(FRigVMTemplateArgument::ETypeCategory_SingleObjectValue, TArray<int32>()).Reserve(128);
-	TypesPerCategory.Add(FRigVMTemplateArgument::ETypeCategory_ArrayObjectValue, TArray<int32>()).Reserve(128);
-	TypesPerCategory.Add(FRigVMTemplateArgument::ETypeCategory_ArrayArrayObjectValue, TArray<int32>()).Reserve(128);
+	TypesPerCategory.Add(FRigVMTemplateArgument::ETypeCategory_Execute, TArray<TRigVMTypeIndex>()).Reserve(8);
+	TypesPerCategory.Add(FRigVMTemplateArgument::ETypeCategory_SingleAnyValue, TArray<TRigVMTypeIndex>()).Reserve(256);
+	TypesPerCategory.Add(FRigVMTemplateArgument::ETypeCategory_ArrayAnyValue, TArray<TRigVMTypeIndex>()).Reserve(256);
+	TypesPerCategory.Add(FRigVMTemplateArgument::ETypeCategory_ArrayArrayAnyValue, TArray<TRigVMTypeIndex>()).Reserve(256);
+	TypesPerCategory.Add(FRigVMTemplateArgument::ETypeCategory_SingleSimpleValue, TArray<TRigVMTypeIndex>()).Reserve(8);
+	TypesPerCategory.Add(FRigVMTemplateArgument::ETypeCategory_ArraySimpleValue, TArray<TRigVMTypeIndex>()).Reserve(8);
+	TypesPerCategory.Add(FRigVMTemplateArgument::ETypeCategory_ArrayArraySimpleValue, TArray<TRigVMTypeIndex>()).Reserve(8);
+	TypesPerCategory.Add(FRigVMTemplateArgument::ETypeCategory_SingleMathStructValue, TArray<TRigVMTypeIndex>()).Reserve(GetMathTypes().Num());
+	TypesPerCategory.Add(FRigVMTemplateArgument::ETypeCategory_ArrayMathStructValue, TArray<TRigVMTypeIndex>()).Reserve(GetMathTypes().Num());
+	TypesPerCategory.Add(FRigVMTemplateArgument::ETypeCategory_ArrayArrayMathStructValue, TArray<TRigVMTypeIndex>()).Reserve(GetMathTypes().Num());
+	TypesPerCategory.Add(FRigVMTemplateArgument::ETypeCategory_SingleScriptStructValue, TArray<TRigVMTypeIndex>()).Reserve(128);
+	TypesPerCategory.Add(FRigVMTemplateArgument::ETypeCategory_ArrayScriptStructValue, TArray<TRigVMTypeIndex>()).Reserve(128);
+	TypesPerCategory.Add(FRigVMTemplateArgument::ETypeCategory_ArrayArrayScriptStructValue, TArray<TRigVMTypeIndex>()).Reserve(128);
+	TypesPerCategory.Add(FRigVMTemplateArgument::ETypeCategory_SingleEnumValue, TArray<TRigVMTypeIndex>()).Reserve(128);
+	TypesPerCategory.Add(FRigVMTemplateArgument::ETypeCategory_ArrayEnumValue, TArray<TRigVMTypeIndex>()).Reserve(128);
+	TypesPerCategory.Add(FRigVMTemplateArgument::ETypeCategory_ArrayArrayEnumValue, TArray<TRigVMTypeIndex>()).Reserve(128);
+	TypesPerCategory.Add(FRigVMTemplateArgument::ETypeCategory_SingleObjectValue, TArray<TRigVMTypeIndex>()).Reserve(128);
+	TypesPerCategory.Add(FRigVMTemplateArgument::ETypeCategory_ArrayObjectValue, TArray<TRigVMTypeIndex>()).Reserve(128);
+	TypesPerCategory.Add(FRigVMTemplateArgument::ETypeCategory_ArrayArrayObjectValue, TArray<TRigVMTypeIndex>()).Reserve(128);
 
 	ArgumentsPerCategory.Add(FRigVMTemplateArgument::ETypeCategory_Execute, TArray<TPair<int32,int32>>()).Reserve(8);
 	ArgumentsPerCategory.Add(FRigVMTemplateArgument::ETypeCategory_SingleAnyValue, TArray<TPair<int32,int32>>()).Reserve(64);
@@ -111,23 +124,24 @@ void FRigVMRegistry::InitializeIfNeeded()
 		FindOrAddType(FRigVMTemplateArgumentType(MathType));
 	}
 
+	Refresh();
+}
+
+void FRigVMRegistry::Refresh()
+{
 	// add all user defined structs
 	for (TObjectIterator<UScriptStruct> ScriptIt; ScriptIt; ++ScriptIt)
 	{
 		UScriptStruct* ScriptStruct = *ScriptIt;
-		if(!IsAllowedType(ScriptStruct))
+		if(IsAllowedType(ScriptStruct))
 		{
-			continue;
-		}
-
-		// if this is a C++ type - skip it
-		if(ScriptStruct->IsNative() && !ScriptStruct->IsChildOf(FRigVMExecuteContext::StaticStruct()))
-		{
-			continue;
+			// if this is a C++ type - skip it
+			if(!ScriptStruct->IsNative() || ScriptStruct->IsChildOf(FRigVMExecuteContext::StaticStruct()))
+			{
+				const FString CPPType = ScriptStruct->GetStructCPPName();
+				FindOrAddType(FRigVMTemplateArgumentType(ScriptStruct));
+			}
 		}
-
-		const FString CPPType = ScriptStruct->GetStructCPPName();
-		FindOrAddType(FRigVMTemplateArgumentType(ScriptStruct));
 	}
 
 	// add all user defined enums
@@ -148,15 +162,22 @@ void FRigVMRegistry::InitializeIfNeeded()
 		const FString CPPType = Enum->CppType.IsEmpty() ? Enum->GetName() : Enum->CppType;
 		FindOrAddType(FRigVMTemplateArgumentType(*CPPType, Enum));
 	}
+	
+	// register all dispatch factories
+	for (TObjectIterator<UScriptStruct> ScriptIt; ScriptIt; ++ScriptIt)
+	{
+		UScriptStruct* ScriptStruct = *ScriptIt;
+		if(ScriptStruct != FRigVMDispatchFactory::StaticStruct() &&
+			ScriptStruct->IsChildOf(FRigVMDispatchFactory::StaticStruct()))
+		{
+			RegisterFactory(ScriptStruct);
+		}
+	}
 }
 
-void FRigVMRegistry::Refresh()
-{
-}
-
-int32 FRigVMRegistry::FindOrAddType(const FRigVMTemplateArgumentType& InType)
+TRigVMTypeIndex FRigVMRegistry::FindOrAddType(const FRigVMTemplateArgumentType& InType)
 {
-	int32 Index = GetTypeIndex(InType);
+	TRigVMTypeIndex Index = GetTypeIndex(InType);
 	if(Index == INDEX_NONE)
 	{
 		FRigVMTemplateArgumentType ElementType;
@@ -185,6 +206,9 @@ int32 FRigVMRegistry::FindOrAddType(const FRigVMTemplateArgumentType& InType)
 			Info.ArrayTypeIndex = GetTypeIndex(ArrayType);
 			
 			Index = Types.Add(Info);
+#if UE_RIGVM_DEBUG_TYPEINDEX
+			Index.Name = Info.Type.CPPType;
+#endif
 		}
 
 		TypeToIndex.Add(InType, Index);
@@ -387,7 +411,7 @@ int32 FRigVMRegistry::FindOrAddType(const FRigVMTemplateArgumentType& InType)
 	return Index;
 }
 
-void FRigVMRegistry::RegisterTypeInCategory(FRigVMTemplateArgument::ETypeCategory InCategory, int32 InTypeIndex)
+void FRigVMRegistry::RegisterTypeInCategory(FRigVMTemplateArgument::ETypeCategory InCategory, TRigVMTypeIndex InTypeIndex)
 {
 	check(InCategory != FRigVMTemplateArgument::ETypeCategory_Invalid);
 
@@ -403,16 +427,16 @@ void FRigVMRegistry::RegisterTypeInCategory(FRigVMTemplateArgument::ETypeCategor
 	}
 }
 
-int32 FRigVMRegistry::GetTypeIndex(const FRigVMTemplateArgumentType& InType) const
+TRigVMTypeIndex FRigVMRegistry::GetTypeIndex(const FRigVMTemplateArgumentType& InType) const
 {
-	if(const int32* Index = TypeToIndex.Find(InType))
+	if(const TRigVMTypeIndex* Index = TypeToIndex.Find(InType))
 	{
 		return *Index;
 	}
 	return INDEX_NONE;
 }
 
-const FRigVMTemplateArgumentType& FRigVMRegistry::GetType(int32 InTypeIndex) const
+const FRigVMTemplateArgumentType& FRigVMRegistry::GetType(TRigVMTypeIndex InTypeIndex) const
 {
 	if(ensure(Types.IsValidIndex(InTypeIndex)))
 	{
@@ -434,7 +458,7 @@ const FRigVMTemplateArgumentType& FRigVMRegistry::FindTypeFromCPPType(const FStr
 	return EmptyType;
 }
 
-int32 FRigVMRegistry::GetTypeIndexFromCPPType(const FString& InCPPType) const
+TRigVMTypeIndex FRigVMRegistry::GetTypeIndexFromCPPType(const FString& InCPPType) const
 {
 	if(ensure(!InCPPType.IsEmpty()))
 	{
@@ -447,7 +471,7 @@ int32 FRigVMRegistry::GetTypeIndexFromCPPType(const FString& InCPPType) const
 	return INDEX_NONE;
 }
 
-bool FRigVMRegistry::IsArrayType(int32 InTypeIndex) const
+bool FRigVMRegistry::IsArrayType(TRigVMTypeIndex InTypeIndex) const
 {
 	if(ensure(Types.IsValidIndex(InTypeIndex)))
 	{
@@ -456,7 +480,7 @@ bool FRigVMRegistry::IsArrayType(int32 InTypeIndex) const
 	return false;
 }
 
-bool FRigVMRegistry::IsExecuteType(int32 InTypeIndex) const
+bool FRigVMRegistry::IsExecuteType(TRigVMTypeIndex InTypeIndex) const
 {
 	if(ensure(Types.IsValidIndex(InTypeIndex)))
 	{
@@ -465,7 +489,7 @@ bool FRigVMRegistry::IsExecuteType(int32 InTypeIndex) const
 	return false;
 }
 
-int32 FRigVMRegistry::GetArrayDimensionsForType(int32 InTypeIndex) const
+int32 FRigVMRegistry::GetArrayDimensionsForType(TRigVMTypeIndex InTypeIndex) const
 {
 	if(ensure(Types.IsValidIndex(InTypeIndex)))
 	{
@@ -478,13 +502,13 @@ int32 FRigVMRegistry::GetArrayDimensionsForType(int32 InTypeIndex) const
 	return 0;
 }
 
-bool FRigVMRegistry::IsWildCardType(int32 InTypeIndex) const
+bool FRigVMRegistry::IsWildCardType(TRigVMTypeIndex InTypeIndex) const
 {
 	return RigVMTypeUtils::TypeIndex::WildCard == InTypeIndex ||
 		RigVMTypeUtils::TypeIndex::WildCardArray == InTypeIndex;
 }
 
-bool FRigVMRegistry::CanMatchTypes(int32 InTypeIndexA, int32 InTypeIndexB, bool bAllowFloatingPointCasts) const
+bool FRigVMRegistry::CanMatchTypes(TRigVMTypeIndex InTypeIndexA, TRigVMTypeIndex InTypeIndexB, bool bAllowFloatingPointCasts) const
 {
 	if(!ensure(Types.IsValidIndex(InTypeIndexA)) || !ensure(Types.IsValidIndex(InTypeIndexB)))
 	{
@@ -515,40 +539,40 @@ bool FRigVMRegistry::CanMatchTypes(int32 InTypeIndexA, int32 InTypeIndexB, bool
 	return false;
 }
 
-const TArray<int32>& FRigVMRegistry::GetCompatibleTypes(int32 InTypeIndex) const
+const TArray<TRigVMTypeIndex>& FRigVMRegistry::GetCompatibleTypes(TRigVMTypeIndex InTypeIndex) const
 {
 	if(InTypeIndex == RigVMTypeUtils::TypeIndex::Float)
 	{
-		static const TArray<int32> CompatibleTypes = {RigVMTypeUtils::TypeIndex::Double};
+		static const TArray<TRigVMTypeIndex> CompatibleTypes = {RigVMTypeUtils::TypeIndex::Double};
 		return CompatibleTypes;
 	}
 	if(InTypeIndex == RigVMTypeUtils::TypeIndex::Double)
 	{
-		static const TArray<int32> CompatibleTypes = {RigVMTypeUtils::TypeIndex::Float};
+		static const TArray<TRigVMTypeIndex> CompatibleTypes = {RigVMTypeUtils::TypeIndex::Float};
 		return CompatibleTypes;
 	}
 	if(InTypeIndex == RigVMTypeUtils::TypeIndex::FloatArray)
 	{
-		static const TArray<int32> CompatibleTypes = {RigVMTypeUtils::TypeIndex::DoubleArray};
+		static const TArray<TRigVMTypeIndex> CompatibleTypes = {RigVMTypeUtils::TypeIndex::DoubleArray};
 		return CompatibleTypes;
 	}
 	if(InTypeIndex == RigVMTypeUtils::TypeIndex::DoubleArray)
 	{
-		static const TArray<int32> CompatibleTypes = {RigVMTypeUtils::TypeIndex::FloatArray};
+		static const TArray<TRigVMTypeIndex> CompatibleTypes = {RigVMTypeUtils::TypeIndex::FloatArray};
 		return CompatibleTypes;
 	}
 
-	static const TArray<int32> EmptyTypes;
+	static const TArray<TRigVMTypeIndex> EmptyTypes;
 	return EmptyTypes;
 }
 
-const TArray<int32>& FRigVMRegistry::GetTypesForCategory(FRigVMTemplateArgument::ETypeCategory InCategory)
+const TArray<TRigVMTypeIndex>& FRigVMRegistry::GetTypesForCategory(FRigVMTemplateArgument::ETypeCategory InCategory)
 {
 	check(InCategory != FRigVMTemplateArgument::ETypeCategory_Invalid);
 	return TypesPerCategory.FindChecked(InCategory);
 }
 
-int32 FRigVMRegistry::GetArrayTypeFromBaseTypeIndex(int32 InTypeIndex) const
+TRigVMTypeIndex FRigVMRegistry::GetArrayTypeFromBaseTypeIndex(TRigVMTypeIndex InTypeIndex) const
 {
 	if(ensure(Types.IsValidIndex(InTypeIndex)))
 	{
@@ -557,7 +581,7 @@ int32 FRigVMRegistry::GetArrayTypeFromBaseTypeIndex(int32 InTypeIndex) const
 	return INDEX_NONE;
 }
 
-int32 FRigVMRegistry::GetBaseTypeFromArrayTypeIndex(int32 InTypeIndex) const
+TRigVMTypeIndex FRigVMRegistry::GetBaseTypeFromArrayTypeIndex(TRigVMTypeIndex InTypeIndex) const
 {
 	if(ensure(Types.IsValidIndex(InTypeIndex)))
 	{
@@ -644,6 +668,10 @@ bool FRigVMRegistry::IsAllowedType(const UStruct* InStruct)
 	{
 		return false;
 	}
+	if(InStruct->IsChildOf(FRigVMDispatchFactory::StaticStruct()))
+	{
+		return false;
+	}
 	for (TFieldIterator<FProperty> It(InStruct); It; ++It)
 	{
 		if(!IsAllowedType(*It))
@@ -733,12 +761,88 @@ void FRigVMRegistry::Register(const TCHAR* InName, FRigVMFunctionPtr InFunctionP
 #endif
 }
 
+void FRigVMRegistry::RegisterFactory(UScriptStruct* InFactoryStruct)
+{
+	check(InFactoryStruct);
+	check(InFactoryStruct != FRigVMDispatchFactory::StaticStruct());
+	check(InFactoryStruct->IsChildOf(FRigVMDispatchFactory::StaticStruct()));
+
+	// ensure to register factories only once
+	const bool bFactoryAlreadyRegistered = Factories.ContainsByPredicate([InFactoryStruct](const FRigVMDispatchFactory* Factory)
+	{
+		return Factory->GetScriptStruct() == InFactoryStruct;
+	});
+	if(bFactoryAlreadyRegistered)
+	{
+		return;
+	}
+
+#if WITH_EDITOR
+	if(InFactoryStruct->HasMetaData(TEXT("Abstract")))
+	{
+		return;
+	}
+#endif
+	
+	FRigVMDispatchFactory* Factory = (FRigVMDispatchFactory*)FMemory::Malloc(InFactoryStruct->GetStructureSize());
+	InFactoryStruct->InitializeStruct(Factory, 1);
+	Factory->FactoryScriptStruct = InFactoryStruct;
+	Factories.Add(Factory);
+	Factory->RegisterDependencyTypes();
+}
+
 const FRigVMFunction* FRigVMRegistry::FindFunction(const TCHAR* InName) const
 {
 	if(const int32* FunctionIndexPtr = FunctionNameToIndex.Find(InName))
 	{
 		return &Functions[*FunctionIndexPtr];
 	}
+	
+	if(Factories.Num() > 0)
+	{
+		static bool IsDispatchingFunction = false;
+		if(IsDispatchingFunction)
+		{
+			return nullptr;
+		}
+		
+		const FString NameString(InName);
+		FString FactoryName, ArgumentsString;
+		if(NameString.Split(TEXT("::"), &FactoryName, &ArgumentsString))
+		{
+			// if the factory has never been registered - we should try to look it up
+			if(FindDispatchFactory(*FactoryName) == nullptr)
+			{
+				if(FactoryName.StartsWith(FRigVMDispatchFactory::DispatchPrefix))
+				{
+					const FString ScriptStructName = FactoryName.Mid(FRigVMDispatchFactory::DispatchPrefix.Len());
+					if(UScriptStruct* FactoryStruct = FindFirstObject<UScriptStruct>(*ScriptStructName, EFindFirstObjectOptions::NativeFirst | EFindFirstObjectOptions::EnsureIfAmbiguous))
+					{
+						FRigVMRegistry* MutableThis = (FRigVMRegistry*)this;
+						MutableThis->RegisterFactory(FactoryStruct);
+					}
+				}
+			}
+
+			if(const FRigVMDispatchFactory* Factory = FindDispatchFactory(*FactoryName))
+			{
+				if(const FRigVMTemplate* Template = Factory->GetTemplate())
+				{
+					const FRigVMTemplateTypeMap ArgumentTypes = Template->GetArgumentTypesFromString(ArgumentsString);
+					if(ArgumentTypes.Num() == Template->NumArguments())
+					{
+						const int32 PermutationIndex = Template->FindPermutation(ArgumentTypes);
+						if(PermutationIndex != INDEX_NONE)
+						{
+							TGuardValue<bool> ReEntryGuard(IsDispatchingFunction, true);
+							return ((FRigVMTemplate*)Template)->GetOrCreatePermutation(PermutationIndex);
+						}
+					}
+				}
+			}
+		}
+	}
+
 	return nullptr;
 }
 
@@ -768,6 +872,26 @@ const FRigVMTemplate* FRigVMRegistry::FindTemplate(const FName& InNotation) cons
 		return &Templates[*TemplateIndexPtr];
 	}
 
+	if(Factories.Num() > 0)
+	{
+		static bool IsDispatchingFunction = false;
+		if(IsDispatchingFunction)
+		{
+			return nullptr;
+		}
+		
+		const FString NotationString(InNotation.ToString());
+		FString FactoryName, ArgumentsString;
+		if(NotationString.Split(TEXT("("), &FactoryName, &ArgumentsString))
+		{
+			if(const FRigVMDispatchFactory* Factory = FindDispatchFactory(*FactoryName))
+			{
+				TGuardValue<bool> ReEntryGuard(IsDispatchingFunction, true);
+				return Factory->GetTemplate();
+			}
+		}
+	}
+
 	return nullptr;
 }
 
@@ -853,7 +977,7 @@ const FRigVMTemplate* FRigVMRegistry::GetOrAddTemplateFromArguments(const FName&
 
 	const int32 Index = Templates.AddElement(Template);
 	Templates[Index].Index = Index;
-	Templates[Index].OnNewArgumentType() = InDelegates.NewArgumentTypeDelegate;
+	Templates[Index].Delegates = InDelegates;
 	TemplateNotationToIndex.Add(Template.GetNotation(), Index);
 
 	for(int32 ArgumentIndex=0; ArgumentIndex < Templates[Index].Arguments.Num(); ArgumentIndex++)
@@ -866,3 +990,21 @@ const FRigVMTemplate* FRigVMRegistry::GetOrAddTemplateFromArguments(const FName&
 	
 	return &Templates[Index];
 }
+
+FRigVMDispatchFactory* FRigVMRegistry::FindDispatchFactory(const FName& InFactoryName) const 
+{
+	FRigVMDispatchFactory* const* FactoryPtr = Factories.FindByPredicate([InFactoryName](const FRigVMDispatchFactory* Factory) -> bool
+	{
+		return Factory->GetFactoryName() == InFactoryName;
+	});
+	if(FactoryPtr)
+	{
+		return *FactoryPtr;
+	}
+	return nullptr;
+}
+
+const TArray<FRigVMDispatchFactory*>& FRigVMRegistry::GetFactories() const
+{
+	return Factories;
+}
diff --git a/Engine/Source/Runtime/RigVM/Private/RigVMCore/RigVMTemplate.cpp b/Engine/Source/Runtime/RigVM/Private/RigVMCore/RigVMTemplate.cpp
--- a/Engine/Source/Runtime/RigVM/Private/RigVMCore/RigVMTemplate.cpp
+++ b/Engine/Source/Runtime/RigVM/Private/RigVMCore/RigVMTemplate.cpp
@@ -49,13 +49,13 @@ FRigVMTemplateArgument::FRigVMTemplateArgument(FProperty* InProperty)
 	}
 	Type.CPPType = *RigVMTypeUtils::PostProcessCPPType(Type.CPPType.ToString(), Type.CPPTypeObject);
 
-	const int32 TypeIndex = FRigVMRegistry::Get().FindOrAddType(Type); 
+	const TRigVMTypeIndex TypeIndex = FRigVMRegistry::Get().FindOrAddType(Type); 
 
 	TypeIndices.Add(TypeIndex);
 	TypeToPermutations.Add(TypeIndex, {0});
 }
 
-FRigVMTemplateArgument::FRigVMTemplateArgument(const FName& InName, ERigVMPinDirection InDirection, int32 InTypeIndex)
+FRigVMTemplateArgument::FRigVMTemplateArgument(const FName& InName, ERigVMPinDirection InDirection, TRigVMTypeIndex InTypeIndex)
 : Index(INDEX_NONE)
 , Name(InName)
 , Direction(InDirection)
@@ -64,7 +64,7 @@ FRigVMTemplateArgument::FRigVMTemplateArgument(const FName& InName, ERigVMPinDir
 	UpdateTypeToPermutations();
 }
 
-FRigVMTemplateArgument::FRigVMTemplateArgument(const FName& InName, ERigVMPinDirection InDirection, const TArray<int32>& InTypeIndices)
+FRigVMTemplateArgument::FRigVMTemplateArgument(const FName& InName, ERigVMPinDirection InDirection, const TArray<TRigVMTypeIndex>& InTypeIndices)
 : Index(INDEX_NONE)
 , Name(InName)
 , Direction(InDirection)
@@ -104,7 +104,7 @@ void FRigVMTemplateArgument::UpdateTypeToPermutations()
 	}
 }
 
-bool FRigVMTemplateArgument::SupportsTypeIndex(int32 InTypeIndex, int32* OutTypeIndex) const
+bool FRigVMTemplateArgument::SupportsTypeIndex(TRigVMTypeIndex InTypeIndex, TRigVMTypeIndex* OutTypeIndex) const
 {
 	if (const TArray<int32>* Permutations = TypeToPermutations.Find(InTypeIndex))
 	{
@@ -116,8 +116,8 @@ bool FRigVMTemplateArgument::SupportsTypeIndex(int32 InTypeIndex, int32* OutType
 	}
 
 	// Try to find compatible type
-	const TArray<int32>& CompatibleTypes = FRigVMRegistry::Get().GetCompatibleTypes(InTypeIndex);
-	for (const int32& CompatibleTypeIndex : CompatibleTypes)
+	const TArray<TRigVMTypeIndex>& CompatibleTypes = FRigVMRegistry::Get().GetCompatibleTypes(InTypeIndex);
+	for (const TRigVMTypeIndex& CompatibleTypeIndex : CompatibleTypes)
 	{
 		if (const TArray<int32>* Permutations = TypeToPermutations.Find(CompatibleTypeIndex))
 		{
@@ -143,10 +143,10 @@ bool FRigVMTemplateArgument::IsSingleton(const TArray<int32>& InPermutationIndic
 		return false;
 	}
 
-	const int32 TypeIndexToCheck = TypeIndices[InPermutationIndices[0]];
+	const TRigVMTypeIndex TRigVMTypeIndex = TypeIndices[InPermutationIndices[0]];
 	for (int32 PermutationIndex = 1; PermutationIndex < InPermutationIndices.Num(); PermutationIndex++)
 	{
-		if (TypeIndices[InPermutationIndices[PermutationIndex]] != TypeIndexToCheck)
+		if (TypeIndices[InPermutationIndices[PermutationIndex]] != TRigVMTypeIndex)
 		{
 			return false;
 		}
@@ -194,17 +194,17 @@ FRigVMTemplateArgument::EArrayType FRigVMTemplateArgument::GetArrayType() const
 	return EArrayType_Invalid;
 }
 
-const TArray<int32>& FRigVMTemplateArgument::GetTypeIndices() const
+const TArray<TRigVMTypeIndex>& FRigVMTemplateArgument::GetTypeIndices() const
 {
 	return TypeIndices;
 }
 
-TArray<int32> FRigVMTemplateArgument::GetSupportedTypeIndices(const TArray<int32>& InPermutationIndices) const
+TArray<TRigVMTypeIndex> FRigVMTemplateArgument::GetSupportedTypeIndices(const TArray<int32>& InPermutationIndices) const
 {
-	TArray<int32> SupportedTypes;
+	TArray<TRigVMTypeIndex> SupportedTypes;
 	if(InPermutationIndices.IsEmpty())
 	{
-		for (const int32& TypeIndex : TypeIndices)
+		for (const TRigVMTypeIndex& TypeIndex : TypeIndices)
 		{
 			SupportedTypes.AddUnique(TypeIndex);
 		}
@@ -226,7 +226,7 @@ TArray<FString> FRigVMTemplateArgument::GetSupportedTypeStrings(const TArray<int
 	TArray<FString> SupportedTypes;
 	if(InPermutationIndices.IsEmpty())
 	{
-		for (const int32& TypeIndex : TypeIndices)
+		for (const TRigVMTypeIndex& TypeIndex : TypeIndices)
 		{
 			const FString TypeString = FRigVMRegistry::Get().GetType(TypeIndex).CPPType.ToString();
 			SupportedTypes.AddUnique(TypeString);
@@ -236,7 +236,7 @@ TArray<FString> FRigVMTemplateArgument::GetSupportedTypeStrings(const TArray<int
 	{
 		for(const int32 PermutationIndex : InPermutationIndices)
 		{
-			const int32 TypeIndex = TypeIndices[PermutationIndex];
+			const TRigVMTypeIndex TypeIndex = TypeIndices[PermutationIndex];
 			const FString TypeString = FRigVMRegistry::Get().GetType(TypeIndex).CPPType.ToString();
 			SupportedTypes.AddUnique(TypeString);
 		}
@@ -271,7 +271,7 @@ FRigVMTemplate::FRigVMTemplate(UScriptStruct* InStruct, const FString& InTemplat
 			FRigVMTemplateArgument Argument(*It);
 			Argument.Index = Arguments.Num();
 
-			if(IsValidArgumentForTemplate(Argument))
+			if(IsValidArgumentForTemplate(Argument) && Argument.GetDirection() != ERigVMPinDirection::Hidden)
 			{
 				Arguments.Add(Argument);
 			}
@@ -284,7 +284,10 @@ FRigVMTemplate::FRigVMTemplate(UScriptStruct* InStruct, const FString& InTemplat
 	{
 		if(const FRigVMTemplateArgument* Argument = FindArgument(It->GetFName()))
 		{
-			ArgumentNotations.Add(GetArgumentNotation(*Argument));
+			if(Argument->GetDirection() != ERigVMPinDirection::Hidden)
+			{
+				ArgumentNotations.Add(GetArgumentNotation(*Argument));
+			}
 		}
 	}
 
@@ -309,7 +312,10 @@ FRigVMTemplate::FRigVMTemplate(const FName& InTemplateName, const TArray<FRigVMT
 		if(IsValidArgumentForTemplate(Argument))
 		{
 			Arguments.Add(Argument);
-			ArgumentNotations.Add(GetArgumentNotation(Argument));
+			if(Argument.GetDirection() != ERigVMPinDirection::Hidden)
+			{
+				ArgumentNotations.Add(GetArgumentNotation(Argument));
+			}
 		}
 	}
 
@@ -321,12 +327,40 @@ FRigVMTemplate::FRigVMTemplate(const FName& InTemplateName, const TArray<FRigVMT
 	}
 }
 
+FLinearColor FRigVMTemplate::GetColorFromMetadata(FString InMetadata)
+{
+	FLinearColor Color = FLinearColor::Black;
+
+	FString Metadata = InMetadata;
+	Metadata.TrimStartAndEndInline();
+	static const FString SplitString(TEXT(" "));
+	FString Red, Green, Blue, GreenAndBlue;
+	if (Metadata.Split(SplitString, &Red, &GreenAndBlue))
+	{
+		Red.TrimEndInline();
+		GreenAndBlue.TrimStartInline();
+		if (GreenAndBlue.Split(SplitString, &Green, &Blue))
+		{
+			Green.TrimEndInline();
+			Blue.TrimStartInline();
+
+			const float RedValue = FCString::Atof(*Red);
+			const float GreenValue = FCString::Atof(*Green);
+			const float BlueValue = FCString::Atof(*Blue);
+			Color = FLinearColor(RedValue, GreenValue, BlueValue);
+		}
+	}
+
+	return Color;
+}
+
 bool FRigVMTemplate::IsValidArgumentForTemplate(const FRigVMTemplateArgument& InArgument)
 {
 	static const TArray<ERigVMPinDirection> ValidDirections = {
 		ERigVMPinDirection::Input,
 		ERigVMPinDirection::Output,
 		ERigVMPinDirection::IO,
+		ERigVMPinDirection::Hidden,
 		ERigVMPinDirection::Visible
 	};
 
@@ -400,6 +434,39 @@ TArray<UStruct*> FRigVMTemplate::GetSuperStructs(UStruct* InStruct, bool bInclud
 	return SuperStructs;
 }
 
+FRigVMTemplate::FTypeMap FRigVMTemplate::GetArgumentTypesFromString(const FString& InTypeString) const
+{
+	FTypeMap Types;
+	if(!InTypeString.IsEmpty())
+	{
+		const FRigVMRegistry& Registry = FRigVMRegistry::Get();
+
+		FString Left, Right = InTypeString;
+		while(!Right.IsEmpty())
+		{
+			if(!Right.Split(TEXT(","), &Left, &Right))
+			{
+				Left = Right;
+				Right.Reset();
+			}
+
+			FString ArgumentName, TypeName;
+			if(Left.Split(TEXT(":"), &ArgumentName, &TypeName))
+			{
+				if(const FRigVMTemplateArgument* Argument = FindArgument(*ArgumentName))
+				{
+					const TRigVMTypeIndex TypeIndex = Registry.GetTypeIndexFromCPPType(TypeName);
+					if(TypeIndex != INDEX_NONE)
+					{
+						Types.Add(Argument->GetName(), TypeIndex);
+					}
+				}
+			}
+		}
+	}
+	return Types;
+}
+
 bool FRigVMTemplate::IsValid() const
 {
 	return !Notation.IsNone();
@@ -428,37 +495,15 @@ FName FRigVMTemplate::GetName() const
 
 FLinearColor FRigVMTemplate::GetColor(const TArray<int32>& InPermutationIndices) const
 {
-	bool bFirstColorFound = false;
-	FLinearColor ResolvedColor = FLinearColor::White;
-
-	auto GetColorFromMetadata = [] (FString InMetadata) -> FLinearColor
+	if(const FRigVMDispatchFactory* Factory = GetDispatchFactory())
 	{
-		FLinearColor Color = FLinearColor::Black;
-
-		FString Metadata = InMetadata;
-		Metadata.TrimStartAndEndInline();
-		FString SplitString(TEXT(" "));
-		FString Red, Green, Blue, GreenAndBlue;
-		if (Metadata.Split(SplitString, &Red, &GreenAndBlue))
-		{
-			Red.TrimEndInline();
-			GreenAndBlue.TrimStartInline();
-			if (GreenAndBlue.Split(SplitString, &Green, &Blue))
-			{
-				Green.TrimEndInline();
-				Blue.TrimStartInline();
-
-				float RedValue = FCString::Atof(*Red);
-				float GreenValue = FCString::Atof(*Green);
-				float BlueValue = FCString::Atof(*Blue);
-				Color = FLinearColor(RedValue, GreenValue, BlueValue);
-			}
-		}
+		return Factory->GetNodeColor();
+	}
 
-		return Color;
-	};
+	bool bFirstColorFound = false;
+	FLinearColor ResolvedColor = FLinearColor::White;
 
-	auto VisitPermutation = [&bFirstColorFound, &ResolvedColor, GetColorFromMetadata, this](int32 InPermutationIndex) -> bool
+	auto VisitPermutation = [&bFirstColorFound, &ResolvedColor, this](int32 InPermutationIndex) -> bool
 	{
 		static const FName NodeColorName = TEXT("NodeColor");
 		FString NodeColorMetadata;
@@ -516,6 +561,16 @@ FLinearColor FRigVMTemplate::GetColor(const TArray<int32>& InPermutationIndices)
 
 FText FRigVMTemplate::GetTooltipText(const TArray<int32>& InPermutationIndices) const
 {
+	if(const FRigVMDispatchFactory* Factory = GetDispatchFactory())
+	{
+		FTypeMap Types;
+		if(InPermutationIndices.Num() == 1)
+		{
+			Types = GetTypesForPermutation(InPermutationIndices[0]);
+		}
+		return Factory->GetNodeTooltip(Types);
+	}
+
 	FText ResolvedTooltipText;
 
 	auto VisitPermutation = [&ResolvedTooltipText, this](int32 InPermutationIndex) -> bool
@@ -575,6 +630,11 @@ FText FRigVMTemplate::GetTooltipText(const TArray<int32>& InPermutationIndices)
 
 FText FRigVMTemplate::GetDisplayNameForArgument(const FName& InArgumentName, const TArray<int32>& InPermutationIndices) const
 {
+	if(const FRigVMDispatchFactory* Factory = GetDispatchFactory())
+	{
+		return Factory->GetDisplayNameForArgument(InArgumentName);
+	}
+
 	if(const FRigVMTemplateArgument* Argument = FindArgument(InArgumentName))
 	{
 		FText ResolvedDisplayName;
@@ -635,6 +695,11 @@ FText FRigVMTemplate::GetDisplayNameForArgument(const FName& InArgumentName, con
 
 FString FRigVMTemplate::GetArgumentMetaData(const FName& InArgumentName, const FName& InMetaDataKey, const TArray<int32>& InPermutationIndices) const
 {
+	if(const FRigVMDispatchFactory* Factory = GetDispatchFactory())
+	{
+		return Factory->GetArgumentMetaData(InArgumentName, InMetaDataKey);
+	}
+
 	if(const FRigVMTemplateArgument* Argument = FindArgument(InArgumentName))
 	{
 		FString ResolvedMetaData;
@@ -733,9 +798,9 @@ bool FRigVMTemplate::Merge(const FRigVMTemplate& InOther)
 		{
 			// find the previously defined permutation.
 			UE_LOG(LogRigVM, Display, TEXT("RigVMFunction '%s' cannot be merged into the '%s' template. It collides with '%s'."),
-				InOther.GetPermutation(0)->Name,
+				*InOther.GetPermutation(0)->Name,
 				*GetNotation().ToString(),
-				GetPermutation(PermutationIndex)->Name);
+				*GetPermutation(PermutationIndex)->Name);
 			return false;
 		}
 	}
@@ -753,7 +818,7 @@ bool FRigVMTemplate::Merge(const FRigVMTemplate& InOther)
 
 		// Add Other argument information into the TypeToPermutations map
 		{
-			const int32 OtherTypeIndex = InOther.Arguments[ArgumentIndex].TypeIndices[0];
+			const TRigVMTypeIndex OtherTypeIndex = InOther.Arguments[ArgumentIndex].TypeIndices[0];
 			const int32 NewPermutationIndex = NewArgs[ArgumentIndex].TypeIndices.Num();
 			if (TArray<int32>* ArgTypePermutations = NewArgs[ArgumentIndex].TypeToPermutations.Find(OtherTypeIndex))
 			{
@@ -781,7 +846,7 @@ const FRigVMTemplateArgument* FRigVMTemplate::FindArgument(const FName& InArgume
 	});
 }
 
-bool FRigVMTemplate::ArgumentSupportsTypeIndex(const FName& InArgumentName, int32 InTypeIndex, int32* OutTypeIndex) const
+bool FRigVMTemplate::ArgumentSupportsTypeIndex(const FName& InArgumentName, TRigVMTypeIndex InTypeIndex, TRigVMTypeIndex* OutTypeIndex) const
 {
 	if (const FRigVMTemplateArgument* Argument = FindArgument(InArgumentName))
 	{
@@ -801,6 +866,57 @@ const FRigVMFunction* FRigVMTemplate::GetPermutation(int32 InIndex) const
 	return nullptr;
 }
 
+const FRigVMFunction* FRigVMTemplate::GetOrCreatePermutation(int32 InIndex)
+{
+	if(const FRigVMFunction* Function = GetPermutation(InIndex))
+	{
+		return Function;
+	}
+	
+	if(Permutations[InIndex] == INDEX_NONE && UsesDispatch())
+	{
+		FRigVMRegistry& Registry = FRigVMRegistry::Get();
+		
+		FTypeMap Types;
+		for(const FRigVMTemplateArgument& Argument : Arguments)
+		{
+			Types.Add(Argument.GetName(), Argument.TypeIndices[InIndex]);
+		}
+		
+		if(const FRigVMFunctionPtr DispatchFunction = Delegates.RequestDispatchFunctionDelegate.Execute(this, Types))
+		{
+			FRigVMDispatchFactory* Factory = Delegates.GetDispatchFactoryDelegate.Execute();
+			check(Factory);
+
+			TArray<FRigVMFunctionArgument> FunctionArguments;
+			for(const FRigVMTemplateArgument& Argument : Arguments)
+			{
+				const FRigVMTemplateArgumentType& Type = Registry.GetType(Argument.TypeIndices[InIndex]);
+				FunctionArguments.Add(FRigVMFunctionArgument(Argument.Name.ToString(), Type.CPPType.ToString()));
+			}
+			
+			static constexpr TCHAR Format[] = TEXT("%s::%s");
+			const FString PermutationName = Factory->GetPermutationName(Types);
+			const int32 FunctionIndex = Permutations[InIndex] = Registry.Functions.Num();
+			
+			Registry.Functions.AddElement(
+				FRigVMFunction(
+					PermutationName,
+					DispatchFunction,
+					Factory,
+					FunctionIndex,
+					FunctionArguments
+				)
+			);
+			Registry.FunctionNameToIndex.Add(*PermutationName, FunctionIndex);
+
+			return &Registry.Functions[FunctionIndex];
+		}
+	}
+	
+	return nullptr;
+}
+
 bool FRigVMTemplate::ContainsPermutation(const FRigVMFunction* InPermutation) const
 {
 	return FindPermutation(InPermutation) != INDEX_NONE;
@@ -812,6 +928,17 @@ int32 FRigVMTemplate::FindPermutation(const FRigVMFunction* InPermutation) const
 	return Permutations.Find(InPermutation->Index);
 }
 
+int32 FRigVMTemplate::FindPermutation(const FTypeMap& InTypes) const
+{
+	FTypeMap Types = InTypes;
+	int32 PermutationIndex = INDEX_NONE;
+	if(FullyResolve(Types, PermutationIndex))
+	{
+		return PermutationIndex;
+	}
+	return INDEX_NONE;
+}
+
 bool FRigVMTemplate::FullyResolve(FRigVMTemplate::FTypeMap& InOutTypes, int32& OutPermutationIndex) const
 {
 	TArray<int32> PermutationIndices;
@@ -847,7 +974,7 @@ bool FRigVMTemplate::Resolve(FTypeMap& InOutTypes, TArray<int32>& OutPermutation
 			InOutTypes.Add(Argument.Name, Argument.TypeIndices[0]);
 			continue;
 		}
-		else if (const int32* InputType = InputTypes.Find(Argument.Name))
+		else if (const TRigVMTypeIndex* InputType = InputTypes.Find(Argument.Name))
 		{
 			int32 MatchedType = *InputType;
 			bool bFoundMatch = false;
@@ -906,7 +1033,7 @@ bool FRigVMTemplate::Resolve(FTypeMap& InOutTypes, TArray<int32>& OutPermutation
 		{
 			InOutTypes.Add(Argument.Name, RigVMTypeUtils::TypeIndex::WildCard);
 
-			if(const int32* InputType = InputTypes.Find(Argument.Name))
+			if(const TRigVMTypeIndex* InputType = InputTypes.Find(Argument.Name))
 			{
 				if(Registry.IsArrayType(*InputType))
 				{
@@ -946,7 +1073,7 @@ bool FRigVMTemplate::Resolve(FTypeMap& InOutTypes, TArray<int32>& OutPermutation
 	return !OutPermutationIndices.IsEmpty();
 }
 
-bool FRigVMTemplate::ResolveArgument(const FName& InArgumentName, const int32 InTypeIndex,
+bool FRigVMTemplate::ResolveArgument(const FName& InArgumentName, const TRigVMTypeIndex InTypeIndex,
 	FTypeMap& InOutTypes) const
 {
 	const FRigVMRegistry& Registry = FRigVMRegistry::Get();
@@ -976,7 +1103,7 @@ bool FRigVMTemplate::ResolveArgument(const FName& InArgumentName, const int32 In
 	if(Resolve(ResolvedTypes, PermutationIndices, true))
 	{
 		// let's see if the input argument resolved into the expected type
-		const int32 ResolvedInputType = ResolvedTypes.FindChecked(InArgumentName);
+		const TRigVMTypeIndex ResolvedInputType = ResolvedTypes.FindChecked(InArgumentName);
 		if(!Registry.CanMatchTypes(ResolvedInputType, InTypeIndex, true))
 		{
 			return false;
@@ -1034,10 +1161,34 @@ bool FRigVMTemplate::ResolveArgument(const FName& InArgumentName, const int32 In
 	return false;
 }
 
+FRigVMTemplateTypeMap FRigVMTemplate::GetTypesForPermutation(const int32 InPermutationIndex) const
+{
+	FTypeMap TypeMap;
+	for (int32 ArgIndex = 0; ArgIndex < NumArguments(); ++ArgIndex)
+	{
+		const FRigVMTemplateArgument* Argument = GetArgument(ArgIndex);
+		if (Argument->GetTypeIndices().Num() > InPermutationIndex)
+		{
+			TypeMap.Add(Argument->GetName(), Argument->GetTypeIndices()[InPermutationIndex]);
+		}
+		else
+		{
+			TypeMap.Reset();
+			return TypeMap;
+		}
+	}
+	return TypeMap;
+}
+
 #if WITH_EDITOR
 
 FString FRigVMTemplate::GetCategory() const
 {
+	if(const FRigVMDispatchFactory* Factory = GetDispatchFactory())
+	{
+		return Factory->GetCategory();
+	}
+	
 	FString Category;
 	GetPermutation(0)->Struct->GetStringMetaDataHierarchical(FRigVMStruct::CategoryMetaName, &Category);
 
@@ -1072,6 +1223,11 @@ FString FRigVMTemplate::GetCategory() const
 
 FString FRigVMTemplate::GetKeywords() const
 {
+	if(const FRigVMDispatchFactory* Factory = GetDispatchFactory())
+	{
+		return Factory->GetKeywords();
+	}
+
 	TArray<FString> KeywordsMetadata;
 	KeywordsMetadata.Add(GetName().ToString());
 
@@ -1095,7 +1251,7 @@ FString FRigVMTemplate::GetKeywords() const
 
 #endif
 
-bool FRigVMTemplate::AddTypeForArgument(const FName& InArgumentName, int32 InTypeIndex)
+bool FRigVMTemplate::AddTypeForArgument(const FName& InArgumentName, TRigVMTypeIndex InTypeIndex)
 {
 	if(OnNewArgumentType().IsBound())
 	{
@@ -1104,7 +1260,7 @@ bool FRigVMTemplate::AddTypeForArgument(const FName& InArgumentName, int32 InTyp
 		{
 			for(FRigVMTemplateArgument& Argument : Arguments)
 			{
-				const int32* TypeIndex = Types.Find(Argument.Name);
+				const TRigVMTypeIndex* TypeIndex = Types.Find(Argument.Name);
 				if(TypeIndex == nullptr)
 				{
 					return false;
@@ -1116,7 +1272,7 @@ bool FRigVMTemplate::AddTypeForArgument(const FName& InArgumentName, int32 InTyp
 			}
 			for(FRigVMTemplateArgument& Argument : Arguments)
 			{
-				const int32 TypeIndex = Types.FindChecked(Argument.Name);
+				const TRigVMTypeIndex TypeIndex = Types.FindChecked(Argument.Name);
 				Argument.TypeIndices.Add(TypeIndex);
 				Argument.TypeToPermutations.FindOrAdd(TypeIndex).Add(Permutations.Num());
 			}
diff --git a/Engine/Source/Runtime/RigVM/Public/RigVMCore/RigVM.h b/Engine/Source/Runtime/RigVM/Public/RigVMCore/RigVM.h
--- a/Engine/Source/Runtime/RigVM/Public/RigVMCore/RigVM.h
+++ b/Engine/Source/Runtime/RigVM/Public/RigVMCore/RigVM.h
@@ -194,6 +194,7 @@ public:
 	// Execute instructions can then refer to the function by index.
 	UFUNCTION()
 	virtual int32 AddRigVMFunction(UScriptStruct* InRigVMStruct, const FName& InMethodName);
+	virtual int32 AddRigVMFunction(const FString& InFunctionName);
 
 	// Returns the name of a function given its index
 	UFUNCTION()
@@ -679,6 +680,11 @@ private:
 	FORCEINLINE TArray<FRigVMFunctionPtr>& GetFunctions() { return *FunctionsPtr; }
 	FORCEINLINE const TArray<FRigVMFunctionPtr>& GetFunctions() const { return *FunctionsPtr; }
 
+	TArray<const FRigVMDispatchFactory*> FactoriesStorage;
+	TArray<const FRigVMDispatchFactory*>* FactoriesPtr;
+	FORCEINLINE TArray<const FRigVMDispatchFactory*>& GetFactories() { return *FactoriesPtr; }
+	FORCEINLINE const TArray<const FRigVMDispatchFactory*>& GetFactories() const { return *FactoriesPtr; }
+
 	UPROPERTY()
 	TArray<FRigVMParameter> Parameters;
 
diff --git a/Engine/Source/Runtime/RigVM/Public/RigVMCore/RigVMDispatchFactory.h b/Engine/Source/Runtime/RigVM/Public/RigVMCore/RigVMDispatchFactory.h
new file mode 100644
--- /dev/null
+++ b/Engine/Source/Runtime/RigVM/Public/RigVMCore/RigVMDispatchFactory.h
@@ -0,0 +1,106 @@
+// Copyright Epic Games, Inc. All Rights Reserved.
+
+#pragma once
+
+#include "CoreMinimal.h"
+#include "RigVMCore/RigVMStructUpgradeInfo.h"
+#include "RigVMCore/RigVMTemplate.h"
+#include "RigVMDispatchFactory.generated.h"
+
+/**
+ * A factory to generate a template and its dispatch functions
+ */
+USTRUCT()
+struct RIGVM_API FRigVMDispatchFactory
+{
+	GENERATED_BODY()
+
+public:
+
+	FORCEINLINE FRigVMDispatchFactory()
+		: FactoryScriptStruct(nullptr)
+		, CachedTemplate(nullptr)
+	{}
+	
+	FORCEINLINE virtual ~FRigVMDispatchFactory() {}
+
+	// returns the name of the factory template
+	FName GetFactoryName() const;
+
+	// returns the struct for this dispatch factory
+	FORCEINLINE UScriptStruct* GetScriptStruct() const { return FactoryScriptStruct; };
+
+#if WITH_EDITOR
+	
+	// returns the title of the node for a given type set.
+	FORCEINLINE virtual FString GetNodeTitle(const FRigVMTemplateTypeMap& InTypes) const { return GetScriptStruct()->GetDisplayNameText().ToString(); }
+
+	// returns the color of the node for a given type set.
+	virtual FLinearColor GetNodeColor() const;
+
+	// returns the tooltip for the node
+	virtual FText GetNodeTooltip(const FRigVMTemplateTypeMap& InTypes) const;
+
+	// returns the default value for an argument
+	FORCEINLINE virtual FString GetArgumentDefaultValue(const FName& InArgumentName, TRigVMTypeIndex InTypeIndex) const { return FString(); }
+
+	// returns the tooltip for an argument
+	FORCEINLINE virtual FText GetArgumentTooltip(const FName& InArgumentName, TRigVMTypeIndex InTypeIndex) const { return FText(); }
+
+	// returns the names of the input aggregate arguments
+	FORCEINLINE virtual TArray<FName> GetAggregateInputArguments() const { return TArray<FName>(); }
+
+	// returns the names of the output aggregate arguments
+	FORCEINLINE virtual TArray<FName> GetAggregateOutputArguments() const { return TArray<FName>(); }
+
+	// returns the next name to be used for an aggregate pin
+	virtual FName GetNextAggregateName(const FName& InLastAggregatePinName) const;
+
+	// Returns the display name text for an argument 
+	FORCEINLINE virtual FText GetDisplayNameForArgument(const FName& InArgumentName) const { return FText::FromName(InArgumentName); }
+
+	// Returns meta data on the property of the permutations 
+	FORCEINLINE virtual FString GetArgumentMetaData(const FName& InArgumentName, const FName& InMetaDataKey) const { return FString(); }
+
+	// Returns the category this factory is under
+	virtual FString GetCategory() const;
+
+	// Returns the keywords used for looking up this factory
+	virtual FString GetKeywords() const;
+
+#endif
+
+	// registered needed types during registration of the factory
+	FORCEINLINE virtual void RegisterDependencyTypes() const {}
+
+	// returns opaque arguments expected to be passed to this factory in FRigVMExtendedExecuteContext::OpaqueArguments
+	FORCEINLINE virtual TArray<TPair<FName,FString>> GetOpaqueArguments() const { return TArray<TPair<FName,FString>>(); }
+
+	// returns the arguments of the template
+	FORCEINLINE virtual TArray<FRigVMTemplateArgument> GetArguments() const { return TArray<FRigVMTemplateArgument>(); }
+
+	// returns the delegate to react to new types being added to an argument.
+	// this happens if types are being loaded later after this factory has already been deployed
+	FORCEINLINE virtual FRigVMTemplateTypeMap OnNewArgumentType(const FName& InArgumentName, TRigVMTypeIndex InTypeIndex) const { return FRigVMTemplateTypeMap(); }
+
+	// returns the upgrade info to use for this factory
+	FORCEINLINE virtual FRigVMStructUpgradeInfo GetUpgradeInfo(const FRigVMTemplateTypeMap& InTypes) const { return FRigVMStructUpgradeInfo(); }
+
+	// returns the dispatch function for a given type set
+	FRigVMFunctionPtr GetDispatchFunction(const FRigVMTemplateTypeMap& InTypes) const;
+
+	// builds and returns the template
+	const FRigVMTemplate* GetTemplate() const;
+
+protected:
+
+	FString GetPermutationName(const FRigVMTemplateTypeMap& InTypes) const;
+	FORCEINLINE virtual FRigVMFunctionPtr GetDispatchFunctionImpl(const FRigVMTemplateTypeMap& InTypes) const { return nullptr; }
+
+	static const FString DispatchPrefix;
+
+	UScriptStruct* FactoryScriptStruct;
+	mutable const FRigVMTemplate* CachedTemplate;
+	friend struct FRigVMTemplate;
+	friend struct FRigVMRegistry;
+};
diff --git a/Engine/Source/Runtime/RigVM/Public/RigVMCore/RigVMExecuteContext.h b/Engine/Source/Runtime/RigVM/Public/RigVMCore/RigVMExecuteContext.h
--- a/Engine/Source/Runtime/RigVM/Public/RigVMCore/RigVMExecuteContext.h
+++ b/Engine/Source/Runtime/RigVM/Public/RigVMCore/RigVMExecuteContext.h
@@ -11,6 +11,7 @@
 
 struct FRigVMExtendedExecuteContext;
 struct FRigVMExecuteContext;
+struct FRigVMDispatchFactory;
 class URigVM;
 
 USTRUCT()
@@ -297,7 +298,8 @@ struct RIGVM_API FRigVMExtendedExecuteContext
 
 	FORCEINLINE FRigVMExtendedExecuteContext()
 		: VM(nullptr)
-		, LastExecutionMicroSeconds() 
+		, LastExecutionMicroSeconds()
+		, Factory(nullptr)
 	{
 		Reset();
 	}
@@ -309,6 +311,7 @@ struct RIGVM_API FRigVMExtendedExecuteContext
 		Slices.Reset();
 		Slices.Add(FRigVMSlice());
 		SliceOffsets.Reset();
+		Factory = nullptr;
 	}
 
 	FORCEINLINE void CopyFrom(const FRigVMExtendedExecuteContext& Other)
@@ -402,4 +405,5 @@ struct RIGVM_API FRigVMExtendedExecuteContext
 	TArray<FRigVMSlice> Slices;
 	TArray<uint16> SliceOffsets;
 	double LastExecutionMicroSeconds;
+	const FRigVMDispatchFactory* Factory;
 };
diff --git a/Engine/Source/Runtime/RigVM/Public/RigVMCore/RigVMFunction.h b/Engine/Source/Runtime/RigVM/Public/RigVMCore/RigVMFunction.h
--- a/Engine/Source/Runtime/RigVM/Public/RigVMCore/RigVMFunction.h
+++ b/Engine/Source/Runtime/RigVM/Public/RigVMCore/RigVMFunction.h
@@ -13,6 +13,7 @@ typedef TArrayView<void*> FRigVMUserDataArray;
 typedef void (*FRigVMFunctionPtr)(FRigVMExtendedExecuteContext& RigVMExecuteContext, FRigVMMemoryHandleArray RigVMMemoryHandles);
 
 struct FRigVMTemplate;
+struct FRigVMDispatchFactory;
 
 /**
  * The Pin Direction is used to differentiate different kinds of 
@@ -36,6 +37,8 @@ struct RIGVM_API FRigVMFunctionArgument
 {
 	const TCHAR* Name;
 	const TCHAR* Type;
+	TSharedPtr<FString> NameString;
+	TSharedPtr<FString> TypeString;
 
 	FRigVMFunctionArgument()
 		: Name(nullptr)
@@ -48,6 +51,16 @@ struct RIGVM_API FRigVMFunctionArgument
 		, Type(InType)
 	{
 	}
+
+	FRigVMFunctionArgument(const FString& InName, const FString& InType)
+		: Name(nullptr)
+		, Type(nullptr)
+	{
+		NameString = MakeShareable(new FString(InName));
+		TypeString = MakeShareable(new FString(InType));
+		Name = NameString->operator*();
+		Type = TypeString->operator*();
+	}
 };
 
 /**
@@ -56,8 +69,9 @@ struct RIGVM_API FRigVMFunctionArgument
  */
 struct RIGVM_API FRigVMFunction
 {
-	const TCHAR* Name;
+	FString Name;
 	UScriptStruct* Struct;
+	FRigVMDispatchFactory* Factory;
 	FRigVMFunctionPtr FunctionPtr;
 	int32 Index;
 	int32 TemplateIndex;
@@ -65,8 +79,9 @@ struct RIGVM_API FRigVMFunction
 	mutable TArray<int32> ArgumentTypeIndices;
 
 	FRigVMFunction()
-		: Name(nullptr)
+		: Name()
 		, Struct(nullptr)
+		, Factory(nullptr)
 		, FunctionPtr(nullptr)
 		, Index(INDEX_NONE)
 		, TemplateIndex(INDEX_NONE)
@@ -78,6 +93,7 @@ struct RIGVM_API FRigVMFunction
 	FRigVMFunction(const TCHAR* InName, FRigVMFunctionPtr InFunctionPtr, UScriptStruct* InStruct = nullptr, int32 InIndex = INDEX_NONE, const TArray<FRigVMFunctionArgument>& InArguments = TArray<FRigVMFunctionArgument>())
 		: Name(InName)
 		, Struct(InStruct)
+		, Factory(nullptr)
 		, FunctionPtr(InFunctionPtr)
 		, Index(InIndex)
 		, TemplateIndex(INDEX_NONE)
@@ -85,7 +101,18 @@ struct RIGVM_API FRigVMFunction
 	{
 	}
 
-	FORCEINLINE bool IsValid() const { return Name != nullptr && FunctionPtr != nullptr; }
+	FRigVMFunction(const FString& InName, FRigVMFunctionPtr InFunctionPtr, FRigVMDispatchFactory* InFactory, int32 InIndex = INDEX_NONE, const TArray<FRigVMFunctionArgument>& InArguments = TArray<FRigVMFunctionArgument>())
+	: Name(InName)
+	, Struct(nullptr)
+	, Factory(InFactory)
+	, FunctionPtr(InFunctionPtr)
+	, Index(InIndex)
+	, TemplateIndex(INDEX_NONE)
+	, Arguments(InArguments)
+	{
+	}
+
+	FORCEINLINE bool IsValid() const { return !Name.IsEmpty() && FunctionPtr != nullptr; }
 	FString GetName() const;
 	FName GetMethodName() const;
 	FString GetModuleName() const;
diff --git a/Engine/Source/Runtime/RigVM/Public/RigVMCore/RigVMMemoryStorage.h b/Engine/Source/Runtime/RigVM/Public/RigVMCore/RigVMMemoryStorage.h
--- a/Engine/Source/Runtime/RigVM/Public/RigVMCore/RigVMMemoryStorage.h
+++ b/Engine/Source/Runtime/RigVM/Public/RigVMCore/RigVMMemoryStorage.h
@@ -86,6 +86,462 @@ public:
 		return FRigVMPropertyPath::Empty;
 	}
 
+	// Returns the resolved property the data is pointing to
+	FORCEINLINE const FProperty* GetResolvedProperty(bool bIsHiddenArgument = false) const
+	{
+		if(bIsHiddenArgument)
+		{
+			check(PropertyPath == nullptr);
+			return GetArrayElementPropertyChecked(Property);
+		}
+		else if(PropertyPath)
+		{
+			return PropertyPath->GetTailProperty();
+		}
+		return Property;
+	}
+
+	// Returns true if this memory handle maps to a given type of property
+	template<typename PropertyType>
+	FORCEINLINE bool IsPropertyType(bool bIsHiddenArgument = false) const
+	{
+		const FProperty* ResolvedProperty = GetResolvedProperty(bIsHiddenArgument);
+		return ResolvedProperty->IsA<PropertyType>();
+	}
+
+	// Returns true if this memory handle maps to a given array type of property
+	template<typename PropertyType>
+	FORCEINLINE bool IsPropertyArrayType(bool bIsHiddenArgument = false) const
+	{
+		const FProperty* ResolvedProperty = GetResolvedProperty(bIsHiddenArgument);
+		if(const FProperty* ElementProperty = GetArrayElementProperty(ResolvedProperty))
+		{
+			return ElementProperty->IsA<PropertyType>();
+		}
+		return false;
+	}
+
+	// returns true if the handle is of the given type
+	template <
+		typename T,
+		typename TEnableIf<TIsTArray<T>::Value>::Type* = nullptr
+	>
+	FORCEINLINE bool IsType(bool bIsHiddenArgument = false) const
+	{
+		return IsTypeArray<T::ElementType>(bIsHiddenArgument);
+	}
+
+	// returns true if the handle is of the given type
+	template <
+		typename T,
+		typename TEnableIf<TRigVMIsBool<T>::Value, T>::Type* = nullptr
+	>
+	FORCEINLINE bool IsType(bool bIsHiddenArgument = false) const
+	{
+		return IsBool(bIsHiddenArgument);
+	}
+
+	// returns true if the handle is of the given type
+	template <
+		typename T,
+		typename TEnableIf<TRigVMIsBool<T>::Value, T>::Type* = nullptr
+	>
+	FORCEINLINE bool IsTypeArray(bool bIsHiddenArgument = false) const
+	{
+		return IsBoolArray(bIsHiddenArgument);
+	}
+
+	// returns true if the handle is of the given type
+	template <
+		typename T,
+		typename TEnableIf<TRigVMIsFloat<T>::Value, T>::Type* = nullptr
+	>
+	FORCEINLINE bool IsType(bool bIsHiddenArgument = false) const
+	{
+		return IsFloat(bIsHiddenArgument);
+	}
+
+	// returns true if the handle is of the given type
+	template <
+		typename T,
+		typename TEnableIf<TRigVMIsFloat<T>::Value, T>::Type* = nullptr
+	>
+	FORCEINLINE bool IsTypeArray(bool bIsHiddenArgument = false) const
+	{
+		return IsFloatArray(bIsHiddenArgument);
+	}
+
+	// returns true if the handle is of the given type
+	template <
+		typename T,
+		typename TEnableIf<TRigVMIsDouble<T>::Value, T>::Type* = nullptr
+	>
+	FORCEINLINE bool IsType(bool bIsHiddenArgument = false) const
+	{
+		return IsDouble(bIsHiddenArgument);
+	}
+
+	// returns true if the handle is of the given type
+	template <
+		typename T,
+		typename TEnableIf<TRigVMIsDouble<T>::Value, T>::Type* = nullptr
+	>
+	FORCEINLINE bool IsTypeArray(bool bIsHiddenArgument = false) const
+	{
+		return IsDoubleArray(bIsHiddenArgument);
+	}
+
+	// returns true if the handle is of the given type
+	template <
+		typename T,
+		typename TEnableIf<TRigVMIsInt32<T>::Value, T>::Type* = nullptr
+	>
+	FORCEINLINE bool IsType(bool bIsHiddenArgument = false) const
+	{
+		return IsInt32(bIsHiddenArgument);
+	}
+
+	// returns true if the handle is of the given type
+	template <
+		typename T,
+		typename TEnableIf<TRigVMIsInt32<T>::Value, T>::Type* = nullptr
+	>
+	FORCEINLINE bool IsTypeArray(bool bIsHiddenArgument = false) const
+	{
+		return IsInt32Array(bIsHiddenArgument);
+	}
+	// returns true if the handle is of the given type
+	template <
+		typename T,
+		typename TEnableIf<TRigVMIsName<T>::Value, T>::Type* = nullptr
+	>
+	FORCEINLINE bool IsType(bool bIsHiddenArgument = false) const
+	{
+		return IsName(bIsHiddenArgument);
+	}
+
+	// returns true if the handle is of the given type
+	template <
+		typename T,
+		typename TEnableIf<TRigVMIsName<T>::Value, T>::Type* = nullptr
+	>
+	FORCEINLINE bool IsTypeArray(bool bIsHiddenArgument = false) const
+	{
+		return IsNameArray(bIsHiddenArgument);
+	}
+
+	// returns true if the handle is of the given type
+	template <
+		typename T,
+		typename TEnableIf<TRigVMIsString<T>::Value, T>::Type* = nullptr
+	>
+	FORCEINLINE bool IsType(bool bIsHiddenArgument = false) const
+	{
+		return IsString(bIsHiddenArgument);
+	}
+
+	// returns true if the handle is of the given type
+	template <
+		typename T,
+		typename TEnableIf<TRigVMIsString<T>::Value, T>::Type* = nullptr
+	>
+	FORCEINLINE bool IsTypeArray(bool bIsHiddenArgument = false) const
+	{
+		return IsStringArray(bIsHiddenArgument);
+	}
+
+	// returns true if the handle is of the given type
+	template <
+		typename T,
+		typename TEnableIf<TRigVMIsBaseStructure<T>::Value, T>::Type* = nullptr
+	>
+	FORCEINLINE bool IsType(bool bIsHiddenArgument = false) const
+	{
+		return IsStruct<T>(bIsHiddenArgument);
+	}
+
+	// returns true if the handle is of the given type
+	template <
+		typename T,
+		typename TEnableIf<TRigVMIsBaseStructure<T>::Value, T>::Type* = nullptr
+	>
+	FORCEINLINE bool IsTypeArray(bool bIsHiddenArgument = false) const
+	{
+		return IsStructArray<T>(bIsHiddenArgument);
+	}
+
+	// returns true if the handle is of the given type
+	template <
+		typename T,
+		typename TEnableIf<TModels<CRigVMUStruct, T>::Value>::Type * = nullptr
+	>
+	FORCEINLINE bool IsType(bool bIsHiddenArgument = false) const
+	{
+		return IsStruct<T>(bIsHiddenArgument);
+	}
+
+	// returns true if the handle is of the given type
+	template <
+		typename T,
+		typename TEnableIf<TModels<CRigVMUStruct, T>::Value>::Type * = nullptr
+	>
+	FORCEINLINE bool IsTypeArray(bool bIsHiddenArgument = false) const
+	{
+		return IsStructArray<T>(bIsHiddenArgument);
+	}
+
+	// returns true if the handle is of the given type
+	template <
+		typename T,
+		typename TEnableIf<TIsEnum<T>::Value>::Type* = nullptr
+	>
+	FORCEINLINE bool IsType(bool bIsHiddenArgument = false) const
+	{
+		return IsEnum<T>(bIsHiddenArgument);
+	}
+
+	// returns true if the handle is of the given type
+	template <
+		typename T,
+		typename TEnableIf<TIsEnum<T>::Value>::Type* = nullptr
+	>
+	FORCEINLINE bool IsTypeArray(bool bIsHiddenArgument = false) const
+	{
+		return IsEnumArray<T>(bIsHiddenArgument);
+	}
+
+	// returns true if the handle is of the given type
+	template <
+		typename T,
+		typename TEnableIf<TModels<CRigVMUClass, T>::Value>::Type * = nullptr
+	>
+	FORCEINLINE bool IsType(bool bIsHiddenArgument = false) const
+	{
+		return IsObject<T>(bIsHiddenArgument);
+	}
+
+	// returns true if the handle is of the given type
+	template <
+		typename T,
+		typename TEnableIf<TModels<CRigVMUClass, T>::Value>::Type * = nullptr
+	>
+	FORCEINLINE bool IsTypeArray(bool bIsHiddenArgument = false) const
+	{
+		return IsObjectArray<T>(bIsHiddenArgument);
+	}
+	
+	// Returns true if this memory handle maps to a bool property
+	FORCEINLINE bool IsBool(bool bIsHiddenArgument = false) const
+	{
+		return IsPropertyType<FBoolProperty>(bIsHiddenArgument);
+	}
+
+	// Returns true if this memory handle maps to a bool array property
+	FORCEINLINE bool IsBoolArray(bool bIsHiddenArgument = false) const
+	{
+		return IsPropertyArrayType<FBoolProperty>(bIsHiddenArgument);
+	}
+	
+	// Returns true if this memory handle maps to a float property
+	FORCEINLINE bool IsFloat(bool bIsHiddenArgument = false) const
+	{
+		return IsPropertyType<FFloatProperty>(bIsHiddenArgument);
+	}
+
+	// Returns true if this memory handle maps to a float array property
+	FORCEINLINE bool IsFloatArray(bool bIsHiddenArgument = false) const
+	{
+		return IsPropertyArrayType<FFloatProperty>(bIsHiddenArgument);
+	}
+
+	// Returns true if this memory handle maps to a double property
+	FORCEINLINE bool IsDouble(bool bIsHiddenArgument = false) const
+	{
+		return IsPropertyType<FDoubleProperty>(bIsHiddenArgument);
+	}
+
+	// Returns true if this memory handle maps to a double array property
+	FORCEINLINE bool IsDoubleArray(bool bIsHiddenArgument = false) const
+	{
+		return IsPropertyArrayType<FDoubleProperty>(bIsHiddenArgument);
+	}
+
+	// Returns true if this memory handle maps to a int32 property
+	FORCEINLINE bool IsInt32(bool bIsHiddenArgument = false) const
+	{
+		return IsPropertyType<FIntProperty>(bIsHiddenArgument);
+	}
+
+	// Returns true if this memory handle maps to a int32 array property
+	FORCEINLINE bool IsInt32Array(bool bIsHiddenArgument = false) const
+	{
+		return IsPropertyArrayType<FIntProperty>(bIsHiddenArgument);
+	}
+
+	// Returns true if this memory handle maps to an FName property
+    FORCEINLINE bool IsName(bool bIsHiddenArgument = false) const
+    {
+    	return IsPropertyType<FNameProperty>(bIsHiddenArgument);
+    }
+
+	// Returns true if this memory handle maps to an FName array property
+    FORCEINLINE bool IsNameArray(bool bIsHiddenArgument = false) const
+    {
+    	return IsPropertyArrayType<FNameProperty>(bIsHiddenArgument);
+    }
+
+	// Returns true if this memory handle maps to an FString property
+	FORCEINLINE bool IsString(bool bIsHiddenArgument = false) const
+	{
+		return IsPropertyType<FStrProperty>(bIsHiddenArgument);
+	}
+
+	// Returns true if this memory handle maps to an FString array property
+	FORCEINLINE bool IsStringArray(bool bIsHiddenArgument = false) const
+	{
+		return IsPropertyArrayType<FStrProperty>(bIsHiddenArgument);
+	}
+
+	// Returns true if this memory handle maps to an enum property
+	FORCEINLINE bool IsEnum(const UEnum* InEnum, bool bIsHiddenArgument = false) const
+	{
+		const FProperty* ResolvedProperty = GetResolvedProperty();
+		return IsEnum(ResolvedProperty, InEnum, bIsHiddenArgument);
+	}
+
+	// Returns true if this memory handle maps to an enum array property
+	FORCEINLINE bool IsEnumArray(const UEnum* InEnum, bool bIsHiddenArgument = false) const
+	{
+		const FProperty* ResolvedProperty = GetResolvedProperty();
+		if(bIsHiddenArgument)
+		{
+			ResolvedProperty = GetArrayElementPropertyChecked(ResolvedProperty);
+		}
+		if(const FArrayProperty* ArrayProperty = CastField<FArrayProperty>(ResolvedProperty))
+		{
+			return IsEnum(ArrayProperty->Inner, InEnum, false);
+		}
+		return false;
+	}
+
+	// Returns true if this memory handle maps to an enum property
+	template<typename T>
+	FORCEINLINE bool IsEnum(bool bIsHiddenArgument = false) const
+	{
+		return IsEnum(StaticEnum<T>(), bIsHiddenArgument);
+	}
+
+	// Returns true if this memory handle maps to an enum array property
+	template<typename T>
+	FORCEINLINE bool IsEnumArray(bool bIsHiddenArgument = false) const
+	{
+		return IsEnumArray(StaticEnum<T>(), bIsHiddenArgument);
+	}
+
+	// Returns true if this memory handle maps to a struct property
+	FORCEINLINE bool IsStruct(const UScriptStruct* InStruct, bool bIsHiddenArgument = false) const
+	{
+		const FProperty* ResolvedProperty = GetResolvedProperty();
+		return IsStruct(ResolvedProperty, InStruct, bIsHiddenArgument);
+	}
+
+	// Returns true if this memory handle maps to a struct array property
+	FORCEINLINE bool IsStructArray(const UScriptStruct* InStruct, bool bIsHiddenArgument = false) const
+	{
+		const FProperty* ResolvedProperty = GetResolvedProperty();
+		if(bIsHiddenArgument)
+		{
+			ResolvedProperty = GetArrayElementPropertyChecked(ResolvedProperty);
+		}
+		if(const FArrayProperty* ArrayProperty = CastField<FArrayProperty>(ResolvedProperty))
+		{
+			return IsStruct(ArrayProperty->Inner, InStruct, false);
+		}
+		return false;
+	}
+
+	// Returns true if this memory handle maps to a struct property
+	template <
+		typename T,
+		typename TEnableIf<TRigVMIsBaseStructure<T>::Value, T>::Type* = nullptr
+	>
+	FORCEINLINE bool IsStruct(bool bIsHiddenArgument = false) const
+	{
+		return IsStruct(TBaseStructure<T>::Get(), bIsHiddenArgument);
+	}
+
+	// Returns true if this memory handle maps to a struct property
+	template <
+		typename T,
+		typename TEnableIf<TModels<CRigVMUStruct, T>::Value>::Type * = nullptr
+	>
+	FORCEINLINE bool IsStruct(bool bIsHiddenArgument = false) const
+	{
+		return IsStruct(T::StaticStruct(), bIsHiddenArgument);
+	}
+
+	// Returns true if this memory handle maps to a struct array property
+	template <
+		typename T,
+		typename TEnableIf<TRigVMIsBaseStructure<T>::Value, T>::Type* = nullptr
+	>
+	FORCEINLINE bool IsStructArray(bool bIsHiddenArgument = false) const
+	{
+		return IsStructArray(TBaseStructure<T>::Get(), bIsHiddenArgument);
+	}
+
+	// Returns true if this memory handle maps to a struct array property
+	template <
+		typename T,
+		typename TEnableIf<TModels<CRigVMUStruct, T>::Value>::Type * = nullptr
+	>
+	FORCEINLINE bool IsStructArray(bool bIsHiddenArgument = false) const
+	{
+		return IsStructArray(T::StaticStruct(), bIsHiddenArgument);
+	}
+
+	// Returns true if this memory handle maps to an object property
+	FORCEINLINE bool IsObject(const UClass* InClass, bool bIsHiddenArgument = false) const
+	{
+		const FProperty* ResolvedProperty = GetResolvedProperty();
+		return IsObject(ResolvedProperty, InClass, bIsHiddenArgument);
+	}
+
+	// Returns true if this memory handle maps to an objct array property
+	FORCEINLINE bool IsObjectArray(const UClass* InClass, bool bIsHiddenArgument = false) const
+	{
+		const FProperty* ResolvedProperty = GetResolvedProperty();
+		if(bIsHiddenArgument)
+		{
+			ResolvedProperty = GetArrayElementPropertyChecked(ResolvedProperty);
+		}
+		if(const FArrayProperty* ArrayProperty = CastField<FArrayProperty>(ResolvedProperty))
+		{
+			return IsObject(ArrayProperty->Inner, InClass, false);
+		}
+		return false;
+	}
+
+	// Returns true if this memory handle maps to an object property
+	template <
+		typename T,
+		typename TEnableIf<TModels<CRigVMUClass, T>::Value>::Type * = nullptr
+	>
+	FORCEINLINE bool IsObject(bool bIsHiddenArgument = false) const
+	{
+		return IsObject(T::StaticClass(), bIsHiddenArgument);
+	}
+
+	// Returns true if this memory handle maps to an objct array property
+	template <
+		typename T,
+		typename TEnableIf<TModels<CRigVMUClass, T>::Value>::Type * = nullptr
+	>
+	FORCEINLINE bool IsObjectArray(bool bIsHiddenArgument = false) const
+	{
+		return IsObjectArray(T::StaticClass(), bIsHiddenArgument);
+	}
+
 private:
 
 	FORCEINLINE_DEBUGGABLE uint8* GetData_Internal(bool bFollowPropertyPath, int32 InSliceIndex) const
@@ -118,6 +574,79 @@ private:
 		return Ptr;
 	}
 
+	FORCEINLINE static const FProperty* GetArrayElementProperty(const FProperty* InProperty)
+	{
+		if(const FArrayProperty* ArrayProperty = CastField<FArrayProperty>(InProperty))
+		{
+			return ArrayProperty->Inner;
+		}
+		return nullptr;
+	}
+
+	FORCEINLINE static const FProperty* GetArrayElementPropertyChecked(const FProperty* InProperty)
+	{
+		const FArrayProperty* ArrayProperty = CastFieldChecked<FArrayProperty>(InProperty);
+		return ArrayProperty->Inner;
+	}
+
+	FORCEINLINE static bool IsEnum(const FProperty* InProperty, const UEnum* InEnum, bool bUseArrayElement)
+	{
+		const FProperty* Property = InProperty;
+		if(bUseArrayElement)
+		{
+			Property = GetArrayElementPropertyChecked(Property);
+		}
+		if(const FEnumProperty* EnumProperty = CastField<FEnumProperty>(Property))
+		{
+			if(EnumProperty->GetEnum() == InEnum)
+			{
+				return true;
+			}
+		}
+		if(const FByteProperty* EnumProperty = CastField<FByteProperty>(Property))
+		{
+			if(EnumProperty->Enum == InEnum)
+			{
+				return true;
+			}
+		}
+		return false;
+	}
+
+	FORCEINLINE static bool IsStruct(const FProperty* InProperty, const UScriptStruct* InScriptStruct, bool bUseArrayElement)
+	{
+		const FProperty* Property = InProperty;
+		if(bUseArrayElement)
+		{
+			Property = GetArrayElementPropertyChecked(Property);
+		}
+		if(const FStructProperty* StructProperty = CastField<FStructProperty>(Property))
+		{
+			if(StructProperty->Struct == InScriptStruct)
+			{
+				return true;
+			}
+		}
+		return false;
+	}
+
+	FORCEINLINE static bool IsObject(const FProperty* InProperty, const UClass* InClass, bool bUseArrayElement)
+	{
+		const FProperty* Property = InProperty;
+		if(bUseArrayElement)
+		{
+			Property = GetArrayElementPropertyChecked(Property);
+		}
+		if(const FObjectProperty* ObjectProperty = CastField<FObjectProperty>(Property))
+		{
+			if(ObjectProperty->PropertyClass == InClass)
+			{
+				return true;
+			}
+		}
+		return false;
+	}
+
 	// The pointer of the memory of the head property
 	uint8* Ptr;
 
diff --git a/Engine/Source/Runtime/RigVM/Public/RigVMCore/RigVMRegistry.h b/Engine/Source/Runtime/RigVM/Public/RigVMCore/RigVMRegistry.h
--- a/Engine/Source/Runtime/RigVM/Public/RigVMCore/RigVMRegistry.h
+++ b/Engine/Source/Runtime/RigVM/Public/RigVMCore/RigVMRegistry.h
@@ -3,14 +3,11 @@
 #pragma once
 
 #include "CoreMinimal.h"
+#include "RigVMTypeIndex.h"
 #include "RigVMMemory.h"
 #include "RigVMFunction.h"
 #include "RigVMTemplate.h"
-
-struct RIGVM_API FRigVMTemplateDelegates
-{
-	FRigVMTemplate_NewArgumentTypeDelegate NewArgumentTypeDelegate;
-};
+#include "RigVMDispatchFactory.h"
 
 /**
  * The FRigVMRegistry is used to manage all known function pointers
@@ -22,6 +19,8 @@ struct RIGVM_API FRigVMRegistry
 {
 public:
 
+	~FRigVMRegistry();
+	
 	// Returns the singleton registry
 	static FRigVMRegistry& Get();
 
@@ -30,6 +29,9 @@ public:
 	// for example "FMyStruct::MyVirtualMethod"
 	void Register(const TCHAR* InName, FRigVMFunctionPtr InFunctionPtr, UScriptStruct* InStruct = nullptr, const TArray<FRigVMFunctionArgument>& InArguments = TArray<FRigVMFunctionArgument>());
 
+	// Registers a dispatch factory given its struct.
+	void RegisterFactory(UScriptStruct* InFactoryStruct);
+
 	// Initializes the registry by storing the defaults
 	void InitializeIfNeeded();
 	
@@ -39,19 +41,79 @@ public:
 
 	// Adds a type if it doesn't exist yet and returns its index.
 	// This function is not thead-safe
-	int32 FindOrAddType(const FRigVMTemplateArgumentType& InType);
+	TRigVMTypeIndex FindOrAddType(const FRigVMTemplateArgumentType& InType);
 
 	// Returns the type index given a type
-	int32 GetTypeIndex(const FRigVMTemplateArgumentType& InType) const;
+	TRigVMTypeIndex GetTypeIndex(const FRigVMTemplateArgumentType& InType) const;
 
 	// Returns the type index given a cpp type and a type object
-	FORCEINLINE int32 GetTypeIndex(const FName& InCPPType, UObject* InCPPTypeObject) const
+	FORCEINLINE TRigVMTypeIndex GetTypeIndex(const FName& InCPPType, UObject* InCPPTypeObject) const
 	{
 		return GetTypeIndex(FRigVMTemplateArgumentType(InCPPType, InCPPTypeObject));
 	}
 
+	// Returns the type index given an enum
+	template <
+		typename T,
+		typename TEnableIf<TIsEnum<T>::Value>::Type* = nullptr
+	>
+	FORCEINLINE TRigVMTypeIndex GetTypeIndex(bool bAsArray = false) const
+	{
+		FRigVMTemplateArgumentType Type(StaticEnum<T>());
+		if(bAsArray)
+		{
+			Type.ConvertToArray();
+		}
+		return GetTypeIndex(Type);
+	}
+
+	// Returns the type index given a struct
+	template <
+		typename T,
+		typename TEnableIf<TRigVMIsBaseStructure<T>::Value, T>::Type* = nullptr
+	>
+	FORCEINLINE TRigVMTypeIndex GetTypeIndex(bool bAsArray = false) const
+	{
+		FRigVMTemplateArgumentType Type(TBaseStructure<T>::Get());
+		if(bAsArray)
+		{
+			Type.ConvertToArray();
+		}
+		return GetTypeIndex(Type);
+	}
+
+	// Returns the type index given a struct
+	template <
+		typename T,
+		typename TEnableIf<TModels<CRigVMUStruct, T>::Value>::Type * = nullptr
+	>
+	FORCEINLINE TRigVMTypeIndex GetTypeIndex(bool bAsArray = false) const
+	{
+		FRigVMTemplateArgumentType Type(T::StaticStruct());
+		if(bAsArray)
+		{
+			Type.ConvertToArray();
+		}
+		return GetTypeIndex(Type);
+	}
+
+	// Returns the type index given an object
+	template <
+		typename T,
+		typename TEnableIf<TModels<CRigVMUClass, T>::Value>::Type * = nullptr
+	>
+	FORCEINLINE TRigVMTypeIndex GetTypeIndex(bool bAsArray = false) const
+	{
+		FRigVMTemplateArgumentType Type(T::StaticClass());
+		if(bAsArray)
+		{
+			Type.ConvertToArray();
+		}
+		return GetTypeIndex(Type);
+	}
+
 	// Returns the type given its index
-	const FRigVMTemplateArgumentType& GetType(int32 InTypeIndex) const;
+	const FRigVMTemplateArgumentType& GetType(TRigVMTypeIndex InTypeIndex) const;
 
 	// Returns the number of types
 	FORCEINLINE int32 NumTypes() const { return Types.Num(); }
@@ -60,34 +122,34 @@ public:
 	const FRigVMTemplateArgumentType& FindTypeFromCPPType(const FString& InCPPType) const;
 
 	// Returns the type index given only its cpp type
-	int32 GetTypeIndexFromCPPType(const FString& InCPPType) const;
+	TRigVMTypeIndex GetTypeIndexFromCPPType(const FString& InCPPType) const;
 
 	// Returns true if the type is an array
-	bool IsArrayType(int32 InTypeIndex) const;
+	bool IsArrayType(TRigVMTypeIndex InTypeIndex) const;
 
 	// Returns true if the type is an execute type
-	bool IsExecuteType(int32 InTypeIndex) const;
+	bool IsExecuteType(TRigVMTypeIndex InTypeIndex) const;
 
 	// Returns the dimensions of the array 
-	int32 GetArrayDimensionsForType(int32 InTypeIndex) const;
+	int32 GetArrayDimensionsForType(TRigVMTypeIndex InTypeIndex) const;
 
 	// Returns true if the type is a wildcard type
-	bool IsWildCardType(int32 InTypeIndex) const;
+	bool IsWildCardType(TRigVMTypeIndex InTypeIndex) const;
 
 	// Returns true if the types can be matched.
-	bool CanMatchTypes(int32 InTypeIndexA, int32 InTypeIndexB, bool bAllowFloatingPointCasts) const;
+	bool CanMatchTypes(TRigVMTypeIndex InTypeIndexA, TRigVMTypeIndex InTypeIndexB, bool bAllowFloatingPointCasts) const;
 
 	// Returns the list of compatible types for a given type
-	const TArray<int32>& GetCompatibleTypes(int32 InTypeIndex) const;
+	const TArray<TRigVMTypeIndex>& GetCompatibleTypes(TRigVMTypeIndex InTypeIndex) const;
 
 	// Returns all compatible types given a category
-	const TArray<int32>& GetTypesForCategory(FRigVMTemplateArgument::ETypeCategory InCategory);
+	const TArray<TRigVMTypeIndex>& GetTypesForCategory(FRigVMTemplateArgument::ETypeCategory InCategory);
 
 	// Returns the type index of the array matching the given element type index
-	int32 GetArrayTypeFromBaseTypeIndex(int32 InTypeIndex) const;
+	TRigVMTypeIndex GetArrayTypeFromBaseTypeIndex(TRigVMTypeIndex InTypeIndex) const;
 
 	// Returns the type index of the element matching the given array type index
-	int32 GetBaseTypeFromArrayTypeIndex(int32 InTypeIndex) const;
+	TRigVMTypeIndex GetBaseTypeFromArrayTypeIndex(TRigVMTypeIndex InTypeIndex) const;
 
 	// Returns the function given its name (or nullptr)
 	const FRigVMFunction* FindFunction(const TCHAR* InName) const;
@@ -110,6 +172,12 @@ public:
 		const TArray<FRigVMTemplateArgument>& InArguments,
 		const FRigVMTemplateDelegates& InDelegates);
 
+	// Returns a dispatch factory given its name (or nullptr)
+	FRigVMDispatchFactory* FindDispatchFactory(const FName& InFactoryName) const;
+
+	// Returns all dispatch factories
+	const TArray<FRigVMDispatchFactory*>& GetFactories() const;
+
 	static const TArray<UScriptStruct*>& GetMathTypes();
 
 private:
@@ -134,8 +202,8 @@ private:
 		{}
 		
 		FRigVMTemplateArgumentType Type;
-		int32 BaseTypeIndex;
-		int32 ArrayTypeIndex;
+		TRigVMTypeIndex BaseTypeIndex;
+		TRigVMTypeIndex ArrayTypeIndex;
 		bool bIsArray;
 		bool bIsExecute;
 	};
@@ -155,12 +223,11 @@ private:
 	static bool IsAllowedType(const UStruct* InStruct);
 	static bool IsAllowedType(const UClass* InClass);
 
-	void RegisterTypeInCategory(FRigVMTemplateArgument::ETypeCategory InCategory, int32 InTypeIndex);
+	void RegisterTypeInCategory(FRigVMTemplateArgument::ETypeCategory InCategory, TRigVMTypeIndex InTypeIndex);
 
 	// memory for all (known) types
-	// We use TChunkedArray because we need the memory locations to be stable, since we only ever add and never remove.
 	TArray<FTypeInfo> Types;
-	TMap<FRigVMTemplateArgumentType, int32> TypeToIndex;
+	TMap<FRigVMTemplateArgumentType, TRigVMTypeIndex> TypeToIndex;
 
 	// memory for all functions
 	// We use TChunkedArray because we need the memory locations to be stable, since we only ever add and never remove.
@@ -169,6 +236,9 @@ private:
 	// memory for all templates
 	TChunkedArray<FRigVMTemplate> Templates;
 
+	// memory for all dispatch factories
+	TArray<FRigVMDispatchFactory*> Factories;
+
 	// name lookup for functions
 	TMap<FName, int32> FunctionNameToIndex;
 
@@ -176,11 +246,12 @@ private:
 	TMap<FName, int32> TemplateNotationToIndex;
 
 	// Maps storing the default types per type category
-	TMap<FRigVMTemplateArgument::ETypeCategory, TArray<int32>> TypesPerCategory;
+	TMap<FRigVMTemplateArgument::ETypeCategory, TArray<TRigVMTypeIndex>> TypesPerCategory;
 
 	// Lookup per type category to know which argument to keep in sync
 	TMap<FRigVMTemplateArgument::ETypeCategory, TArray<TPair<int32,int32>>> ArgumentsPerCategory;
 
 	static FRigVMRegistry s_RigVMRegistry;
 	friend struct FRigVMStruct;
+	friend struct FRigVMTemplate;
 };
diff --git a/Engine/Source/Runtime/RigVM/Public/RigVMCore/RigVMStruct.h b/Engine/Source/Runtime/RigVM/Public/RigVMCore/RigVMStruct.h
--- a/Engine/Source/Runtime/RigVM/Public/RigVMCore/RigVMStruct.h
+++ b/Engine/Source/Runtime/RigVM/Public/RigVMCore/RigVMStruct.h
@@ -8,6 +8,7 @@
 #include "RigVMCore/RigVMExternalVariable.h"
 #include "RigVMCore/RigVMTraits.h"
 #include "RigVMCore/RigVMUserWorkflow.h"
+#include "RigVMCore/RigVMStructUpgradeInfo.h"
 #include "RigVMStruct.generated.h"
 
 // delegates used for variable introspection / creation
@@ -136,80 +137,6 @@ private:
 	friend struct FScope;
 };
 
-/** Structure used to upgrade to a new implementation of a node */
-struct RIGVM_API FRigVMStructUpgradeInfo
-{
-public:
-	
-	FRigVMStructUpgradeInfo();
-
-	template<typename Old, typename New>
-	FRigVMStructUpgradeInfo(const Old& InOld, const New& InNew)
-	{
-		OldStruct = Old::StaticStruct();
-		NewStruct = New::StaticStruct();
-		SetDefaultValues(&InNew);
-	}
-
-	// returns true if this upgrade info can be applied
-	bool IsValid() const;
-
-	// returns the old struct trying to be upgraded
-	UScriptStruct* GetOldStruct() const { return OldStruct; }
-
-	// returns the new struct to upgrade to
-	UScriptStruct* GetNewStruct() const { return NewStruct; }
-
-	// returns the map for all default values
-	const TMap<FName, FString>& GetDefaultValues() const { return DefaultValues; }
-
-	// returns the default value for a given pin
-	const FString& GetDefaultValueForPin(const FName& InPinName) const;
-
-	// adds a pin to be remapped
-	void AddRemappedPin(const FString& InOldPinPath, const FString& InNewPinPath, bool bAsInput = true, bool bAsOutput = true);
-
-	// remaps a pin path based on our internals
-	FString RemapPin(const FString& InPinPath, bool bIsInput, bool bContainsNodeName) const;
-
-	// adds a new aggregate pin
-	FString AddAggregatePin(FString InPinName = FString());
-
-	// returns the aggregate pins to add
-	const TArray<FString>& GetAggregatePins() const { return AggregatePins; }
-
-private:
-	
-	// The complete node path including models / collapse node.
-	// The path may look like "RigGraph|CollapseNode1|Add"
-	FString NodePath;
-	
-	// The old struct this upgrade info originates from
-	UScriptStruct* OldStruct;
-
-	// The new struct this upgrade info is targeting
-	UScriptStruct* NewStruct;
-
-	// Remapping info for re-linking inputs 
-	// Entries can be root pins or sub pins
-	TMap<FString, FString> InputLinkMap;
-
-	// Remapping info for re-linking outputs 
-	// Entries can be root pins or sub pins
-	TMap<FString, FString> OutputLinkMap;
-
-	// New sets of default values
-	TMap<FName, FString> DefaultValues;
-
-	// Aggregate pins to add
-	TArray<FString> AggregatePins;
-
-	// sets the default values from the new struct.
-	void SetDefaultValues(const FRigVMStruct* InNewStructMemory);
-
-	friend class URigVMController;
-};
-
 /**
  * The base class for all RigVM enabled structs.
  */
diff --git a/Engine/Source/Runtime/RigVM/Public/RigVMCore/RigVMStructUpgradeInfo.h b/Engine/Source/Runtime/RigVM/Public/RigVMCore/RigVMStructUpgradeInfo.h
new file mode 100644
--- /dev/null
+++ b/Engine/Source/Runtime/RigVM/Public/RigVMCore/RigVMStructUpgradeInfo.h
@@ -0,0 +1,89 @@
+// Copyright Epic Games, Inc. All Rights Reserved.
+
+#pragma once
+
+#include "CoreMinimal.h"
+
+struct FRigVMStruct;
+
+/** Structure used to upgrade to a new implementation of a node */
+struct RIGVM_API FRigVMStructUpgradeInfo
+{
+public:
+	
+	FRigVMStructUpgradeInfo();
+
+	template<typename Old, typename New>
+	FRigVMStructUpgradeInfo(const Old& InOld, const New& InNew, const FName& InOldDispatchFunction = NAME_None, const FName& InNewDispatchFunction = NAME_None)
+	{
+		OldStruct = Old::StaticStruct();
+		NewStruct = New::StaticStruct();
+		OldDispatchFunction = InOldDispatchFunction;
+		NewDispatchFunction = InNewDispatchFunction;
+		SetDefaultValues(&InNew);
+	}
+
+	// returns true if this upgrade info can be applied
+	bool IsValid() const;
+
+	// returns the old struct trying to be upgraded
+	UScriptStruct* GetOldStruct() const { return OldStruct; }
+
+	// returns the new struct to upgrade to
+	UScriptStruct* GetNewStruct() const { return NewStruct; }
+
+	// returns the map for all default values
+	const TMap<FName, FString>& GetDefaultValues() const { return DefaultValues; }
+
+	// returns the default value for a given pin
+	const FString& GetDefaultValueForPin(const FName& InPinName) const;
+
+	// adds a pin to be remapped
+	void AddRemappedPin(const FString& InOldPinPath, const FString& InNewPinPath, bool bAsInput = true, bool bAsOutput = true);
+
+	// remaps a pin path based on our internals
+	FString RemapPin(const FString& InPinPath, bool bIsInput, bool bContainsNodeName) const;
+
+	// adds a new aggregate pin
+	FString AddAggregatePin(FString InPinName = FString());
+
+	// returns the aggregate pins to add
+	const TArray<FString>& GetAggregatePins() const { return AggregatePins; }
+
+private:
+	
+	// The complete node path including models / collapse node.
+	// The path may look like "RigGraph|CollapseNode1|Add"
+	FString NodePath;
+	
+	// The old struct this upgrade info originates from
+	UScriptStruct* OldStruct;
+
+	// The new struct this upgrade info is targeting
+	UScriptStruct* NewStruct;
+
+	// The old optional function name this upgrade info originates from
+	FName OldDispatchFunction;
+
+	// The new optional function name this upgrade info is targeting
+	FName NewDispatchFunction;
+
+	// Remapping info for re-linking inputs 
+	// Entries can be root pins or sub pins
+	TMap<FString, FString> InputLinkMap;
+
+	// Remapping info for re-linking outputs 
+	// Entries can be root pins or sub pins
+	TMap<FString, FString> OutputLinkMap;
+
+	// New sets of default values
+	TMap<FName, FString> DefaultValues;
+
+	// Aggregate pins to add
+	TArray<FString> AggregatePins;
+
+	// sets the default values from the new struct.
+	void SetDefaultValues(const FRigVMStruct* InNewStructMemory);
+
+	friend class URigVMController;
+};
\ No newline at end of file
diff --git a/Engine/Source/Runtime/RigVM/Public/RigVMCore/RigVMTemplate.h b/Engine/Source/Runtime/RigVM/Public/RigVMCore/RigVMTemplate.h
--- a/Engine/Source/Runtime/RigVM/Public/RigVMCore/RigVMTemplate.h
+++ b/Engine/Source/Runtime/RigVM/Public/RigVMCore/RigVMTemplate.h
@@ -3,15 +3,27 @@
 #pragma once
 
 #include "CoreMinimal.h"
+#include "RigVMTypeIndex.h"
 #include "RigVMFunction.h"
 #include "UObject/Object.h"
 #include "RigVMTypeUtils.h"
+#include "RigVMTraits.h"
 #include "RigVMTemplate.generated.h"
 
 struct FRigVMTemplate;
-typedef TMap<FName, int32> FRigVMTemplateTypeMap;
+struct FRigVMDispatchFactory;
+typedef TMap<FName, TRigVMTypeIndex> FRigVMTemplateTypeMap;
 
-DECLARE_DELEGATE_RetVal_ThreeParams(FRigVMTemplateTypeMap, FRigVMTemplate_NewArgumentTypeDelegate, const FRigVMTemplate* /* InTemplate */, const FName& /* InArgumentName */, int32 /* InTypeIndexToAdd */);
+DECLARE_DELEGATE_RetVal_ThreeParams(FRigVMTemplateTypeMap, FRigVMTemplate_NewArgumentTypeDelegate, const FRigVMTemplate* /* InTemplate */, const FName& /* InArgumentName */, TRigVMTypeIndex /* InTypeIndexToAdd */);
+DECLARE_DELEGATE_RetVal_TwoParams(FRigVMFunctionPtr, FRigVMTemplate_RequestDispatchFunctionDelegate, const FRigVMTemplate* /* InTemplate */,  const FRigVMTemplateTypeMap& /* InTypes */);
+DECLARE_DELEGATE_RetVal(FRigVMDispatchFactory*, FRigVMTemplate_GetDispatchFactoryDelegate);
+
+struct RIGVM_API FRigVMTemplateDelegates
+{
+	FRigVMTemplate_NewArgumentTypeDelegate NewArgumentTypeDelegate;
+	FRigVMTemplate_GetDispatchFactoryDelegate GetDispatchFactoryDelegate;
+	FRigVMTemplate_RequestDispatchFunctionDelegate RequestDispatchFunctionDelegate;
+};
 
 USTRUCT()
 struct RIGVM_API FRigVMTemplateArgumentType
@@ -45,6 +57,12 @@ struct RIGVM_API FRigVMTemplateArgumentType
 	{
 	}
 
+	FRigVMTemplateArgumentType(UEnum* InEnum)
+	: CPPType(*RigVMTypeUtils::CPPTypeFromEnum(InEnum))
+	, CPPTypeObject(InEnum)
+	{
+	}
+
 	static FRigVMTemplateArgumentType Array()
 	{
 		return FRigVMTemplateArgumentType(RigVMTypeUtils::GetWildCardArrayCPPTypeName(), RigVMTypeUtils::GetWildCardCPPTypeObject());
@@ -110,11 +128,8 @@ struct RIGVM_API FRigVMTemplateArgumentType
  * The template argument represents a single parameter
  * in a function call and all of its possible types
  */
-USTRUCT()
 struct RIGVM_API FRigVMTemplateArgument
 {
-	GENERATED_BODY()
-	
 	enum EArrayType
 	{
 		EArrayType_SingleValue,
@@ -151,8 +166,8 @@ struct RIGVM_API FRigVMTemplateArgument
 	// default constructor
 	FRigVMTemplateArgument();
 
-	FRigVMTemplateArgument(const FName& InName, ERigVMPinDirection InDirection, int32 InType);
-	FRigVMTemplateArgument(const FName& InName, ERigVMPinDirection InDirection, const TArray<int32>& InTypeIndices);
+	FRigVMTemplateArgument(const FName& InName, ERigVMPinDirection InDirection, TRigVMTypeIndex InType);
+	FRigVMTemplateArgument(const FName& InName, ERigVMPinDirection InDirection, const TArray<TRigVMTypeIndex>& InTypeIndices);
 	FRigVMTemplateArgument(const FName& InName, ERigVMPinDirection InDirection, const TArray<ETypeCategory>& InTypeCategories);
 
 	// returns the name of the argument
@@ -162,13 +177,13 @@ struct RIGVM_API FRigVMTemplateArgument
 	ERigVMPinDirection GetDirection() const { return Direction; }
 
 	// returns true if this argument supports a given type across a set of permutations
-	bool SupportsTypeIndex(int32 InTypeIndex, int32* OutTypeIndex = nullptr) const;
+	bool SupportsTypeIndex(TRigVMTypeIndex InTypeIndex, TRigVMTypeIndex* OutTypeIndex = nullptr) const;
 
 	// returns the flat list of types (including duplicates) of this argument
-	const TArray<int32>& GetTypeIndices() const;
+	const TArray<TRigVMTypeIndex>& GetTypeIndices() const;
 
 	// returns an array of all of the supported types
-	TArray<int32> GetSupportedTypeIndices(const TArray<int32>& InPermutationIndices = TArray<int32>()) const;
+	TArray<TRigVMTypeIndex> GetSupportedTypeIndices(const TArray<int32>& InPermutationIndices = TArray<int32>()) const;
 
 #if WITH_EDITOR
 	// returns an array of all supported types as strings. this is used for automated testing only.
@@ -186,19 +201,12 @@ struct RIGVM_API FRigVMTemplateArgument
 
 protected:
 
-	UPROPERTY()
 	int32 Index;
-
-	UPROPERTY()
 	FName Name;
-
-	UPROPERTY()
 	ERigVMPinDirection Direction;
+	TArray<TRigVMTypeIndex> TypeIndices;
 
-	UPROPERTY()
-	TArray<int32> TypeIndices;
-
-	TMap<int32, TArray<int32>> TypeToPermutations;
+	TMap<TRigVMTypeIndex, TArray<int32>> TypeToPermutations;
 	TArray<ETypeCategory> TypeCategories;
 
 	// constructor from a property
@@ -217,14 +225,12 @@ protected:
  * to build polymorphic nodes (RigVMTemplateNode) that can
  * take on any of the permutations supported by the template.
  */
-USTRUCT()
 struct RIGVM_API FRigVMTemplate
 {
-	GENERATED_BODY()
 public:
 
 	typedef FRigVMTemplateTypeMap FTypeMap;
-	typedef TPair<FName, int32> FTypePair;
+	typedef TPair<FName, TRigVMTypeIndex> FTypePair;
 
 	// Default constructor
 	FRigVMTemplate();
@@ -254,7 +260,7 @@ public:
 	const FRigVMTemplateArgument* FindArgument(const FName& InArgumentName) const;
 
 	// returns true if a given arg supports a type
-	bool ArgumentSupportsTypeIndex(const FName& InArgumentName, int32 InTypeIndex, int32* OutTypeIndex = nullptr) const;
+	bool ArgumentSupportsTypeIndex(const FName& InArgumentName, TRigVMTypeIndex InTypeIndex, TRigVMTypeIndex* OutTypeIndex = nullptr) const;
 
 	// returns the number of permutations supported by this template
 	int32 NumPermutations() const { return Permutations.Num(); }
@@ -262,12 +268,18 @@ public:
 	// returns a permutation given an index
 	const FRigVMFunction* GetPermutation(int32 InIndex) const;
 
+	// returns a permutation given an index and creates it using the backing factory if needed
+	const FRigVMFunction* GetOrCreatePermutation(int32 InIndex);
+
 	// returns true if a given function is a permutation of this template
 	bool ContainsPermutation(const FRigVMFunction* InPermutation) const;
 
 	// returns the index of the permutation within the template of a given function (or INDEX_NONE)
 	int32 FindPermutation(const FRigVMFunction* InPermutation) const;
 
+	// returns the index of the permutation within the template of a given set of types
+	int32 FindPermutation(const FTypeMap& InTypes) const;
+
 	// returns true if the template was able to resolve to single permutation
 	bool FullyResolve(FTypeMap& InOutTypes, int32& OutPermutationIndex) const;
 
@@ -275,7 +287,10 @@ public:
 	bool Resolve(FTypeMap& InOutTypes, TArray<int32> & OutPermutationIndices, bool bAllowFloatingPointCasts) const;
 
 	// returns true if the template can resolve an argument to a new type
-	bool ResolveArgument(const FName& InArgumentName, const int32 InTypeIndex, FTypeMap& InOutTypes) const;
+	bool ResolveArgument(const FName& InArgumentName, const TRigVMTypeIndex InTypeIndex, FTypeMap& InOutTypes) const;
+
+	// returns the types for a specific permutation
+	FRigVMTemplateTypeMap GetTypesForPermutation(const int32 InPermutationIndex) const;
 
 	// returns true if a given argument is valid for a template
 	static bool IsValidArgumentForTemplate(const FRigVMTemplateArgument& InArgument);
@@ -289,6 +304,9 @@ public:
 	// returns an array of structs in the inheritance order of a given struct
 	static TArray<UStruct*> GetSuperStructs(UStruct* InStruct, bool bIncludeLeaf = true);
 
+	// converts the types provided by a string (like "A:float,B:int32") into a type map
+	FTypeMap GetArgumentTypesFromString(const FString& InTypeString) const;
+
 #if WITH_EDITOR
 
 	// Returns the color based on the permutation's metadata
@@ -310,10 +328,27 @@ public:
 
 	// Adds a new argument to the template. The template needs to rely on the delegate to ask the
 	// template factory how to deal with the new type.
-	bool AddTypeForArgument(const FName& InArgumentName, int32 InTypeIndex);
+	bool AddTypeForArgument(const FName& InArgumentName, TRigVMTypeIndex InTypeIndex);
 
 	// Returns the delegate to be able to react to type changes dynamically
-	FRigVMTemplate_NewArgumentTypeDelegate& OnNewArgumentType() { return NewArgumentTypeDelegate; }
+	FRigVMTemplate_NewArgumentTypeDelegate& OnNewArgumentType() { return Delegates.NewArgumentTypeDelegate; }
+
+	// Returns the factory this template was created by
+	FORCEINLINE const FRigVMDispatchFactory* GetDispatchFactory() const
+	{
+		if(Delegates.GetDispatchFactoryDelegate.IsBound())
+		{
+			return Delegates.GetDispatchFactoryDelegate.Execute();
+		}
+		return nullptr;
+	}
+
+	// Returns true if this template is backed by a dispatch factory
+	FORCEINLINE bool UsesDispatch() const
+	{
+		return Delegates.RequestDispatchFunctionDelegate.IsBound() &&
+			Delegates.GetDispatchFactoryDelegate.IsBound();
+	}  
 
 private:
 
@@ -323,22 +358,17 @@ private:
 	// Constructor from a template name, arguments and a function index
 	FRigVMTemplate(const FName& InTemplateName, const TArray<FRigVMTemplateArgument>& InArguments, int32 InFunctionIndex);
 
-	UPROPERTY()
-	int32 Index;
+	static FLinearColor GetColorFromMetadata(FString InMetadata);
 
-	UPROPERTY()
+	int32 Index;
 	FName Notation;
-
-	UPROPERTY()
 	TArray<FRigVMTemplateArgument> Arguments;
-
-	UPROPERTY()
 	TArray<int32> Permutations;
 
-	FRigVMTemplate_NewArgumentTypeDelegate NewArgumentTypeDelegate; 
+	FRigVMTemplateDelegates Delegates;
 
 	friend struct FRigVMRegistry;
 	friend class URigVMController;
 	friend class URigVMLibraryNode;
+	friend struct FRigVMDispatchFactory;
 };
-
diff --git a/Engine/Source/Runtime/RigVM/Public/RigVMCore/RigVMTraits.h b/Engine/Source/Runtime/RigVM/Public/RigVMCore/RigVMTraits.h
--- a/Engine/Source/Runtime/RigVM/Public/RigVMCore/RigVMTraits.h
+++ b/Engine/Source/Runtime/RigVM/Public/RigVMCore/RigVMTraits.h
@@ -23,6 +23,34 @@ struct CRigVMUClass
 	);
 };
 
+template <typename T>
+struct TRigVMIsBool
+{
+	enum { Value = false };
+};
+template <> struct TRigVMIsBool<bool> { enum { Value = true }; };
+
+template <typename T>
+struct TRigVMIsFloat
+{
+	enum { Value = false };
+};
+template <> struct TRigVMIsFloat<float> { enum { Value = true }; };
+
+template <typename T>
+struct TRigVMIsDouble
+{
+	enum { Value = false };
+};
+template <> struct TRigVMIsDouble<double> { enum { Value = true }; };
+
+template <typename T>
+struct TRigVMIsInt32
+{
+	enum { Value = false };
+};
+template <> struct TRigVMIsInt32<int32> { enum { Value = true }; };
+
 template <typename T>
 struct TRigVMIsName
 {
diff --git a/Engine/Source/Runtime/RigVM/Public/RigVMCore/RigVMTypeIndex.h b/Engine/Source/Runtime/RigVM/Public/RigVMCore/RigVMTypeIndex.h
new file mode 100644
--- /dev/null
+++ b/Engine/Source/Runtime/RigVM/Public/RigVMCore/RigVMTypeIndex.h
@@ -0,0 +1,84 @@
+// Copyright Epic Games, Inc. All Rights Reserved.
+
+#pragma once
+
+#include "CoreMinimal.h"
+#include "RigVMDefines.h"
+
+#if UE_RIGVM_DEBUG_TYPEINDEX
+class RIGVM_API TRigVMTypeIndex
+{
+public:
+	
+	TRigVMTypeIndex()
+		: Name(NAME_None)
+		, Index(INDEX_NONE)
+	{}
+
+	TRigVMTypeIndex(int32 InIndex)
+		: Name(NAME_None)
+		, Index(InIndex)
+	{}
+
+	int32 GetIndex() const { return Index; }
+	const FName& GetName() const { return Name; }
+
+	operator int() const
+	{
+		return Index;
+	}
+
+	bool operator ==(const TRigVMTypeIndex& Other) const
+	{
+		return Index == Other.Index;
+	}
+
+	bool operator ==(const int32& Other) const
+	{
+		return Index == Other;
+	}
+
+	bool operator !=(const TRigVMTypeIndex& Other) const
+	{
+		return Index != Other.Index;
+	}
+
+	bool operator !=(const int32& Other) const
+	{
+		return Index != Other;
+	}
+
+	bool operator >(const TRigVMTypeIndex& Other) const
+	{
+		return Index > Other.Index;
+	}
+
+	bool operator >(const int32& Other) const
+	{
+		return Index > Other;
+	}
+
+	bool operator <(const TRigVMTypeIndex& Other) const
+	{
+		return Index < Other.Index;
+	}
+
+	bool operator <(const int32& Other) const
+	{
+		return Index < Other;
+	}
+
+	friend uint32 GetTypeHash(const TRigVMTypeIndex& InIndex)
+	{
+		return GetTypeHash(InIndex.Index);
+	}
+
+protected:
+	FName Name;
+	int32 Index;
+
+	friend struct FRigVMRegistry; 
+};
+#else
+typedef int32 TRigVMTypeIndex;
+#endif
diff --git a/Engine/Source/Runtime/RigVM/Public/RigVMDefines.h b/Engine/Source/Runtime/RigVM/Public/RigVMDefines.h
--- a/Engine/Source/Runtime/RigVM/Public/RigVMDefines.h
+++ b/Engine/Source/Runtime/RigVM/Public/RigVMDefines.h
@@ -20,4 +20,8 @@
 
 #ifndef UE_RIGVM_AGGREGATE_NODES_ENABLED
 #define UE_RIGVM_AGGREGATE_NODES_ENABLED 1
-#endif
\ No newline at end of file
+#endif
+
+#ifndef UE_RIGVM_DEBUG_TYPEINDEX
+#define UE_RIGVM_DEBUG_TYPEINDEX 1
+#endif
diff --git a/Engine/Source/Runtime/RigVM/Public/RigVMTypeUtils.h b/Engine/Source/Runtime/RigVM/Public/RigVMTypeUtils.h
--- a/Engine/Source/Runtime/RigVM/Public/RigVMTypeUtils.h
+++ b/Engine/Source/Runtime/RigVM/Public/RigVMTypeUtils.h
@@ -3,6 +3,7 @@
 #pragma once
 
 #include "CoreMinimal.h"
+#include "RigVMCore/RigVMTypeIndex.h"
 #include "RigVMCore/RigVMUnknownType.h"
 #include "UObject/Interface.h"
 
@@ -48,23 +49,23 @@ namespace RigVMTypeUtils
 	class RIGVM_API TypeIndex
 	{
 	public:
-		static inline int32 Execute = INDEX_NONE;	
-		static inline int32 Bool = INDEX_NONE;	
-		static inline int32 Float = INDEX_NONE;	
-		static inline int32 Double = INDEX_NONE;	
-		static inline int32 Int32 = INDEX_NONE;	
-		static inline int32 UInt8 = INDEX_NONE;	
-		static inline int32 FName = INDEX_NONE;	
-		static inline int32 FString = INDEX_NONE;
-		static inline int32 WildCard = INDEX_NONE;	
-		static inline int32 BoolArray = INDEX_NONE;	
-		static inline int32 FloatArray = INDEX_NONE;	
-		static inline int32 DoubleArray = INDEX_NONE;	
-		static inline int32 Int32Array = INDEX_NONE;	
-		static inline int32 UInt8Array = INDEX_NONE;	
-		static inline int32 FNameArray = INDEX_NONE;	
-		static inline int32 FStringArray = INDEX_NONE;	
-		static inline int32 WildCardArray = INDEX_NONE;	
+		static inline TRigVMTypeIndex Execute = INDEX_NONE;	
+		static inline TRigVMTypeIndex Bool = INDEX_NONE;	
+		static inline TRigVMTypeIndex Float = INDEX_NONE;	
+		static inline TRigVMTypeIndex Double = INDEX_NONE;	
+		static inline TRigVMTypeIndex Int32 = INDEX_NONE;	
+		static inline TRigVMTypeIndex UInt8 = INDEX_NONE;	
+		static inline TRigVMTypeIndex FName = INDEX_NONE;	
+		static inline TRigVMTypeIndex FString = INDEX_NONE;
+		static inline TRigVMTypeIndex WildCard = INDEX_NONE;	
+		static inline TRigVMTypeIndex BoolArray = INDEX_NONE;	
+		static inline TRigVMTypeIndex FloatArray = INDEX_NONE;	
+		static inline TRigVMTypeIndex DoubleArray = INDEX_NONE;	
+		static inline TRigVMTypeIndex Int32Array = INDEX_NONE;	
+		static inline TRigVMTypeIndex UInt8Array = INDEX_NONE;	
+		static inline TRigVMTypeIndex FNameArray = INDEX_NONE;	
+		static inline TRigVMTypeIndex FStringArray = INDEX_NONE;	
+		static inline TRigVMTypeIndex WildCardArray = INDEX_NONE;	
 	};
 
 	// Returns true if the type specified is an array
