RigVM: Support upgrading from struct to dispatch

diff --git a/Engine/Source/Developer/RigVMDeveloper/Private/RigVMModel/RigVMController.cpp b/Engine/Source/Developer/RigVMDeveloper/Private/RigVMModel/RigVMController.cpp
--- a/Engine/Source/Developer/RigVMDeveloper/Private/RigVMModel/RigVMController.cpp
+++ b/Engine/Source/Developer/RigVMDeveloper/Private/RigVMModel/RigVMController.cpp
@@ -2170,10 +2170,13 @@ TArray<URigVMNode*> URigVMController::UpgradeNodes(const TArray<URigVMNode*>& In
 	{
 		FRigVMController_PinPathRemapDelegate RemapPinDelegate;
 		URigVMNode* UpgradedNode = UpgradeNode(Node, bSetupUndoRedo, &RemapPinDelegate);
-		UpgradedNodes.Add(UpgradedNode);
-		if(RemapPinDelegate.IsBound())
+		if(UpgradedNode)
 		{
-			RemapPinDelegates.Add(UpgradedNode->GetName(), RemapPinDelegate);
+			UpgradedNodes.Add(UpgradedNode);
+			if(RemapPinDelegate.IsBound())
+			{
+				RemapPinDelegates.Add(UpgradedNode->GetName(), RemapPinDelegate);
+			}
 		}
 	}
 
@@ -2220,91 +2223,120 @@ URigVMNode* URigVMController::UpgradeNode(URigVMNode* InNode, bool bSetupUndoRed
 
 	URigVMNode* UpgradedNode = nullptr;
 
+	const FRigVMStructUpgradeInfo UpgradeInfo = InNode->GetUpgradeInfo();
+	check(UpgradeInfo.IsValid());
+	
+	FName MethodName = TEXT("Execute");
 	if(URigVMUnitNode* UnitNode = Cast<URigVMUnitNode>(InNode))
 	{
-		const FName MethodName = UnitNode->GetMethodName();
-		
-		const FRigVMStructUpgradeInfo UpgradeInfo = UnitNode->GetUpgradeInfo();
-		check(UpgradeInfo.IsValid());
+		MethodName = UnitNode->GetMethodName();
+	}
 
-		if(OutRemapPinDelegate)
+	if(OutRemapPinDelegate)
+	{
+		*OutRemapPinDelegate = FRigVMController_PinPathRemapDelegate::CreateLambda([UpgradeInfo](const FString& InPinPath, bool bIsInput) -> FString
 		{
-			*OutRemapPinDelegate = FRigVMController_PinPathRemapDelegate::CreateLambda([UpgradeInfo](const FString& InPinPath, bool bIsInput) -> FString
-			{
-				return UpgradeInfo.RemapPin(InPinPath, bIsInput, true);
-			});
-		}
+			return UpgradeInfo.RemapPin(InPinPath, bIsInput, true);
+		});
+	}
 
-		if(!RemoveNode(InNode, bSetupUndoRedo, true, false, false))
+	if(!RemoveNode(InNode, bSetupUndoRedo, true, false, false))
+	{
+		if(bSetupUndoRedo)
 		{
-			if(bSetupUndoRedo)
-			{
-				ActionStack->CancelAction(Action, this);
-			}
-			ReportErrorf(TEXT("Unable to remove node %s."), *NodeName);
-			return nullptr;
+			ActionStack->CancelAction(Action, this);
 		}
+		ReportErrorf(TEXT("Unable to remove node %s."), *NodeName);
+		return nullptr;
+	}
 
-		URigVMNode* NewNode = AddUnitNode(UpgradeInfo.GetNewStruct(), MethodName, NodePosition, NodeName, bSetupUndoRedo, false);
-		if(NewNode == nullptr)
+	URigVMNode* NewNode = nullptr;
+	if(UpgradeInfo.GetNewStruct()->IsChildOf(FRigVMStruct::StaticStruct()))
+	{
+		NewNode = AddUnitNode(UpgradeInfo.GetNewStruct(), MethodName, NodePosition, NodeName, bSetupUndoRedo, false);
+	}
+	else if(UpgradeInfo.GetNewStruct()->IsChildOf(FRigVMDispatchFactory::StaticStruct()) && !UpgradeInfo.NewDispatchFunction.IsNone())
+	{
+		if(const FRigVMFunction* Function = FRigVMRegistry::Get().FindFunction(*UpgradeInfo.NewDispatchFunction.ToString()))
 		{
-			if(bSetupUndoRedo)
+			if(const FRigVMTemplate* Template = Function->GetTemplate())
 			{
-				ActionStack->CancelAction(Action, this);
+				if(const FRigVMDispatchFactory* Factory = Template->GetDispatchFactory())
+				{
+					if(Factory->GetScriptStruct() == UpgradeInfo.GetNewStruct())
+					{
+						NewNode = AddTemplateNode(Template->GetNotation(), NodePosition, NodeName, bSetupUndoRedo, false);
+						if(NewNode)
+						{
+							for(int32 ArgumentIndex=0;ArgumentIndex<Function->GetArguments().Num();ArgumentIndex++)
+							{
+								const FRigVMFunctionArgument& Argument = Function->GetArguments()[ArgumentIndex];
+								if(URigVMPin* Pin = NewNode->FindPin(Argument.Name))
+								{
+									if(Pin->IsWildCard())
+									{
+										ResolveWildCardPin(Pin, Function->GetArgumentTypeIndices()[ArgumentIndex], bSetupUndoRedo, false);
+									}
+								}
+							}
+						}
+					}
+				}
 			}
-			ReportErrorf(TEXT("Unable to upgrade node %s."), *NodeName);
-			return nullptr;
 		}
-
-		const TArray<FString>& AggregatePins = UpgradeInfo.GetAggregatePins();
-		for(const FString& AggregatePinName : AggregatePins)
+	}
+	
+	if(NewNode == nullptr)
+	{
+		if(bSetupUndoRedo)
 		{
-			const FName PreviousName = NewNode->GetFName();
-			AddAggregatePin(PreviousName.ToString(), AggregatePinName, FString(), bSetupUndoRedo, false);
-			NewNode = GetGraph()->FindNodeByName(PreviousName);
+			ActionStack->CancelAction(Action, this);
 		}
+		ReportErrorf(TEXT("Unable to upgrade node %s."), *NodeName);
+		return nullptr;
+	}
+
+	const TArray<FString>& AggregatePins = UpgradeInfo.GetAggregatePins();
+	for(const FString& AggregatePinName : AggregatePins)
+	{
+		const FName PreviousName = NewNode->GetFName();
+		AddAggregatePin(PreviousName.ToString(), AggregatePinName, FString(), bSetupUndoRedo, false);
+		NewNode = GetGraph()->FindNodeByName(PreviousName);
+	}
 
-		for(URigVMPin* Pin : NewNode->GetPins())
+	for(URigVMPin* Pin : NewNode->GetPins())
+	{
+		const FString DefaultValue = UpgradeInfo.GetDefaultValueForPin(Pin->GetFName());
+		if(!DefaultValue.IsEmpty())
 		{
-			const FString DefaultValue = UpgradeInfo.GetDefaultValueForPin(Pin->GetFName());
-			if(!DefaultValue.IsEmpty())
+			SetPinDefaultValue(Pin, DefaultValue, true, bSetupUndoRedo, false);
+
+			if(FPinState* PinState = PinStates.Find(Pin->GetPinPath()))
 			{
-				SetPinDefaultValue(Pin, DefaultValue, true, bSetupUndoRedo, false);
+				PinState->DefaultValue.Reset();
 			}
 		}
+	}
 
-		// redirect pin state paths
-		for(TPair<FString, FPinState>& PinState : PinStates)
+	// redirect pin state paths
+	for(TPair<FString, FPinState>& PinState : PinStates)
+	{
+		for(int32 TrueFalse = 0; TrueFalse < 2; TrueFalse++)
 		{
-			for(int32 TrueFalse = 0; TrueFalse < 2; TrueFalse++)
+			const FString RemappedInputPath = UpgradeInfo.RemapPin(PinState.Key, TrueFalse == 0, false);
+			if(RemappedInputPath != PinState.Key)
 			{
-				const FString RemappedInputPath = UpgradeInfo.RemapPin(PinState.Key, TrueFalse == 0, false);
-				if(RemappedInputPath != PinState.Key)
+				if(!RedirectedPinPaths.Contains(PinState.Key))
 				{
-					if(!RedirectedPinPaths.Contains(PinState.Key))
-					{
-						RedirectedPinPaths.Add(PinState.Key, RemappedInputPath);
-					}
+					RedirectedPinPaths.Add(PinState.Key, RemappedInputPath);
 				}
 			}
 		}
-
-
-		UpgradedNode = NewNode;
-	}
-	else
-	{
-		// for now we don't allow to upgrade anything else but unit nodes
-		checkNoEntry();
 	}
 
+	UpgradedNode = NewNode;
 	check(UpgradedNode);
 
-	// reapply the pin states but don't touch defaults
-	for(TPair<FString, FPinState>& PinState : PinStates)
-	{
-		PinState.Value.DefaultValue.Reset();
-	}
 	ApplyPinStates(UpgradedNode, PinStates, RedirectedPinPaths, bSetupUndoRedo);
 
 	if(bSetupUndoRedo)
diff --git a/Engine/Source/Runtime/RigVM/Private/RigVMCore/RigVMDispatchFactory.cpp b/Engine/Source/Runtime/RigVM/Private/RigVMCore/RigVMDispatchFactory.cpp
--- a/Engine/Source/Runtime/RigVM/Private/RigVMCore/RigVMDispatchFactory.cpp
+++ b/Engine/Source/Runtime/RigVM/Private/RigVMCore/RigVMDispatchFactory.cpp
@@ -67,7 +67,7 @@ FString FRigVMDispatchFactory::GetKeywords() const
 
 FRigVMFunctionPtr FRigVMDispatchFactory::GetDispatchFunction(const FRigVMTemplateTypeMap& InTypes) const
 {
-	const FString PermutationName = GetPermutationName(InTypes);
+	const FString PermutationName = GetPermutationNameImpl(InTypes);
 	if(const FRigVMFunction* ExistingFunction = FRigVMRegistry::Get().FindFunction(*PermutationName))
 	{
 		return ExistingFunction->FunctionPtr;
@@ -76,6 +76,19 @@ FRigVMFunctionPtr FRigVMDispatchFactory::GetDispatchFunction(const FRigVMTemplat
 }
 
 FString FRigVMDispatchFactory::GetPermutationName(const FRigVMTemplateTypeMap& InTypes) const
+{
+#if WITH_EDITOR
+	const TArray<FRigVMTemplateArgument> Arguments = GetArguments();
+	check(InTypes.Num() == Arguments.Num());
+	for(const FRigVMTemplateArgument& Argument : Arguments)
+	{
+		check(InTypes.Contains(Argument.GetName()));
+	}
+#endif
+	return GetPermutationNameImpl(InTypes);
+}
+
+FString FRigVMDispatchFactory::GetPermutationNameImpl(const FRigVMTemplateTypeMap& InTypes) const
 {
 	const FRigVMRegistry& Registry = FRigVMRegistry::Get();
 	TArray<FString> TypePairStrings;
diff --git a/Engine/Source/Runtime/RigVM/Private/RigVMCore/RigVMFunction.cpp b/Engine/Source/Runtime/RigVM/Private/RigVMCore/RigVMFunction.cpp
--- a/Engine/Source/Runtime/RigVM/Private/RigVMCore/RigVMFunction.cpp
+++ b/Engine/Source/Runtime/RigVM/Private/RigVMCore/RigVMFunction.cpp
@@ -50,22 +50,41 @@ FString FRigVMFunction::GetModuleRelativeHeaderPath() const
 	return FString();
 }
 
-const TArray<int32>& FRigVMFunction::GetArgumentTypeIndices() const
+const TArray<TRigVMTypeIndex>& FRigVMFunction::GetArgumentTypeIndices() const
 {
 	if(ArgumentTypeIndices.IsEmpty() && !Arguments.IsEmpty())
 	{
-		for(const FRigVMFunctionArgument& Argument : Arguments)
+		if(Struct)
 		{
-			if(const FProperty* Property = Struct->FindPropertyByName(Argument.Name))
+			for(const FRigVMFunctionArgument& Argument : Arguments)
 			{
-				FName CPPType = NAME_None;
-				UObject* CPPTypeObject = nullptr;
-				FRigVMExternalVariable::GetTypeFromProperty(Property, CPPType, CPPTypeObject);
+				if(const FProperty* Property = Struct->FindPropertyByName(Argument.Name))
+				{
+					FName CPPType = NAME_None;
+					UObject* CPPTypeObject = nullptr;
+					FRigVMExternalVariable::GetTypeFromProperty(Property, CPPType, CPPTypeObject);
 
-				const FRigVMTemplateArgumentType Type(CPPType, CPPTypeObject);
-				ArgumentTypeIndices.Add(FRigVMRegistry::Get().FindOrAddType(Type));
+					const FRigVMTemplateArgumentType Type(CPPType, CPPTypeObject);
+					ArgumentTypeIndices.Add(FRigVMRegistry::Get().FindOrAddType(Type));
+				}
 			}
 		}
+		else if(const FRigVMTemplate* Template = GetTemplate())
+		{
+			const int32 PermutationIndex = Template->FindPermutation(this);
+			check(PermutationIndex != INDEX_NONE);
+
+			for(const FRigVMFunctionArgument& FunctionArgument : Arguments)
+			{
+				const FRigVMTemplateArgument* TemplateArgument = Template->FindArgument(FunctionArgument.Name);
+				check(TemplateArgument);
+				ArgumentTypeIndices.Add(TemplateArgument->GetTypeIndices()[PermutationIndex]);
+			}
+		}
+		else
+		{
+			checkNoEntry();
+		}
 	}
 	return ArgumentTypeIndices;
 }
diff --git a/Engine/Source/Runtime/RigVM/Private/RigVMCore/RigVMRegistry.cpp b/Engine/Source/Runtime/RigVM/Private/RigVMCore/RigVMRegistry.cpp
--- a/Engine/Source/Runtime/RigVM/Private/RigVMCore/RigVMRegistry.cpp
+++ b/Engine/Source/Runtime/RigVM/Private/RigVMCore/RigVMRegistry.cpp
@@ -1011,6 +1011,11 @@ FRigVMDispatchFactory* FRigVMRegistry::FindDispatchFactory(const FName& InFactor
 	return nullptr;
 }
 
+FRigVMDispatchFactory* FRigVMRegistry::FindOrAddDispatchFactory(UScriptStruct* InFactoryStruct)
+{
+	return (FRigVMDispatchFactory*)RegisterFactory(InFactoryStruct);
+}
+
 const TArray<FRigVMDispatchFactory*>& FRigVMRegistry::GetFactories() const
 {
 	return Factories;
diff --git a/Engine/Source/Runtime/RigVM/Private/RigVMCore/RigVMStruct.cpp b/Engine/Source/Runtime/RigVM/Private/RigVMCore/RigVMStruct.cpp
--- a/Engine/Source/Runtime/RigVM/Private/RigVMCore/RigVMStruct.cpp
+++ b/Engine/Source/Runtime/RigVM/Private/RigVMCore/RigVMStruct.cpp
@@ -92,6 +92,49 @@ FRigVMStructUpgradeInfo::FRigVMStructUpgradeInfo()
 {
 }
 
+FRigVMStructUpgradeInfo FRigVMStructUpgradeInfo::MakeFromStructToFactory(UScriptStruct* InRigVMStruct, UScriptStruct* InFactoryStruct)
+{
+	check(InRigVMStruct);
+	check(InFactoryStruct);
+	
+	const FRigVMDispatchFactory* Factory = FRigVMRegistry::Get().FindOrAddDispatchFactory(InFactoryStruct);
+	check(Factory);
+	const FRigVMTemplateTypeMap Types = GetTypeMapFromStruct(InRigVMStruct);
+	const FString PermutationName = Factory->GetPermutationName(Types);
+
+	FRigVMStructUpgradeInfo Info;
+	Info.OldStruct = InRigVMStruct;
+	Info.NewStruct = InFactoryStruct;
+	Info.NewDispatchFunction = *PermutationName;
+	return Info;
+}
+
+FRigVMTemplateTypeMap FRigVMStructUpgradeInfo::GetTypeMapFromStruct(UScriptStruct* InScriptStruct)
+{
+	FRigVMTemplateTypeMap Types;
+#if WITH_EDITOR
+	for (TFieldIterator<FProperty> It(InScriptStruct); It; ++It)
+	{
+		if(It->HasAnyPropertyFlags(CPF_Transient))
+		{
+			continue;
+		}
+
+		if(!It->HasMetaData(FRigVMStruct::InputMetaName) &&
+			!It->HasMetaData(FRigVMStruct::OutputMetaName) &&
+			!It->HasMetaData(FRigVMStruct::VisibleMetaName))
+		{
+			continue;
+		}
+		
+		const FName PropertyName = It->GetFName();
+		const TRigVMTypeIndex TypeIndex = FRigVMTemplateArgument(*It).GetTypeIndices()[0];
+		Types.Add(PropertyName, TypeIndex);
+	}
+#endif
+	return Types;
+}
+
 bool FRigVMStructUpgradeInfo::IsValid() const
 {
 	return NodePath.IsEmpty() && (OldStruct != nullptr) && (NewStruct != nullptr);
@@ -107,6 +150,11 @@ const FString& FRigVMStructUpgradeInfo::GetDefaultValueForPin(const FName& InPin
 	return EmptyString;
 }
 
+void FRigVMStructUpgradeInfo::SetDefaultValueForPin(const FName& InPinName, const FString& InDefaultValue)
+{
+	DefaultValues.FindOrAdd(InPinName) = InDefaultValue;
+}
+
 void FRigVMStructUpgradeInfo::AddRemappedPin(const FString& InOldPinPath, const FString& InNewPinPath, bool bAsInput,
 	bool bAsOutput)
 {
diff --git a/Engine/Source/Runtime/RigVM/Private/RigVMCore/RigVMTemplate.cpp b/Engine/Source/Runtime/RigVM/Private/RigVMCore/RigVMTemplate.cpp
--- a/Engine/Source/Runtime/RigVM/Private/RigVMCore/RigVMTemplate.cpp
+++ b/Engine/Source/Runtime/RigVM/Private/RigVMCore/RigVMTemplate.cpp
@@ -896,7 +896,7 @@ const FRigVMFunction* FRigVMTemplate::GetOrCreatePermutation(int32 InIndex)
 			}
 			
 			static constexpr TCHAR Format[] = TEXT("%s::%s");
-			const FString PermutationName = Factory->GetPermutationName(Types);
+			const FString PermutationName = Factory->GetPermutationNameImpl(Types);
 			const int32 FunctionIndex = Permutations[InIndex] = Registry.Functions.Num();
 			
 			Registry.Functions.AddElement(
@@ -908,6 +908,7 @@ const FRigVMFunction* FRigVMTemplate::GetOrCreatePermutation(int32 InIndex)
 					FunctionArguments
 				)
 			);
+			Registry.Functions[FunctionIndex].TemplateIndex = Index;
 			Registry.FunctionNameToIndex.Add(*PermutationName, FunctionIndex);
 
 			return &Registry.Functions[FunctionIndex];
diff --git a/Engine/Source/Runtime/RigVM/Public/RigVMCore/RigVMDispatchFactory.h b/Engine/Source/Runtime/RigVM/Public/RigVMCore/RigVMDispatchFactory.h
--- a/Engine/Source/Runtime/RigVM/Public/RigVMCore/RigVMDispatchFactory.h
+++ b/Engine/Source/Runtime/RigVM/Public/RigVMCore/RigVMDispatchFactory.h
@@ -92,9 +92,14 @@ public:
 	// builds and returns the template
 	const FRigVMTemplate* GetTemplate() const;
 
+	// returns the name of the permutation for a given set of types
+	FString GetPermutationName(const FRigVMTemplateTypeMap& InTypes) const;
+
 protected:
 
-	FString GetPermutationName(const FRigVMTemplateTypeMap& InTypes) const;
+	// returns the name of the permutation for a given set of types
+	FString GetPermutationNameImpl(const FRigVMTemplateTypeMap& InTypes) const;
+
 	FORCEINLINE virtual FRigVMFunctionPtr GetDispatchFunctionImpl(const FRigVMTemplateTypeMap& InTypes) const { return nullptr; }
 
 	static const FString DispatchPrefix;
diff --git a/Engine/Source/Runtime/RigVM/Public/RigVMCore/RigVMFunction.h b/Engine/Source/Runtime/RigVM/Public/RigVMCore/RigVMFunction.h
--- a/Engine/Source/Runtime/RigVM/Public/RigVMCore/RigVMFunction.h
+++ b/Engine/Source/Runtime/RigVM/Public/RigVMCore/RigVMFunction.h
@@ -5,6 +5,7 @@
 #include "CoreMinimal.h"
 #include "RigVMExecuteContext.h"
 #include "RigVMMemory.h"
+#include "RigVMTypeIndex.h"
 #include "Blueprint/BlueprintSupport.h"
 
 typedef TArrayView<FRigVMMemoryHandle> FRigVMMemoryHandleArray;
@@ -76,7 +77,7 @@ struct RIGVM_API FRigVMFunction
 	int32 Index;
 	int32 TemplateIndex;
 	TArray<FRigVMFunctionArgument> Arguments;
-	mutable TArray<int32> ArgumentTypeIndices;
+	mutable TArray<TRigVMTypeIndex> ArgumentTypeIndices;
 
 	FRigVMFunction()
 		: Name()
@@ -118,7 +119,7 @@ struct RIGVM_API FRigVMFunction
 	FString GetModuleName() const;
 	FString GetModuleRelativeHeaderPath() const;
 	const TArray<FRigVMFunctionArgument>& GetArguments() const { return Arguments; }
-	const TArray<int32>& GetArgumentTypeIndices() const;
+	const TArray<TRigVMTypeIndex>& GetArgumentTypeIndices() const;
 	bool IsAdditionalArgument(const FRigVMFunctionArgument& InArgument) const;
 	const FRigVMTemplate* GetTemplate() const;
 };
diff --git a/Engine/Source/Runtime/RigVM/Public/RigVMCore/RigVMRegistry.h b/Engine/Source/Runtime/RigVM/Public/RigVMCore/RigVMRegistry.h
--- a/Engine/Source/Runtime/RigVM/Public/RigVMCore/RigVMRegistry.h
+++ b/Engine/Source/Runtime/RigVM/Public/RigVMCore/RigVMRegistry.h
@@ -175,6 +175,9 @@ public:
 	// Returns a dispatch factory given its name (or nullptr)
 	FRigVMDispatchFactory* FindDispatchFactory(const FName& InFactoryName) const;
 
+	// Returns a dispatch factory given its static struct (or nullptr)
+	FRigVMDispatchFactory* FindOrAddDispatchFactory(UScriptStruct* InFactoryStruct);
+
 	// Returns all dispatch factories
 	const TArray<FRigVMDispatchFactory*>& GetFactories() const;
 
diff --git a/Engine/Source/Runtime/RigVM/Public/RigVMCore/RigVMStructUpgradeInfo.h b/Engine/Source/Runtime/RigVM/Public/RigVMCore/RigVMStructUpgradeInfo.h
--- a/Engine/Source/Runtime/RigVM/Public/RigVMCore/RigVMStructUpgradeInfo.h
+++ b/Engine/Source/Runtime/RigVM/Public/RigVMCore/RigVMStructUpgradeInfo.h
@@ -3,6 +3,7 @@
 #pragma once
 
 #include "CoreMinimal.h"
+#include "RigVMTemplate.h"
 
 struct FRigVMStruct;
 
@@ -14,15 +15,25 @@ public:
 	FRigVMStructUpgradeInfo();
 
 	template<typename Old, typename New>
-	FRigVMStructUpgradeInfo(const Old& InOld, const New& InNew, const FName& InOldDispatchFunction = NAME_None, const FName& InNewDispatchFunction = NAME_None)
+	FRigVMStructUpgradeInfo(const Old& InOld, const New& InNew)
 	{
 		OldStruct = Old::StaticStruct();
 		NewStruct = New::StaticStruct();
+		OldDispatchFunction = NAME_None;
+		NewDispatchFunction = NAME_None;
+		SetDefaultValues(&InNew);
+	}
+
+	FRigVMStructUpgradeInfo(UScriptStruct* InOldDispatchStruct, UScriptStruct* InNewDispatchStruct, const FName& InOldDispatchFunction, const FName& InNewDispatchFunction)
+	{
+		OldStruct = InOldDispatchStruct;
+		NewStruct = InNewDispatchStruct;
 		OldDispatchFunction = InOldDispatchFunction;
 		NewDispatchFunction = InNewDispatchFunction;
-		SetDefaultValues(&InNew);
 	}
 
+	static FRigVMStructUpgradeInfo MakeFromStructToFactory(UScriptStruct* InRigVMStruct, UScriptStruct* InFactoryStruct);
+
 	// returns true if this upgrade info can be applied
 	bool IsValid() const;
 
@@ -38,6 +49,9 @@ public:
 	// returns the default value for a given pin
 	const FString& GetDefaultValueForPin(const FName& InPinName) const;
 
+	// sets the default value for a given pin
+	void SetDefaultValueForPin(const FName& InPinName, const FString& InDefaultValue);
+
 	// adds a pin to be remapped
 	void AddRemappedPin(const FString& InOldPinPath, const FString& InNewPinPath, bool bAsInput = true, bool bAsOutput = true);
 
@@ -50,6 +64,9 @@ public:
 	// returns the aggregate pins to add
 	const TArray<FString>& GetAggregatePins() const { return AggregatePins; }
 
+	// returns a type map representing the struct members
+	static FRigVMTemplateTypeMap GetTypeMapFromStruct(UScriptStruct* InScriptStruct);
+
 private:
 	
 	// The complete node path including models / collapse node.
diff --git a/Engine/Source/Runtime/RigVM/Public/RigVMCore/RigVMTemplate.h b/Engine/Source/Runtime/RigVM/Public/RigVMCore/RigVMTemplate.h
--- a/Engine/Source/Runtime/RigVM/Public/RigVMCore/RigVMTemplate.h
+++ b/Engine/Source/Runtime/RigVM/Public/RigVMCore/RigVMTemplate.h
@@ -217,6 +217,7 @@ protected:
 	friend struct FRigVMTemplate;
 	friend class URigVMController;
 	friend struct FRigVMRegistry;
+	friend struct FRigVMStructUpgradeInfo;
 };
 
 /**
