RigVM: Correctly follow pins across the collapse node boundary

diff --git a/Engine/Source/Developer/RigVMDeveloper/Private/RigVMCompiler/RigVMAST.cpp b/Engine/Source/Developer/RigVMDeveloper/Private/RigVMCompiler/RigVMAST.cpp
--- a/Engine/Source/Developer/RigVMDeveloper/Private/RigVMCompiler/RigVMAST.cpp
+++ b/Engine/Source/Developer/RigVMDeveloper/Private/RigVMCompiler/RigVMAST.cpp
@@ -3118,8 +3118,12 @@ void FRigVMParserAST::Inline(TArray<URigVMGraph*> InGraphs, const TArray<FRigVMA
 						{
 							if (!ShouldRecursePin(ParentSourcePinProxy))
 							{
-								SourcePinProxy = FRigVMASTProxy();
-								break;
+								// only discard results here if we haven't crossed a collapse node boundary
+								if(ParentSourcePinProxy.GetSubjectChecked<URigVMPin>()->GetGraph() == ChildPin->GetGraph())
+								{
+									SourcePinProxy = FRigVMASTProxy();
+									break;
+								}
 							}
 						}
 					}
