Control Rig: Collapse nodes with templates

diff --git a/Engine/Plugins/Animation/ControlRig/Source/ControlRig/Public/ControlRigObjectVersion.h b/Engine/Plugins/Animation/ControlRig/Source/ControlRig/Public/ControlRigObjectVersion.h
--- a/Engine/Plugins/Animation/ControlRig/Source/ControlRig/Public/ControlRigObjectVersion.h
+++ b/Engine/Plugins/Animation/ControlRig/Source/ControlRig/Public/ControlRigObjectVersion.h
@@ -72,6 +72,9 @@ struct CONTROLRIG_API FControlRigObjectVersion
 		// Added rig hierarchy element metadata
 		HierarchyElementMetadata,
 
+		// Converted library nodes to templates
+		LibraryNodeTemplates,
+
 		// -----<new versions can be added above this line>-------------------------------------------------
 		VersionPlusOne,
 		LatestVersion = VersionPlusOne - 1
diff --git a/Engine/Plugins/Animation/ControlRig/Source/ControlRigDeveloper/Private/ControlRigBlueprint.cpp b/Engine/Plugins/Animation/ControlRig/Source/ControlRigDeveloper/Private/ControlRigBlueprint.cpp
--- a/Engine/Plugins/Animation/ControlRig/Source/ControlRigDeveloper/Private/ControlRigBlueprint.cpp
+++ b/Engine/Plugins/Animation/ControlRig/Source/ControlRigDeveloper/Private/ControlRigBlueprint.cpp
@@ -1357,7 +1357,8 @@ void UControlRigBlueprint::RefreshAllModels(EControlRigBlueprintLoadType InLoadT
 	for (URigVMGraph* GraphToDetach : GraphsToDetach)
 	{
 		URigVMController* Controller = GetOrCreateController(GraphToDetach);
-		Controller->ReattachLinksToPinObjects(true /* follow redirectors */, nullptr, false, true);
+		bool bAllowNonArgumentLinks = GetLinkerCustomVersion(FControlRigObjectVersion::GUID) < FControlRigObjectVersion::LibraryNodeTemplates;
+		Controller->ReattachLinksToPinObjects(true /* follow redirectors */, nullptr, false, true, bAllowNonArgumentLinks);
 	}
 
 	if(bIsPostLoad)
@@ -1366,11 +1367,66 @@ void UControlRigBlueprint::RefreshAllModels(EControlRigBlueprintLoadType InLoadT
 	}
 	
 	TArray<URigVMGraph*> GraphsToClean = GetAllModels();
+
+	// Sort from leaf graphs to root
+	TArray<URigVMGraph*> SortedGraphsToClean;
+	SortedGraphsToClean.Reserve(GraphsToClean.Num());
+	while (SortedGraphsToClean.Num() < GraphsToClean.Num())
+	{
+		bool bGraphAdded = false;
+		for (URigVMGraph* Graph : GraphsToClean)
+		{
+			if (SortedGraphsToClean.Contains(Graph))
+			{
+				continue;
+			}
+			
+			TArray<URigVMGraph*> ContainedGraphs;
+			for (URigVMNode* Node : Graph->GetNodes())
+			{
+				if (URigVMLibraryNode* LibraryNode = Cast<URigVMLibraryNode>(Node))
+				{
+					if (URigVMFunctionReferenceNode* FunctionReferenceNode = Cast<URigVMFunctionReferenceNode>(LibraryNode))
+					{
+						if (FunctionReferenceNode->GetLibrary() != GetLocalFunctionLibrary())
+						{
+							UPackage* Package = FunctionReferenceNode->GetLibrary()->GetPackage();
+							if (!Package->IsFullyLoaded())
+							{
+								Package->FullyLoad();
+							}
+							continue;
+						}
+					}
+					
+					ContainedGraphs.Add(LibraryNode->GetContainedGraph());					
+				}
+			}
+			
+			bool bAllContained = true;
+			for (URigVMGraph* Contained : ContainedGraphs)
+			{
+				if (!SortedGraphsToClean.Contains(Contained))
+				{
+					bAllContained = false;
+					break;
+				}
+			}
+			if (bAllContained)
+			{
+				SortedGraphsToClean.Add(Graph);
+				bGraphAdded = true;
+			}
+		}
+		ensure(bGraphAdded);
+	}
 	
-	for(int32 GraphIndex=0; GraphIndex<GraphsToClean.Num(); GraphIndex++)
+	for(int32 GraphIndex=0; GraphIndex<SortedGraphsToClean.Num(); GraphIndex++)
 	{
-		URigVMGraph* GraphToClean = GraphsToClean[GraphIndex];
+		URigVMGraph* GraphToClean = SortedGraphsToClean[GraphIndex];
 		URigVMController* Controller = GetOrCreateController(GraphToClean);
+		TGuardValue<bool> RecomputeGuard(Controller->bSuspendRecomputingOuterTemplateFilters, true);
+		
 		for(URigVMNode* ModelNode : GraphToClean->GetNodes())
 		{
 			Controller->RemoveUnusedOrphanedPins(ModelNode, false);
@@ -1389,16 +1445,21 @@ void UControlRigBlueprint::RefreshAllModels(EControlRigBlueprintLoadType InLoadT
 						!TemplateNode->IsSingleton() &&
 						TemplateNode->PreferredPermutationPairs.IsEmpty())
 					{
-						TemplateNode->InitializeFilteredPermutationsFromTypes();
+						TemplateNode->InitializeFilteredPermutationsFromTypes(false);
 					}
 				}
 			}
+
+			if (URigVMLibraryNode* LibraryNode = GraphToClean->GetTypedOuter<URigVMLibraryNode>())
+			{
+				Controller->UpdateLibraryTemplate(LibraryNode, false);
+			}
 		}
 #if WITH_EDITOR
 
 		if(bIsPostLoad)
 		{
-			if (Controller->RecomputeAllTemplateFilteredTypes(false))
+			if (Controller->RecomputeAllTemplateFilteredPermutations(false))
 			{
 				ensureMsgf(false, TEXT("Pin type changed during load %s"), *GetPackage()->GetPathName());
 			}
diff --git a/Engine/Plugins/Animation/ControlRig/Source/ControlRigDeveloper/Private/Graph/ControlRigGraphNode.cpp b/Engine/Plugins/Animation/ControlRig/Source/ControlRigDeveloper/Private/Graph/ControlRigGraphNode.cpp
--- a/Engine/Plugins/Animation/ControlRig/Source/ControlRigDeveloper/Private/Graph/ControlRigGraphNode.cpp
+++ b/Engine/Plugins/Animation/ControlRig/Source/ControlRigDeveloper/Private/Graph/ControlRigGraphNode.cpp
@@ -314,13 +314,18 @@ void UControlRigGraphNode::RewireOldPinsToNewPins(TArray<UEdGraphPin*>& InOldPin
 	{
 		for(UEdGraphPin* NewPin : InNewPins)
 		{
-			if(OldPin->PinName == NewPin->PinName && OldPin->PinType == NewPin->PinType && OldPin->Direction == NewPin->Direction)
+			if(OldPin->PinName == NewPin->PinName && OldPin->Direction == NewPin->Direction)
 			{
-				// make sure to remove invalid entries from the linked to list
-				OldPin->LinkedTo.Remove(nullptr);
+				if (OldPin->PinType == NewPin->PinType ||
+					OldPin->PinType.PinSubCategoryObject == RigVMTypeUtils::GetWildCardCPPTypeObject() ||
+					NewPin->PinType.PinSubCategoryObject == RigVMTypeUtils::GetWildCardCPPTypeObject())
+				{
+					// make sure to remove invalid entries from the linked to list
+					OldPin->LinkedTo.Remove(nullptr);
 				
-				NewPin->MovePersistentDataFromOldPin(*OldPin);
-				break;
+					NewPin->MovePersistentDataFromOldPin(*OldPin);
+					break;
+				}
 			}
 		}
 	}
diff --git a/Engine/Plugins/Animation/ControlRig/Source/ControlRigEditor/Private/Graph/SControlRigGraphNode.cpp b/Engine/Plugins/Animation/ControlRig/Source/ControlRigEditor/Private/Graph/SControlRigGraphNode.cpp
--- a/Engine/Plugins/Animation/ControlRig/Source/ControlRigEditor/Private/Graph/SControlRigGraphNode.cpp
+++ b/Engine/Plugins/Animation/ControlRig/Source/ControlRigEditor/Private/Graph/SControlRigGraphNode.cpp
@@ -693,6 +693,19 @@ void SControlRigGraphNode::AddPin(const TSharedRef<SGraphPin>& PinToAdd)
 					PinToAdd->SetCustomPinIcon(CachedImg_CR_Pin_Connected, CachedImg_CR_Pin_Disconnected);
 				}
 				PinToAdd->SetToolTipText(ModelPin->GetToolTipText());
+
+				// If the pin belongs to a template node that does not own an argument for that pin, make it transparent
+				if (URigVMTemplateNode* TemplateNode = Cast<URigVMTemplateNode>(ModelPin->GetNode()))
+				{
+					if (const FRigVMTemplate* Template = TemplateNode->GetTemplate())
+					{
+						URigVMPin* RootPin = ModelPin->GetRootPin();
+						if (Template->FindArgument(RootPin->GetFName()) == nullptr)
+						{
+							PinToAdd->SetColorAndOpacity(PinToAdd->GetColorAndOpacity() * FLinearColor(1.0f,1.0f,1.0f,0.2f));
+						}
+					}
+				}
 			}
 		}
 
diff --git a/Engine/Source/Developer/RigVMDeveloper/Private/RigVMModel/Nodes/RigVMFunctionEntryNode.cpp b/Engine/Source/Developer/RigVMDeveloper/Private/RigVMModel/Nodes/RigVMFunctionEntryNode.cpp
--- a/Engine/Source/Developer/RigVMDeveloper/Private/RigVMModel/Nodes/RigVMFunctionEntryNode.cpp
+++ b/Engine/Source/Developer/RigVMDeveloper/Private/RigVMModel/Nodes/RigVMFunctionEntryNode.cpp
@@ -3,6 +3,24 @@
 #include "RigVMModel/Nodes/RigVMFunctionEntryNode.h"
 #include "RigVMModel/RigVMFunctionLibrary.h"
 
+const FRigVMTemplate* URigVMFunctionEntryNode::GetTemplate() const
+{
+	if (URigVMLibraryNode* LibraryNode = GetTypedOuter<URigVMLibraryNode>())
+	{
+		return LibraryNode->GetTemplate();
+	}
+	return nullptr;
+}
+
+FName URigVMFunctionEntryNode::GetNotation() const
+{
+	if (URigVMLibraryNode* LibraryNode = GetTypedOuter<URigVMLibraryNode>())
+	{
+		return LibraryNode->GetNotation();
+	}
+	return NAME_None;
+}
+
 FLinearColor URigVMFunctionEntryNode::GetNodeColor() const
 {
 	if(URigVMGraph* RootGraph = GetRootGraph())
@@ -20,3 +38,18 @@ bool URigVMFunctionEntryNode::IsDefinedAsVarying() const
 	// todo
 	return true; 
 }
+
+FString URigVMFunctionEntryNode::GetNodeTitle() const
+{
+	return TEXT("Entry");
+}
+
+FText URigVMFunctionEntryNode::GetToolTipText() const
+{
+	return FText::FromName(GetTemplate()->GetNotation());
+}
+
+FText URigVMFunctionEntryNode::GetToolTipTextForPin(const URigVMPin* InPin) const
+{
+	return Super::GetToolTipTextForPin(InPin);
+}
diff --git a/Engine/Source/Developer/RigVMDeveloper/Private/RigVMModel/Nodes/RigVMFunctionReferenceNode.cpp b/Engine/Source/Developer/RigVMDeveloper/Private/RigVMModel/Nodes/RigVMFunctionReferenceNode.cpp
--- a/Engine/Source/Developer/RigVMDeveloper/Private/RigVMModel/Nodes/RigVMFunctionReferenceNode.cpp
+++ b/Engine/Source/Developer/RigVMDeveloper/Private/RigVMModel/Nodes/RigVMFunctionReferenceNode.cpp
@@ -167,6 +167,11 @@ TArray<FRigVMExternalVariable> URigVMFunctionReferenceNode::GetExternalVariables
 	return GetExternalVariables(true);
 }
 
+const FRigVMTemplate* URigVMFunctionReferenceNode::GetTemplate() const
+{
+	return GetReferencedNode()->GetTemplate();
+}
+
 TArray<FRigVMExternalVariable> URigVMFunctionReferenceNode::GetExternalVariables(bool bRemapped) const
 {
 	TArray<FRigVMExternalVariable> Variables;
diff --git a/Engine/Source/Developer/RigVMDeveloper/Private/RigVMModel/Nodes/RigVMFunctionReturnNode.cpp b/Engine/Source/Developer/RigVMDeveloper/Private/RigVMModel/Nodes/RigVMFunctionReturnNode.cpp
--- a/Engine/Source/Developer/RigVMDeveloper/Private/RigVMModel/Nodes/RigVMFunctionReturnNode.cpp
+++ b/Engine/Source/Developer/RigVMDeveloper/Private/RigVMModel/Nodes/RigVMFunctionReturnNode.cpp
@@ -3,6 +3,24 @@
 #include "RigVMModel/Nodes/RigVMFunctionReturnNode.h"
 #include "RigVMModel/RigVMFunctionLibrary.h"
 
+const FRigVMTemplate* URigVMFunctionReturnNode::GetTemplate() const
+{
+	if (URigVMLibraryNode* LibraryNode = GetTypedOuter<URigVMLibraryNode>())
+	{
+		return LibraryNode->GetTemplate();
+	}
+	return nullptr;
+}
+
+FName URigVMFunctionReturnNode::GetNotation() const
+{
+	if (URigVMLibraryNode* LibraryNode = GetTypedOuter<URigVMLibraryNode>())
+	{
+		return LibraryNode->GetNotation();
+	}
+	return NAME_None;
+}
+
 FLinearColor URigVMFunctionReturnNode::GetNodeColor() const
 {
 	if(URigVMGraph* RootGraph = GetRootGraph())
@@ -20,3 +38,18 @@ bool URigVMFunctionReturnNode::IsDefinedAsVarying() const
 	// todo
 	return true; 
 }
+
+FString URigVMFunctionReturnNode::GetNodeTitle() const
+{
+	return TEXT("Return");
+}
+
+FText URigVMFunctionReturnNode::GetToolTipText() const
+{
+	return FText::FromName(GetTemplate()->GetNotation());
+}
+
+FText URigVMFunctionReturnNode::GetToolTipTextForPin(const URigVMPin* InPin) const
+{
+	return Super::GetToolTipTextForPin(InPin);
+}
diff --git a/Engine/Source/Developer/RigVMDeveloper/Private/RigVMModel/Nodes/RigVMLibraryNode.cpp b/Engine/Source/Developer/RigVMDeveloper/Private/RigVMModel/Nodes/RigVMLibraryNode.cpp
--- a/Engine/Source/Developer/RigVMDeveloper/Private/RigVMModel/Nodes/RigVMLibraryNode.cpp
+++ b/Engine/Source/Developer/RigVMDeveloper/Private/RigVMModel/Nodes/RigVMLibraryNode.cpp
@@ -10,6 +10,11 @@
 const TArray<URigVMNode*> URigVMLibraryNode::EmptyNodes;
 const TArray<URigVMLink*> URigVMLibraryNode::EmptyLinks;
 
+URigVMLibraryNode::URigVMLibraryNode()
+{
+	Template.Notation = *FString::Printf(TEXT("%s()"), *GetName());
+}
+
 bool URigVMLibraryNode::IsDefinedAsConstant() const
 {
 	return !IsDefinedAsVarying();
@@ -31,6 +36,21 @@ bool URigVMLibraryNode::IsDefinedAsVarying() const
 	return false;
 }
 
+const FRigVMTemplate* URigVMLibraryNode::GetTemplate() const
+{
+	return &Template;
+}
+
+FName URigVMLibraryNode::GetNotation() const
+{
+	return Template.GetNotation();
+}
+
+FText URigVMLibraryNode::GetToolTipText() const
+{
+	return FText::FromName(Template.GetNotation());
+}
+
 TArray<int32> URigVMLibraryNode::GetInstructionsForVMImpl(URigVM* InVM, const FRigVMASTProxy& InProxy) const
 {
 	TArray<int32> Instructions;
diff --git a/Engine/Source/Developer/RigVMDeveloper/Private/RigVMModel/Nodes/RigVMTemplateNode.cpp b/Engine/Source/Developer/RigVMDeveloper/Private/RigVMModel/Nodes/RigVMTemplateNode.cpp
--- a/Engine/Source/Developer/RigVMDeveloper/Private/RigVMModel/Nodes/RigVMTemplateNode.cpp
+++ b/Engine/Source/Developer/RigVMDeveloper/Private/RigVMModel/Nodes/RigVMTemplateNode.cpp
@@ -125,6 +125,18 @@ void URigVMTemplateNode::ConvertPreferredTypesToTypeIndex()
 	}
 }
 
+TRigVMTypeIndex URigVMTemplateNode::GetPreferredType(const FName& ArgumentName) const
+{
+	for (const FRigVMTemplatePreferredType& PreferredType : PreferredPermutationPairs)
+	{
+		if (PreferredType.Argument == ArgumentName)
+		{
+			return PreferredType.TypeIndex;
+		}
+	}
+	return INDEX_NONE;
+}
+
 UScriptStruct* URigVMTemplateNode::GetScriptStruct() const
 {
 	if(const FRigVMFunction* Function = GetResolvedFunction())
@@ -185,9 +197,8 @@ FText URigVMTemplateNode::GetToolTipTextForPin(const URigVMPin* InPin) const
 			if (const FRigVMTemplateArgument* Arg = Template->FindArgument(RootPin->GetFName()))
 			{
 				FString Tooltip;
-
-				const TArray<int32> PermutationIndices = GetFilteredPermutationsIndices();
-				if(PermutationIndices.Num() == GetTemplate()->NumPermutations())
+				
+				if(FilteredPermutations.Num() == GetTemplate()->NumPermutations())
 				{
 					if(Arg->GetTypeIndices().Num() > 100)
 					{
@@ -198,9 +209,15 @@ FText URigVMTemplateNode::GetToolTipTextForPin(const URigVMPin* InPin) const
 				if(Tooltip.IsEmpty())
 				{
 					FString SupportedTypesJoined;
-					for (int32 Index : PermutationIndices)
+					TArray<TRigVMTypeIndex> TypesPrinted;
+					for (int32 Index=0; Index<Template->NumPermutations(); ++Index)
 					{
 						const TRigVMTypeIndex TypeIndex = Arg->GetTypeIndices()[Index];
+						if (TypesPrinted.Contains(TypeIndex))
+						{
+							continue;
+						}
+						TypesPrinted.Add(TypeIndex);
 						FString Type = FRigVMRegistry::Get().GetType(TypeIndex).CPPType.ToString();
 						if (!FilteredPermutations.Contains(Index))
 						{
@@ -223,6 +240,99 @@ FText URigVMTemplateNode::GetToolTipTextForPin(const URigVMPin* InPin) const
 	return SuperToolTip;
 }
 
+TArray<URigVMPin*> URigVMTemplateNode::GetAggregateInputs() const
+{
+	return GetAggregatePins(ERigVMPinDirection::Input);
+}
+
+TArray<URigVMPin*> URigVMTemplateNode::GetAggregateOutputs() const
+{
+	return GetAggregatePins(ERigVMPinDirection::Output);
+}
+
+TArray<URigVMPin*> URigVMTemplateNode::GetAggregatePins(const ERigVMPinDirection& InDirection) const
+{
+	TArray<URigVMPin*> AggregatePins;
+#if UE_RIGVM_AGGREGATE_NODES_ENABLED
+	if (const FRigVMTemplate* Template = GetTemplate())
+	{
+		if (FilteredPermutations.IsEmpty())
+		{
+			return AggregatePins;
+		}
+
+		auto FindAggregatePins = [&](const int32& PermutationIndex)
+		{
+			TArray<URigVMPin*> Inputs;
+			if (const FRigVMFunction* Permutation = Template->GetPermutation(PermutationIndex))
+			{
+				const TSharedPtr<FStructOnScope> StructOnScope = MakeShareable(new FStructOnScope(Permutation->Struct));
+				if (const UScriptStruct* ScriptStruct = Cast<UScriptStruct>(StructOnScope->GetStruct()))
+				{
+					for (URigVMPin* Pin : GetPins())
+					{
+						if (Pin->GetDirection() == InDirection)
+						{
+							if (const FProperty* Property = ScriptStruct->FindPropertyByName(Pin->GetFName()))
+							{
+								if (Property->HasMetaData(FRigVMStruct::AggregateMetaName))
+								{
+									Inputs.Add(Pin);
+								}
+							}			
+						}
+					}
+				}
+			}
+			return Inputs;
+		};
+		
+		AggregatePins = FindAggregatePins(FilteredPermutations[0]);
+		for (int32 i=1; i<FilteredPermutations.Num(); ++i)
+		{
+			TArray<URigVMPin*> OtherAggregateInputs = FindAggregatePins(FilteredPermutations[i]);
+			if (OtherAggregateInputs.Num() != AggregatePins.Num() ||
+				OtherAggregateInputs.FilterByPredicate([&](const URigVMPin* OtherPin)
+				{
+					return !AggregatePins.Contains(OtherPin);
+				}).Num() > 0)
+			{
+				return {};
+			}
+		}
+	}
+#endif
+	return AggregatePins;
+}
+
+FName URigVMTemplateNode::GetNextAggregateName(const FName& InLastAggregatePinName) const
+{
+	FName NextName;	
+
+#if UE_RIGVM_AGGREGATE_NODES_ENABLED
+	if (const FRigVMTemplate* Template = GetTemplate())
+	{
+		if (FilteredPermutations.IsEmpty())
+		{
+			return NextName;
+		}
+
+		if (const FRigVMFunction* Permutation = Template->GetPermutation(FilteredPermutations[0]))
+		{
+			const TSharedPtr<FStructOnScope> StructOnScope = MakeShareable(new FStructOnScope(Permutation->Struct));
+			if (const UScriptStruct* ScriptStruct = Cast<UScriptStruct>(StructOnScope->GetStruct()))
+			{
+				check(ScriptStruct->IsChildOf(FRigVMStruct::StaticStruct()));
+
+				const FRigVMStruct* StructMemory = (const FRigVMStruct*)StructOnScope->GetStructMemory();
+				return StructMemory->GetNextAggregateName(InLastAggregatePinName);
+			}
+		}
+	}
+#endif
+	return FName();
+}
+
 FName URigVMTemplateNode::GetNotation() const
 {
 	return TemplateNotation;
@@ -282,6 +392,14 @@ bool URigVMTemplateNode::SupportsType(const URigVMPin* InPin, TRigVMTypeIndex In
 				}
 				return true;
 			}
+			else if (IsA<URigVMFunctionEntryNode>() || IsA<URigVMFunctionReturnNode>())
+			{
+				if(OutTypeIndex)
+				{
+					*OutTypeIndex = InTypeIndex;
+				}
+				return true;
+			}
 		}
 		return false;
 	}
@@ -307,6 +425,15 @@ bool URigVMTemplateNode::SupportsType(const URigVMPin* InPin, TRigVMTypeIndex In
 			}
 			return true;
 		}
+
+		// an entry/return node that does not contain an argument for the pin will always support the connections
+		if (IsA<URigVMFunctionEntryNode>() || IsA<URigVMFunctionReturnNode>())
+		{
+			if (Template->FindArgument(RootPin->GetFName()) == nullptr)
+			{
+				return true;
+			}
+		}
 		return false;
 	}
 	
@@ -449,6 +576,12 @@ FString URigVMTemplateNode::GetInitialDefaultValueForPin(const FName& InRootPinN
 	{
 		return FString();
 	}
+
+	URigVMPin* Pin = FindPin(InRootPinName.ToString());
+	if (!Pin)
+	{
+		return FString();
+	}
 	
 	TArray<int32> PermutationIndices = InPermutationIndices;
 	if(PermutationIndices.IsEmpty())
@@ -459,92 +592,95 @@ FString URigVMTemplateNode::GetInitialDefaultValueForPin(const FName& InRootPinN
 	FString DefaultValue;
 	const FRigVMTemplate* Template = GetTemplate();
 	const FRigVMDispatchFactory* Factory = Template->GetDispatchFactory(); 
+	const FRigVMTemplateArgument* Argument = Template->FindArgument(InRootPinName);
 
-	for(const int32 PermutationIndex : PermutationIndices)
+	if (Argument)
 	{
-		FString NewDefaultValue;
-
-		const FRigVMTemplateArgument* Argument = Template->FindArgument(InRootPinName);
-		const TRigVMTypeIndex TypeIndex = Argument->GetTypeIndices()[PermutationIndex];
-
-		if(Factory)
+		for(const int32 PermutationIndex : PermutationIndices)
 		{
-			NewDefaultValue = Factory->GetArgumentDefaultValue(Argument->GetName(), TypeIndex);
-		}
-		else if(const FRigVMFunction* Permutation = Template->GetPermutation(PermutationIndex))
-		{
-			const TSharedPtr<FStructOnScope> StructOnScope = MakeShareable(new FStructOnScope(Permutation->Struct));
-			const FRigVMStruct* DefaultStruct = (const FRigVMStruct*)StructOnScope->GetStructMemory();
+			FString NewDefaultValue;
 
-			const bool bUseQuotes = TypeIndex != RigVMTypeUtils::TypeIndex::FString && TypeIndex != RigVMTypeUtils::TypeIndex::FName;
-			NewDefaultValue = DefaultStruct->ExportToFullyQualifiedText(
-				Cast<UScriptStruct>(StructOnScope->GetStruct()), InRootPinName, nullptr, bUseQuotes);
-		}
-		else
-		{
-			if (FRigVMRegistry::Get().IsArrayType(TypeIndex))
+			const TRigVMTypeIndex TypeIndex = Argument->GetTypeIndices()[PermutationIndex];
+
+			if(Factory)
 			{
-				NewDefaultValue = TEXT("()");
+				NewDefaultValue = Factory->GetArgumentDefaultValue(Argument->GetName(), TypeIndex);
 			}
-			else
+			else if(const FRigVMFunction* Permutation = Template->GetPermutation(PermutationIndex))
 			{
-				const FRigVMTemplateArgumentType& Type = FRigVMRegistry::Get().GetType(TypeIndex);
-				if (UScriptStruct* ScriptStruct = Cast<UScriptStruct>(Type.CPPTypeObject))
-				{
-					TArray<uint8, TAlignedHeapAllocator<16>> TempBuffer;
-					TempBuffer.AddUninitialized(ScriptStruct->GetStructureSize());
+				const TSharedPtr<FStructOnScope> StructOnScope = MakeShareable(new FStructOnScope(Permutation->Struct));
+				const FRigVMStruct* DefaultStruct = (const FRigVMStruct*)StructOnScope->GetStructMemory();
 
-					// call the struct constructor to initialize the struct
-					ScriptStruct->InitializeDefaultValue(TempBuffer.GetData());
-
-					ScriptStruct->ExportText(NewDefaultValue, TempBuffer.GetData(), nullptr, nullptr, PPF_None, nullptr);
-					ScriptStruct->DestroyStruct(TempBuffer.GetData());				
-				}
-				else if (const UEnum* Enum = Cast<UEnum>(Type.CPPTypeObject))
-				{
-					NewDefaultValue = Enum->GetNameStringByValue(0);
-				}
-				else if(const UClass* Class = Cast<UClass>(Type.CPPTypeObject))
-				{
-					// not supporting objects yet
-					ensure(false);
-				}
-				else if (Type.CPPType == RigVMTypeUtils::FloatTypeName)
-				{
-					NewDefaultValue = TEXT("0.000000");				
-				}
-				else if (Type.CPPType == RigVMTypeUtils::DoubleTypeName)
-				{
-					NewDefaultValue = TEXT("0.000000");
-				}
-				else if (Type.CPPType == RigVMTypeUtils::Int32TypeName)
+				const bool bUseQuotes = TypeIndex != RigVMTypeUtils::TypeIndex::FString && TypeIndex != RigVMTypeUtils::TypeIndex::FName;
+				NewDefaultValue = DefaultStruct->ExportToFullyQualifiedText(
+					Cast<UScriptStruct>(StructOnScope->GetStruct()), InRootPinName, nullptr, bUseQuotes);
+			}
+			else
+			{
+				if (FRigVMRegistry::Get().IsArrayType(TypeIndex))
 				{
-					NewDefaultValue = TEXT("0");
+					NewDefaultValue = TEXT("()");
 				}
-				else if (Type.CPPType == RigVMTypeUtils::BoolTypeName)
+				else
 				{
-					NewDefaultValue = TEXT("False");
-				}
-				else if (Type.CPPType == RigVMTypeUtils::FStringTypeName)
+					const FRigVMTemplateArgumentType& Type = FRigVMRegistry::Get().GetType(TypeIndex);
+					if (UScriptStruct* ScriptStruct = Cast<UScriptStruct>(Type.CPPTypeObject))
+					{
+						TArray<uint8, TAlignedHeapAllocator<16>> TempBuffer;
+						TempBuffer.AddUninitialized(ScriptStruct->GetStructureSize());
+
+						// call the struct constructor to initialize the struct
+						ScriptStruct->InitializeDefaultValue(TempBuffer.GetData());
+
+						ScriptStruct->ExportText(NewDefaultValue, TempBuffer.GetData(), nullptr, nullptr, PPF_None, nullptr);
+						ScriptStruct->DestroyStruct(TempBuffer.GetData());				
+					}
+					else if (const UEnum* Enum = Cast<UEnum>(Type.CPPTypeObject))
+					{
+						NewDefaultValue = Enum->GetNameStringByValue(0);
+					}
+					else if(const UClass* Class = Cast<UClass>(Type.CPPTypeObject))
+					{
+						// not supporting objects yet
+						ensure(false);
+					}
+					else if (Type.CPPType == RigVMTypeUtils::FloatTypeName)
+					{
+						NewDefaultValue = TEXT("0.000000");				
+					}
+					else if (Type.CPPType == RigVMTypeUtils::DoubleTypeName)
+					{
+						NewDefaultValue = TEXT("0.000000");
+					}
+					else if (Type.CPPType == RigVMTypeUtils::Int32TypeName)
+					{
+						NewDefaultValue = TEXT("0");
+					}
+					else if (Type.CPPType == RigVMTypeUtils::BoolTypeName)
+					{
+						NewDefaultValue = TEXT("False");
+					}
+					else if (Type.CPPType == RigVMTypeUtils::FStringTypeName)
+					{
+						NewDefaultValue = TEXT("");
+					}
+					else if (Type.CPPType == RigVMTypeUtils::FNameTypeName)
+					{
+						NewDefaultValue = TEXT("");
+					}
+				}			
+			}
+
+			if(!NewDefaultValue.IsEmpty())
+			{
+				if(DefaultValue.IsEmpty())
 				{
-					NewDefaultValue = TEXT("");
+					DefaultValue = NewDefaultValue;
 				}
-				else if (Type.CPPType == RigVMTypeUtils::FNameTypeName)
+				else if(!NewDefaultValue.Equals(DefaultValue))
 				{
-					NewDefaultValue = TEXT("");
+					return FString();
 				}
-			}			
-		}
-
-		if(!NewDefaultValue.IsEmpty())
-		{
-			if(DefaultValue.IsEmpty())
-			{
-				DefaultValue = NewDefaultValue;
-			}
-			else if(!NewDefaultValue.Equals(DefaultValue))
-			{
-				return FString();
 			}
 		}
 	}
@@ -596,6 +732,11 @@ TArray<TRigVMTypeIndex> URigVMTemplateNode::GetFilteredTypesForPin(URigVMPin* In
 		return FilteredTypes;	
 	}
 
+	if (InPin->IsStructMember())
+	{
+		return {InPin->GetTypeIndex()};
+	}
+
 	if (!PreferredPermutationPairs.IsEmpty())
 	{
 		for (const FRigVMTemplatePreferredType& PreferredType : PreferredPermutationPairs)
@@ -607,9 +748,16 @@ TArray<TRigVMTypeIndex> URigVMTemplateNode::GetFilteredTypesForPin(URigVMPin* In
 				{
 					if (TypeIndex == PreferredType.GetTypeIndex())
 					{
-						return {TypeIndex};						
+						if (InPin->IsArrayElement())
+						{
+							return {FRigVMRegistry::Get().GetArrayTypeFromBaseTypeIndex(TypeIndex)};
+						}
+						else
+						{
+							return {TypeIndex};
+						}
 					}
-				}				
+				}	
 			}
 		}
 	}
@@ -635,6 +783,38 @@ TArray<TRigVMTypeIndex> URigVMTemplateNode::GetFilteredTypesForPin(URigVMPin* In
 	return FilteredTypes;
 }
 
+bool URigVMTemplateNode::TryReduceTypesToSingle(TArray<TRigVMTypeIndex>& InOutTypes, const TRigVMTypeIndex PreferredType) const
+{
+	if (InOutTypes.IsEmpty())
+	{
+		return false;
+	}
+
+	if (InOutTypes.Num() == 1)
+	{
+		return true;
+	}
+	
+	for (int32 i=1; i<InOutTypes.Num(); ++i)
+	{
+		if (!FRigVMRegistry::Get().CanMatchTypes(InOutTypes[0], InOutTypes[i], true))
+		{
+			return false;
+		}
+	}
+
+	if (InOutTypes.Contains(PreferredType))
+	{
+		InOutTypes = {PreferredType};
+	}
+	else
+	{
+		InOutTypes = {InOutTypes[0]};
+	}
+			
+	return true;	
+}
+
 TArray<int32> URigVMTemplateNode::GetNewFilteredPermutations(URigVMPin* InPin, URigVMPin* LinkedPin)
 {	
 	TArray<int32> NewFilteredPermutations;
@@ -784,6 +964,7 @@ TArray<int32> URigVMTemplateNode::FindPermutationsForTypes(const TArray<FRigVMTe
 			for (int32 ArgIndex = 0; ArgIndex < Args.Num(); ++ArgIndex)
 			{
 				const FRigVMTemplateArgument* Argument = Args[ArgIndex];
+				if (Argument)
 				{
 					if ((bAllowCasting && !Registry.CanMatchTypes(Argument->GetTypeIndices()[i], TypeIndices[ArgIndex], true)) ||
 						(!bAllowCasting && Argument->GetTypeIndices()[i] != TypeIndices[ArgIndex]))
@@ -792,6 +973,10 @@ TArray<int32> URigVMTemplateNode::FindPermutationsForTypes(const TArray<FRigVMTe
 						break;
 					}
 				}
+				else
+				{
+					// The preferred type doesn't own an argument yet. Supports all types.
+				}
 			}
 
 			if (bAllArgsMatched)
@@ -817,8 +1002,7 @@ TArray<FRigVMTemplatePreferredType> URigVMTemplateNode::GetPreferredTypesForPerm
 			}
 			else
 			{
-				ArgTypes.Reset();
-				return ArgTypes;
+				return {};
 			}
 		}
 	}
@@ -919,7 +1103,7 @@ void URigVMTemplateNode::InitializeFilteredPermutations()
 	{
 		if (!PreferredPermutationPairs.IsEmpty())
 		{
-			FilteredPermutations = FindPermutationsForTypes(PreferredPermutationPairs);
+			FilteredPermutations = FindPermutationsForTypes(PreferredPermutationPairs, true);
 		}
 		else
 		{
@@ -932,37 +1116,53 @@ void URigVMTemplateNode::InitializeFilteredPermutations()
 	}	
 }
 
-void URigVMTemplateNode::InitializeFilteredPermutationsFromTypes()
+void URigVMTemplateNode::InitializeFilteredPermutationsFromTypes(bool bAllowCasting)
 {
 	if (IsSingleton())
 	{
 		return;
 	}
 
-	if (const FRigVMTemplate* Template = GetTemplate())
+	if (IsA<URigVMFunctionEntryNode>() || IsA<URigVMFunctionReturnNode>())
 	{
 		TArray<FRigVMTemplatePreferredType> ArgTypes;
-		for (int32 ArgIndex = 0; ArgIndex < Template->NumArguments(); ++ArgIndex)
+		for (URigVMPin* Pin : GetPins())
 		{
-			const FRigVMTemplateArgument* Argument = Template->GetArgument(ArgIndex);
-			if (URigVMPin* Pin = FindPin(Argument->GetName().ToString()))
+			if (!Pin->IsWildCard())
+			{
+				ArgTypes.Emplace(Pin->GetFName(), Pin->GetTypeIndex());
+			}
+		}
+		
+		PreferredPermutationPairs = ArgTypes;
+	}
+	else
+	{
+		if (const FRigVMTemplate* Template = GetTemplate())
+		{
+			TArray<FRigVMTemplatePreferredType> ArgTypes;
+			for (int32 ArgIndex = 0; ArgIndex < Template->NumArguments(); ++ArgIndex)
 			{
-				if (!Pin->IsWildCard())
+				const FRigVMTemplateArgument* Argument = Template->GetArgument(ArgIndex);
+				if (URigVMPin* Pin = FindPin(Argument->GetName().ToString()))
 				{
-					ArgTypes.Emplace(Argument->GetName(), Pin->GetTypeIndex());
+					if (!Pin->IsWildCard())
+					{
+						ArgTypes.Emplace(Argument->GetName(), Pin->GetTypeIndex());
+					}
 				}
 			}
-		}
 
-		const TArray<int32> Permutations = FindPermutationsForTypes(ArgTypes);
-		if (!Permutations.IsEmpty())
-		{
-			FilteredPermutations = Permutations;
-			PreferredPermutationPairs = ArgTypes;
-		}
-		else
-		{
-			InitializeFilteredPermutations();
+			const TArray<int32> Permutations = FindPermutationsForTypes(ArgTypes, bAllowCasting);
+			if (!Permutations.IsEmpty())
+			{
+				FilteredPermutations = Permutations;
+				PreferredPermutationPairs = ArgTypes;
+			}
+			else
+			{
+				InitializeFilteredPermutations();
+			}
 		}
 	}
 }
diff --git a/Engine/Source/Developer/RigVMDeveloper/Private/RigVMModel/Nodes/RigVMUnitNode.cpp b/Engine/Source/Developer/RigVMDeveloper/Private/RigVMModel/Nodes/RigVMUnitNode.cpp
--- a/Engine/Source/Developer/RigVMDeveloper/Private/RigVMModel/Nodes/RigVMUnitNode.cpp
+++ b/Engine/Source/Developer/RigVMDeveloper/Private/RigVMModel/Nodes/RigVMUnitNode.cpp
@@ -185,6 +185,10 @@ TArray<URigVMPin*> URigVMUnitNode::GetAggregateInputs() const
 			}
 		}
 	}
+	else
+	{
+		return Super::GetAggregateInputs();
+	}
 #endif
 	return AggregateInputs;
 }
@@ -209,6 +213,10 @@ TArray<URigVMPin*> URigVMUnitNode::GetAggregateOutputs() const
 			}
 		}
 	}
+	else
+	{
+		return Super::GetAggregateOutputs();
+	}
 #endif
 	return AggregateOutputs;
 }
@@ -224,6 +232,10 @@ FName URigVMUnitNode::GetNextAggregateName(const FName& InLastAggregatePinName)
 		const FRigVMStruct* StructMemory = (const FRigVMStruct*)StructOnScope->GetStructMemory();
 		return StructMemory->GetNextAggregateName(InLastAggregatePinName);
 	}
+	else
+	{
+		return Super::GetNextAggregateName(InLastAggregatePinName);
+	}
 #endif
 	return FName();
 }
diff --git a/Engine/Source/Developer/RigVMDeveloper/Private/RigVMModel/RigVMController.cpp b/Engine/Source/Developer/RigVMDeveloper/Private/RigVMModel/RigVMController.cpp
--- a/Engine/Source/Developer/RigVMDeveloper/Private/RigVMModel/RigVMController.cpp
+++ b/Engine/Source/Developer/RigVMDeveloper/Private/RigVMModel/RigVMController.cpp
@@ -63,6 +63,7 @@ FRigVMControllerCompileBracketScope::~FRigVMControllerCompileBracketScope()
 
 URigVMController::URigVMController()
 	: bValidatePinDefaults(true)
+	, bSuspendRecomputingOuterTemplateFilters(false)
 	, bSuspendNotifications(false)
 	, bReportWarningsAndErrors(true)
 	, bIgnoreRerouteCompactnessChanges(false)
@@ -77,6 +78,7 @@ URigVMController::URigVMController()
 URigVMController::URigVMController(const FObjectInitializer& ObjectInitializer)
 	: Super(ObjectInitializer)
 	, bValidatePinDefaults(true)
+	, bSuspendRecomputingOuterTemplateFilters(false)
 	, bSuspendNotifications(false)
 	, bReportWarningsAndErrors(true)
 	, bIgnoreRerouteCompactnessChanges(false)
@@ -826,7 +828,7 @@ URigVMUnitNode* URigVMController::AddUnitNode(UScriptStruct* InScriptStruct, con
 		const TArray<FRigVMTemplatePreferredType> NewPreferredPermutationTypes = TemplateNode->GetPreferredTypesForPermutation(PermutationIndex);
 		if (bSetupUndoRedo)
 		{
-			FRigVMSetTemplateFilteredPermutationsAction Action(TemplateNode, nullptr, OldPermutations);
+			FRigVMSetTemplateFilteredPermutationsAction Action(TemplateNode, OldPermutations);
 			ActionStack->AddAction(Action);
 			ActionStack->AddAction(FRigVMSetPreferredTemplatePermutationsAction(TemplateNode, NewPreferredPermutationTypes));
 		}		
@@ -1774,15 +1776,12 @@ bool URigVMController::UnresolveTemplateNodes(const TArray<FName>& InNodeNames,
 		return false;
 	}
 
-	TArray<URigVMTemplateNode*> Nodes;
+	TArray<URigVMNode*> Nodes;
 	for (const FName& NodeName : InNodeNames)
 	{
-		if (URigVMTemplateNode* Node = Cast<URigVMTemplateNode>(GetGraph()->FindNodeByName(NodeName)))
+		if (URigVMNode* Node = GetGraph()->FindNodeByName(NodeName))
 		{
-			if (!Node->IsSingleton())
-			{
-				Nodes.Add(Node);
-			}
+			Nodes.Add(Node);			
 		}
 	}
 
@@ -1811,7 +1810,7 @@ bool URigVMController::UnresolveTemplateNodes(const TArray<FName>& InNodeNames,
 	return false;
 }
 
-bool URigVMController::UnresolveTemplateNodes(const TArray<URigVMTemplateNode*>& InNodes, bool bSetupUndoRedo)
+bool URigVMController::UnresolveTemplateNodes(const TArray<URigVMNode*>& InNodes, bool bSetupUndoRedo)
 {
 	if (!IsValidGraph() || InNodes.IsEmpty())
 	{
@@ -1825,9 +1824,13 @@ bool URigVMController::UnresolveTemplateNodes(const TArray<URigVMTemplateNode*>&
 	}
 	
 	// check if any of the nodes needs to be unresolved
-	const bool bHasNodeToResolve = InNodes.ContainsByPredicate( [](const URigVMTemplateNode* Node) -> bool
+	const bool bHasNodeToResolve = InNodes.ContainsByPredicate( [](const URigVMNode* Node) -> bool
 	{
-		return !Node->IsFullyUnresolved();
+		if (const URigVMTemplateNode* TemplateNode = Cast<URigVMTemplateNode>(Node))
+		{
+			return !TemplateNode->IsFullyUnresolved();
+		}
+		return false;
 	});
 	if (!bHasNodeToResolve)
 	{
@@ -1841,222 +1844,45 @@ bool URigVMController::UnresolveTemplateNodes(const TArray<URigVMTemplateNode*>&
 		ActionStack->BeginAction(Action);
 	}
 
-	// find all links affecting the nodes to unresolve
-	TArray<URigVMNode*> Nodes;
-	Algo::Transform(InNodes, Nodes, [](URigVMTemplateNode* InNode) -> URigVMNode* { return InNode; });
-	TArray<TPair<FString, FString>> LinkedPaths = GetLinkedPinPaths(Nodes);
-
-	const FRigVMRegistry& Registry = FRigVMRegistry::Get();
-	
-	// Find pins outside our set of nodes which were resolved to a type
-	TMap<URigVMPin*, TRigVMTypeIndex> PinsToResolve;
-	for (TPair<FString, FString>& LinkPath : LinkedPaths)
-	{
-		const FString& SourcePath = LinkPath.Key;
-		const FString& TargetPath = LinkPath.Value;
-
-		URigVMPin* SourcePin = GetGraph()->FindPin(SourcePath);
-		URigVMPin* TargetPin = GetGraph()->FindPin(TargetPath);
-
-		bool bSourceOutside = false;
-		URigVMNode* SourceNode = SourcePin->GetNode();
-		if (!Nodes.Contains(SourceNode))
-		{
-			bSourceOutside = true;
-			if (SourceNode->IsInjected())
-			{
-				if (URigVMNode* OutNode = SourceNode->GetTypedOuter<URigVMNode>())
-				{
-					if (Nodes.Contains(OutNode))
-					{
-						bSourceOutside = false;
-					}
-				}
-			}
-		}
-		bool bTargetOutside = false;
-		URigVMNode* TargetNode = TargetPin->GetNode();
-		if (!Nodes.Contains(TargetNode))
-		{
-			bTargetOutside = true;
-			if (TargetNode->IsInjected())
-			{
-				if (URigVMNode* OutNode = TargetNode->GetTypedOuter<URigVMNode>())
-				{
-					if (Nodes.Contains(OutNode))
-					{
-						bTargetOutside = false;
-					}
-				}
-			}
-		}
-		if (bSourceOutside && !bTargetOutside)
-		{
-			if (!SourcePin->IsWildCard())
-			{
-				PinsToResolve.Add(SourcePin, SourcePin->GetTypeIndex());
-			}
-		}
-		else if(bTargetOutside && !bSourceOutside)
-		{
-			if (!TargetPin->IsWildCard())
-			{
-				PinsToResolve.Add(TargetPin, TargetPin->GetTypeIndex());
-			}
-		}			
-	}
-
-	bool bChangedAnyPin = false;
-	bool bChangedFilteredPermutations = false;
-	for(URigVMTemplateNode* Node : InNodes)
 	{
-		if (Node->IsSingleton())
-		{
-			continue;
-		}
-
-		if (!Node->PreferredPermutationPairs.IsEmpty())
-		{
-			if (bSetupUndoRedo)
-			{
-				ActionStack->AddAction(FRigVMSetPreferredTemplatePermutationsAction(Node, {}));
-			}
-			Node->PreferredPermutationPairs = {};
-		}
-		
-		TArray<int32> OldPermutations = Node->FilteredPermutations;
-		Node->InitializeFilteredPermutations();
-		if (bSetupUndoRedo)
-		{
-			FRigVMSetTemplateFilteredPermutationsAction FilteringAction(Node, nullptr, OldPermutations);
-			ActionStack->AddAction(FilteringAction);
-		}
-
-		bChangedFilteredPermutations = OldPermutations != Node->FilteredPermutations;			
-		
-		if(!bChangedFilteredPermutations)
+		TGuardValue<bool> GuardRecompute(bSuspendRecomputingTemplateFilters, true);
+		for (URigVMNode* Node : InNodes)
 		{
-			continue;
-		}
-
-		EjectAllInjectedNodes(Node, bSetupUndoRedo);
-
-		int32 PermutationIndex = INDEX_NONE;
-		FRigVMTemplate::FTypeMap Types;
-		Node->GetTemplate()->FullyResolve(Types, PermutationIndex);
+			EjectAllInjectedNodes(Node, bSetupUndoRedo);
 
-		// if the types' containers don't match the current pin containers try again
-		bool bTypesWereAdjusted = false;
-		FRigVMTemplate::FTypeMap AdjustedTypes = Types;
-		for(TPair<FName, TRigVMTypeIndex>& TypePair : AdjustedTypes)
-		{
-			if(URigVMPin* Pin = Node->FindPin(TypePair.Key.ToString()))
+			if (URigVMTemplateNode* TemplateNode = Cast<URigVMTemplateNode>(Node))
 			{
-				if(Pin->IsArray() != Registry.IsArrayType(TypePair.Value))
+				if (!TemplateNode->PreferredPermutationPairs.IsEmpty())
 				{
-					if(Pin->IsArray())
-					{
-						TypePair.Value = Registry.GetArrayTypeFromBaseTypeIndex(TypePair.Value);
-					}
-					else
+					if (bSetupUndoRedo)
 					{
-						TypePair.Value = Registry.GetBaseTypeFromArrayTypeIndex(TypePair.Value);
+						ActionStack->AddAction(FRigVMSetPreferredTemplatePermutationsAction(TemplateNode, {}));
 					}
-					bTypesWereAdjusted = true;
-				}
+					TemplateNode->PreferredPermutationPairs = {};
+				}	
 			}
-		}
-
-		if(bTypesWereAdjusted)
-		{
-			Node->GetTemplate()->FullyResolve(AdjustedTypes, PermutationIndex);
-			if(AdjustedTypes.Num() == Types.Num())
+	
+			TArray<URigVMLink*> Links = Node->GetLinks();
+			for (int32 i=0; i<Links.Num(); ++i)
 			{
-				Types = AdjustedTypes;
-			}
-		}
-
-		for(const TPair<FName, TRigVMTypeIndex>& TypePair : Types)
-		{
-			const FName& PinName = TypePair.Key;
-			const TRigVMTypeIndex& ExpectedTypeIndex = TypePair.Value;
+				URigVMLink* Link = Links[i];
+				URigVMPin* SourcePin = Link->GetSourcePin();
+				URigVMPin* TargetPin = Link->GetTargetPin();
 			
-			URigVMPin* Pin = Node->FindPin(PinName.ToString());
-			if(Pin == nullptr)
-			{
-				ReportErrorf(TEXT("Cannot find expected pin '%s' on Template Node '%s'"), *PinName.ToString(), *Node->GetNodePath());
-				if(bSetupUndoRedo)
-				{
-					ActionStack->CancelAction(Action, this);
-				}
-				return false;
-			}
-
-			if(ChangePinType(Pin, ExpectedTypeIndex, bSetupUndoRedo, false, true, !Pin->IsArray()))
-			{
-				bChangedAnyPin = true;
-
-				if(Pin->IsArray())
+				URigVMNode* OtherNode = SourcePin->GetNode() == Node ? TargetPin->GetNode() : SourcePin->GetNode();
+				if (!InNodes.Contains(OtherNode))
 				{
-					const TRigVMTypeIndex ElementTypeIndex = FRigVMRegistry::Get().GetBaseTypeFromArrayTypeIndex(ExpectedTypeIndex);
-					
-					TArray<URigVMPin*> SubPins = Pin->GetSubPins();
-					for (URigVMPin* SubPin : SubPins)
-					{
-						ChangePinType(SubPin, ElementTypeIndex, bSetupUndoRedo, false, true, true);
-					}
-				}
+					BreakLink(SourcePin, TargetPin, bSetupUndoRedo);
+				}			
 			}
 		}
-
-		ResolveTemplateNodeMetaData(Node, bSetupUndoRedo);
 	}
-
-	RestoreLinkedPaths(LinkedPaths, {}, {},
-		FRigVMController_CheckPinComatibilityDelegate::CreateLambda([](URigVMPin* A, URigVMPin* B) -> bool
-		{
-			return A->IsWildCard() == B->IsWildCard();
-		}), bSetupUndoRedo);
 	
-	for (TPair<URigVMPin*, TRigVMTypeIndex>& Pin : PinsToResolve)
-	{
-		TArray<URigVMLink*> LinksToBreak;
-		for (URigVMLink* Link : Pin.Key->GetLinks())
-		{
-			if (URigVMPin* OppositePin = Link->GetOppositePin(Pin.Key))
-			{
-				if (Nodes.Contains(OppositePin->GetNode()))
-				{
-					if (OppositePin->GetTypeIndex() != Pin.Value)
-					{
-						LinksToBreak.Add(Link);
-					}
-				}
-			}
-		}
-
-		for (URigVMLink* Link : LinksToBreak)
-		{
-			BreakLink(Link->GetSourcePin(), Link->GetTargetPin(), bSetupUndoRedo);
-		}
-
-		if (Pin.Key->IsWildCard())
-		{
-			ResolveWildCardPin(Pin.Key, Pin.Value, bSetupUndoRedo);
-		}
-	}
+	RecomputeAllTemplateFilteredPermutations(bSetupUndoRedo);
 
 	if(bSetupUndoRedo)
 	{
-		if(bChangedAnyPin || bChangedFilteredPermutations)
-		{
-			ActionStack->EndAction(Action);
-		}
-		else
-		{
-			ActionStack->CancelAction(Action, this);
-			return false;
-		}
+		ActionStack->EndAction(Action);		
 	}
 
 	return true;
@@ -3591,6 +3417,7 @@ TArray<FName> URigVMController::ImportNodesFromText(const FString& InText, bool
 	}
 
 	FRigVMUnitNodeCreatedContext::FScope UnitNodeCreatedScope(UnitNodeCreatedContext, ERigVMNodeCreatedReason::Paste);
+	TMap<URigVMPin*, TRigVMTypeIndex> TypesWhenPasted;
 	for (URigVMNode* CreatedNode : Factory.CreatedNodes)
 	{
 		if(!CanAddNode(CreatedNode, true))
@@ -3598,14 +3425,6 @@ TArray<FName> URigVMController::ImportNodesFromText(const FString& InText, bool
 			continue;
 		}
 
-		if(URigVMTemplateNode* TemplateNode = Cast<URigVMTemplateNode>(CreatedNode))
-		{
-			// make sure the update preferred type pairs
-			// from string to index - since indices may
-			// have changed between sessions
-			TemplateNode->ConvertPreferredTypesToTypeIndex();
-		}
-
 		Graph->Nodes.Add(CreatedNode);
 
 		if (bSetupUndoRedo)
@@ -3631,6 +3450,14 @@ TArray<FName> URigVMController::ImportNodesFromText(const FString& InText, bool
 
 		for(int32 SubNodeIndex=0; SubNodeIndex < SubNodes.Num(); SubNodeIndex++)
 		{
+			if(URigVMTemplateNode* TemplateNode = Cast<URigVMTemplateNode>(CreatedNode))
+			{
+				// make sure the update preferred type pairs
+				// from string to index - since indices may
+				// have changed between sessions
+				TemplateNode->ConvertPreferredTypesToTypeIndex();
+			}
+		
 			if(URigVMCollapseNode* CollapseNode = Cast<URigVMCollapseNode>(SubNodes[SubNodeIndex]))
 			{
 				{
@@ -3643,6 +3470,32 @@ TArray<FName> URigVMController::ImportNodesFromText(const FString& InText, bool
 			}
 		}
 
+		if (URigVMLibraryNode* LibraryNode = Cast<URigVMLibraryNode>(CreatedNode))
+		{
+			TGuardValue<bool> GuardRecompute(bSuspendRecomputingOuterTemplateFilters, true);			
+			TFunction<void(URigVMLibraryNode*)> Recompute = [&](URigVMLibraryNode* OuterNode)
+			{
+				for (URigVMNode* Node : OuterNode->GetContainedNodes())
+				{
+					if (URigVMLibraryNode* ContainedNode = Cast<URigVMLibraryNode>(Node))
+					{
+						Recompute(ContainedNode);
+					}
+				}
+				
+				{
+					FRigVMControllerGraphGuard GraphGuard(this, OuterNode->GetGraph(), bSetupUndoRedo);
+					UpdateLibraryTemplate(OuterNode, bSetupUndoRedo);
+				}
+				{
+					FRigVMControllerGraphGuard GraphGuard(this, OuterNode->GetContainedGraph(), bSetupUndoRedo);
+					RecomputeAllTemplateFilteredPermutations(bSetupUndoRedo);
+				}
+				
+			};
+			Recompute(LibraryNode);
+		}
+
 		for(URigVMNode* SubNode : SubNodes)
 		{
 			if (URigVMUnitNode* UnitNode = Cast<URigVMUnitNode>(SubNode))
@@ -3660,10 +3513,11 @@ TArray<FName> URigVMController::ImportNodesFromText(const FString& InText, bool
 
 			if (URigVMTemplateNode* TemplateNode = Cast<URigVMTemplateNode>(SubNode))
 			{
-				if (!TemplateNode->IsSingleton())
+				if (!TemplateNode->IsSingleton() && !TemplateNode->IsA<URigVMLibraryNode>())
 				{
-					TemplateNode->InitializeFilteredPermutationsFromTypes();
-					UpdateTemplateNodePinTypes(TemplateNode, false);
+					FRigVMControllerGraphGuard GraphGuard(this, TemplateNode->GetGraph(), false);
+					TemplateNode->InitializeFilteredPermutations();
+					UpdateTemplateNodePinTypes(TemplateNode, false, false);					
 				}
 			}
 
@@ -3723,7 +3577,7 @@ TArray<FName> URigVMController::ImportNodesFromText(const FString& InText, bool
 						if (OriginalNode && OriginalNode->IsA<URigVMFunctionEntryNode>())
 						{
 							CreatedLink->SourcePinPath = URigVMPin::JoinPinPath(SourceLeft, SourceRight);
-							SourcePin = CreatedLink->GetSourcePin();							
+							SourcePin = CreatedLink->GetSourcePin();
 						}
 					}
 					if (TargetPin == nullptr)
@@ -3732,7 +3586,54 @@ TArray<FName> URigVMController::ImportNodesFromText(const FString& InText, bool
 						if (OriginalNode && OriginalNode->IsA<URigVMFunctionReturnNode>())
 						{
 							CreatedLink->TargetPinPath = URigVMPin::JoinPinPath(TargetLeft, TargetRight);
-							TargetPin = CreatedLink->GetTargetPin();							
+							TargetPin = CreatedLink->GetTargetPin();
+						}
+					}
+
+					if (SourcePin == nullptr)
+					{
+						URigVMNode* OriginalNode = Graph->FindNode(SourceLeft);
+						if (OriginalNode)
+						{
+							FString OldSourcePinPath = URigVMPin::JoinPinPath(SourceLeft, SourceRight);
+							if (URigVMPin* OldPin = Graph->FindPin(OldSourcePinPath))
+							{
+								if (OldPin->IsStructMember())
+								{
+									URigVMPin* OldRootPin = OldPin->GetRootPin();
+									FString NewSourceRootPinPath = URigVMPin::JoinPinPath(NewSourceNodeName->ToString(), OldRootPin->GetName());
+									if (URigVMPin* NewRootPin = Graph->FindPin(NewSourceRootPinPath))
+									{
+										if (PrepareTemplatePinForType(NewRootPin, {OldRootPin->GetTypeIndex()}, bSetupUndoRedo))
+										{
+											SourcePin = CreatedLink->GetSourcePin();
+										}										
+									}
+								}
+							}							
+						}
+					}
+					if (TargetPin == nullptr)
+					{
+						URigVMNode* OriginalNode = Graph->FindNode(TargetLeft);
+						if (OriginalNode)
+						{
+							FString OldTargetPinPath = URigVMPin::JoinPinPath(TargetLeft, TargetRight);
+							if (URigVMPin* OldPin = Graph->FindPin(OldTargetPinPath))
+							{
+								if (OldPin->IsStructMember())
+								{
+									URigVMPin* OldRootPin = OldPin->GetRootPin();
+									FString NewTargetRootPinPath = URigVMPin::JoinPinPath(NewTargetNodeName->ToString(), OldRootPin->GetName());
+									if (URigVMPin* NewRootPin = Graph->FindPin(NewTargetRootPinPath))
+									{
+										if (PrepareTemplatePinForType(NewRootPin, {OldRootPin->GetTypeIndex()}, bSetupUndoRedo))
+										{
+											TargetPin = CreatedLink->GetTargetPin();
+										}
+									}
+								}
+							}							
 						}
 					}
 					
@@ -3770,29 +3671,8 @@ TArray<FName> URigVMController::ImportNodesFromText(const FString& InText, bool
 						}
 						else
 						{
-							if (URigVMTemplateNode* FirstTemplateNode = Cast<URigVMTemplateNode>(TargetPin->GetNode()))
-							{
-								if (!FirstTemplateNode->IsSingleton())
-								{
-									TArray<TRigVMTypeIndex> InputTypes = GetWildcardFilteredTypeIndices(SourcePin);
-									if (InputTypes.Num() > 0)
-									{
-										PrepareTemplatePinForType(TargetPin, InputTypes, bSetupUndoRedo);										
-									}
-								}
-							}
-							if (URigVMTemplateNode* SecondTemplateNode = Cast<URigVMTemplateNode>(SourcePin->GetNode()))
-							{
-								if (!SecondTemplateNode->IsSingleton())
-								{
-									TArray<TRigVMTypeIndex> OutTypes = GetWildcardFilteredTypeIndices(TargetPin);
-									if (OutTypes.Num() > 0)
-									{
-										PrepareTemplatePinForType(SourcePin, OutTypes, bSetupUndoRedo);										
-									}
-								}
-							}
-							
+							PrepareToLink(TargetPin, SourcePin, bSetupUndoRedo);
+
 							Graph->Links.Add(CreatedLink);
 							SourcePin->Links.Add(CreatedLink);
 							TargetPin->Links.Add(CreatedLink);
@@ -4356,6 +4236,7 @@ URigVMCollapseNode* URigVMController::CollapseNodes(const TArray<URigVMNode*>& I
 		return AllLinks.Find(A) < AllLinks.Find(B);
 	});
 
+	TGuardValue<bool> RecomputeOuterTemplateFilter(bSuspendRecomputingOuterTemplateFilters, true);
 	FRigVMControllerCompileBracketScope CompileScope(this);
 	FRigVMCollapseNodesAction CollapseAction;
 
@@ -4499,6 +4380,22 @@ URigVMCollapseNode* URigVMController::CollapseNodes(const TArray<URigVMNode*>& I
 
 			Notify(ERigVMGraphNotifType::NodeAdded, ReturnNode);
 		}
+
+		for (int32 i=0; i<2; ++i)
+		{
+			URigVMTemplateNode* InterfaceNode = (i==0) ? (URigVMTemplateNode*) EntryNode : (URigVMTemplateNode*) ReturnNode;
+			if (InterfaceNode)
+			{
+				for (URigVMPin* Pin : InterfaceNode->GetPins())
+				{
+					if (!Pin->IsWildCard())
+					{
+						AddPreferredType(InterfaceNode, Pin->GetFName(), Pin->GetTypeIndex(), bSetupUndoRedo);
+						AddArgumentForPin(Pin, nullptr, bSetupUndoRedo);
+					}
+				}
+			}
+		}
 	}
 
 	// create the new nodes within the collapse node
@@ -4738,6 +4635,8 @@ URigVMCollapseNode* URigVMController::CollapseNodes(const TArray<URigVMNode*>& I
 		}
 	}
 
+	RecomputeAllTemplateFilteredPermutations(bSuspendNotifications);
+	
 	if (bSetupUndoRedo)
 	{
 		ActionStack->EndAction(CollapseAction);
@@ -5539,6 +5438,7 @@ URigVMFunctionReferenceNode* URigVMController::PromoteCollapseNodeToFunctionRefe
 
 	FRigVMControllerCompileBracketScope CompileScope(this);
 	URigVMFunctionReferenceNode* FunctionRefNode = nullptr;
+	TGuardValue<bool> GuardRecomputeOuter(bSuspendRecomputingOuterTemplateFilters, true);
 
 	// Create Function
 	URigVMLibraryNode* FunctionDefinition = nullptr;
@@ -5622,7 +5522,8 @@ URigVMFunctionReferenceNode* URigVMController::PromoteCollapseNodeToFunctionRefe
 	// Remove collapse node, add function reference, and add external links
 	if (FunctionDefinition)
 	{
-		FString NodeName = InCollapseNode->GetName();
+		TGuardValue<bool> GuardRecomputeInner(bSuspendRecomputingTemplateFilters, true);
+ 		FString NodeName = InCollapseNode->GetName();
 		FVector2D NodePosition = InCollapseNode->GetPosition();
 		TMap<FString, FPinState> PinStates = GetPinStates(InCollapseNode);
 
@@ -5748,6 +5649,18 @@ URigVMCollapseNode* URigVMController::PromoteFunctionReferenceNodeToCollapseNode
 	URigVMCollapseNode* CollapseNode = DuplicateObject<URigVMCollapseNode>(FunctionDefinition, Graph, *NodeName);
 	if(CollapseNode)
 	{
+		// Copy the library template
+		{
+			FRigVMByteArray ArchiveBytes;
+			FMemoryWriter ArchiveWriter(ArchiveBytes);
+			FunctionDefinition->Template.Serialize(ArchiveWriter);
+
+			FMemoryReader ArchiveReader(ArchiveBytes);
+			CollapseNode->Template.Serialize(ArchiveReader);
+
+			CollapseNode->InitializeFilteredPermutations();
+		}
+	
 		{
 			FRigVMControllerGraphGuard Guard(this, CollapseNode->GetContainedGraph(), false);
 			ReattachLinksToPinObjects();
@@ -7900,6 +7813,7 @@ FString URigVMController::AddAggregatePin(URigVMNode* InNode, const FString& InP
 	{
 		FRigVMControllerGraphGuard GraphGuard(this, AggregateNode->GetContainedGraph(), bSetupUndoRedo);
 		TGuardValue<bool> GuardEditGraph(GetGraph()->bEditable, true);
+		TGuardValue<bool> GuardOuterTemplateRecompute(bSuspendRecomputingOuterTemplateFilters, true);
 
 		URigVMNode* InnerNode = (AggregateNode == nullptr) ? InNode : AggregateNode->GetFirstInnerNode();
 
@@ -7956,7 +7870,7 @@ FString URigVMController::AddAggregatePin(URigVMNode* InNode, const FString& InP
 		URigVMPin* NewUnitPinOppositeArg = NewNode->GetOppositeAggregatePin();
 		URigVMNode* PreviousNode = nullptr;
 		if(AggregateNode->IsInputAggregate())
-		{		
+		{
 			URigVMFunctionEntryNode* EntryNode = AggregateNode->GetEntryNode();
 			URigVMPin* EntryPin = EntryNode->FindPin(NewExposedPinName.ToString());
 			URigVMPin* ReturnPin = AggregateNode->GetReturnNode()->FindPin(NewUnitPinOppositeArg->GetName());
@@ -8012,6 +7926,8 @@ FString URigVMController::AddAggregatePin(URigVMNode* InNode, const FString& InP
 		return FString();
 	}
 
+	UpdateTemplateNodePinTypes(AggregateNode, bSetupUndoRedo);
+
 	if (bSetupUndoRedo)
 	{
 		ActionStack->EndAction(Action);
@@ -8114,6 +8030,7 @@ bool URigVMController::RemoveAggregatePin(URigVMPin* InPin, bool bSetupUndoRedo,
 							AggregateNode->InvalidateCache();
 						}
 
+						TGuardValue<bool> GuardOuterTemplateRecompute(bSuspendRecomputingOuterTemplateFilters, true);
 						const FString FirstAggregatePin = AggregateNode->GetFirstAggregatePin()->GetName();
 						const FString SecondAggregatePin = AggregateNode->GetSecondAggregatePin()->GetName();
 						FString OtherArg = TargetPin->GetName() == FirstAggregatePin ? SecondAggregatePin : FirstAggregatePin;
@@ -8149,6 +8066,7 @@ bool URigVMController::RemoveAggregatePin(URigVMPin* InPin, bool bSetupUndoRedo,
 							AggregateNode->InvalidateCache();
 						}
 
+						TGuardValue<bool> GuardOuterTemplateRecompute(bSuspendRecomputingOuterTemplateFilters, true);
 						const FString FirstAggregatePin = AggregateNode->GetFirstAggregatePin()->GetName();
 						const FString SecondAggregatePin = AggregateNode->GetSecondAggregatePin()->GetName();
 						FString OtherArg = SourcePin->GetName() == FirstAggregatePin ? SecondAggregatePin : FirstAggregatePin;
@@ -9276,6 +9194,36 @@ bool URigVMController::AddLink(URigVMPin* OutputPin, URigVMPin* InputPin, bool b
 		BreakAllLinks(OutputPin, false, bSetupUndoRedo);
 	}
 
+	bool bAddedTemporaryOutputPreferredType = false, bAddedTemporaryInputPreferredType = false;
+	for (int32 i=0; i<2; ++i)
+	{
+		URigVMPin* EdgePin = (i==0) ? OutputPin : InputPin;
+		bool& bAddedTemporaryPreferredType = (i==0) ? bAddedTemporaryOutputPreferredType : bAddedTemporaryInputPreferredType;
+		if (EdgePin->IsStructMember())
+		{
+			// We need to make sure the pin is not removed if any links are broken
+			// We temporarily add a preferred type to make sure that doesn't happen
+			if (URigVMTemplateNode* TemplateNode = Cast<URigVMTemplateNode>(EdgePin->GetNode()))
+			{
+				bAddedTemporaryPreferredType = true;
+				URigVMPin* RootPin = EdgePin->GetRootPin();
+				for (FRigVMTemplatePreferredType& Preferred : TemplateNode->PreferredPermutationPairs)
+				{
+					if (Preferred.GetArgument() == RootPin->GetFName())
+					{
+						bAddedTemporaryPreferredType = false;
+						break;
+					}
+				}
+
+				if (bAddedTemporaryPreferredType)
+				{
+					AddPreferredType(TemplateNode, RootPin->GetFName(), RootPin->GetTypeIndex(), bSetupUndoRedo);					
+				}
+			}
+		}
+	}
+
 	BreakAllLinks(InputPin, true, bSetupUndoRedo);
 	if (bSetupUndoRedo)
 	{
@@ -9324,36 +9272,11 @@ bool URigVMController::AddLink(URigVMPin* OutputPin, URigVMPin* InputPin, bool b
 	{
 		URigVMPin* FirstToResolve = (InUserDirection == ERigVMPinDirection::Input) ? OutputPin : InputPin;
 		URigVMPin* SecondToResolve = (FirstToResolve == OutputPin) ? InputPin : OutputPin;
-		if (URigVMTemplateNode* FirstTemplateNode = Cast<URigVMTemplateNode>(FirstToResolve->GetNode()))
+		if (!PrepareToLink(FirstToResolve, SecondToResolve, bSetupUndoRedo))
 		{
-			if (!FirstTemplateNode->IsSingleton())
-			{
-				TArray<TRigVMTypeIndex> InputTypes = GetWildcardFilteredTypeIndices(SecondToResolve);
-				if (InputTypes.Num() > 0)
-				{
-					if (!PrepareTemplatePinForType(FirstToResolve, InputTypes, bSetupUndoRedo))
-					{
-						ActionStack->CancelAction(Action, this);
-						return false;
-					}
-				}
-			}
+			ActionStack->CancelAction(Action, this);
+			return false;
 		}
-		if (URigVMTemplateNode* SecondTemplateNode = Cast<URigVMTemplateNode>(SecondToResolve->GetNode()))
-		{
-			if (!SecondTemplateNode->IsSingleton())
-			{
-				TArray<TRigVMTypeIndex> OutTypes = GetWildcardFilteredTypeIndices(FirstToResolve);
-				if (OutTypes.Num() > 0)
-				{
-					if (!PrepareTemplatePinForType(SecondToResolve, OutTypes, bSetupUndoRedo))
-					{
-						ActionStack->CancelAction(Action, this);
-						return false;
-					}
-				}
-			}
-		}	
 	}
 
 	if (bSetupUndoRedo)
@@ -9376,6 +9299,48 @@ bool URigVMController::AddLink(URigVMPin* OutputPin, URigVMPin* InputPin, bool b
 	}
 	Notify(ERigVMGraphNotifType::LinkAdded, Link);
 
+	// Remove any temporary preferred type that was added
+	for (int32 i=0; i<2; ++i)
+	{
+		URigVMPin* EdgePin = (i==0) ? OutputPin : InputPin;
+		bool& bAddedTemporaryPreferredType = (i==0) ? bAddedTemporaryOutputPreferredType : bAddedTemporaryInputPreferredType;
+		if (bAddedTemporaryPreferredType)
+		{
+			if (URigVMTemplateNode* TemplateNode = Cast<URigVMTemplateNode>(EdgePin->GetNode()))
+			{
+				URigVMPin* RootPin = EdgePin->GetRootPin();
+				RemovePreferredType(TemplateNode, RootPin->GetFName(), bSetupUndoRedo);				
+			}
+		}
+	}
+
+	// Find any library pin that doesn't own an argument, but should 
+	if (URigVMLibraryNode* LibraryNode = Graph->GetTypedOuter<URigVMLibraryNode>())
+	{
+		for (URigVMPin* Pin : LibraryNode->GetPins())
+		{
+			if (!LibraryNode->GetTemplate()->FindArgument(Pin->GetFName()))
+			{
+				if (ShouldPinOwnArgument(Pin))
+				{
+					URigVMPin* InterfacePin = nullptr;
+					if (URigVMFunctionEntryNode* EntryNode = LibraryNode->GetEntryNode())
+					{
+						InterfacePin = EntryNode->FindPin(Pin->GetName());
+					}
+					if (!InterfacePin)
+					{
+						if (URigVMFunctionReturnNode* ReturnNode = LibraryNode->GetReturnNode())
+						{
+							InterfacePin = ReturnNode->FindPin(Pin->GetName());
+						}
+					}
+					AddArgumentForPin(InterfacePin, nullptr, bSetupUndoRedo);
+				}
+			}
+		}
+	}
+
 	if (bSetupUndoRedo)
 	{
 		UpdateRerouteNodeAfterChangingLinks(OutputPin, bSetupUndoRedo);
@@ -9539,7 +9504,7 @@ bool URigVMController::BreakLink(URigVMPin* OutputPin, URigVMPin* InputPin, bool
 
 			if (!bIsTransacting && !bSuspendRecomputingTemplateFilters)
 			{
-				RecomputeAllTemplateFilteredTypes(bSetupUndoRedo);
+				RecomputeAllTemplateFilteredPermutations(bSetupUndoRedo);
 			}
 			
 			if (!bSuspendNotifications)
@@ -9656,7 +9621,7 @@ bool URigVMController::BreakAllLinks(URigVMPin* Pin, bool bAsInput, bool bSetupU
 
 	if (LinksBroken > 0 && !bSuspendRecomputingTemplateFilters)
 	{
-		RecomputeAllTemplateFilteredTypes(bSetupUndoRedo);
+		RecomputeAllTemplateFilteredPermutations(bSetupUndoRedo);
 	}
 
 	if (bSetupUndoRedo)
@@ -9700,7 +9665,7 @@ bool URigVMController::BreakAllLinksRecursive(URigVMPin* Pin, bool bAsInput, boo
 
 	if (bBrokenLinks && !bSuspendRecomputingTemplateFilters)
 	{
-		RecomputeAllTemplateFilteredTypes(bSetupUndoRedo);
+		RecomputeAllTemplateFilteredPermutations(bSetupUndoRedo);
 	}
 	
 	return bBrokenLinks;
@@ -9749,7 +9714,7 @@ FName URigVMController::AddExposedPin(const FName& InPinName, ERigVMPinDirection
 		}
 	}
 
-	/* Eventually we want to disallow this only on function library collapse nodes
+	// For now we only allow wildcards for collapsed nodes
 	if (CPPTypeObject)
 	{
 		if(CPPTypeObject == RigVMTypeUtils::GetWildCardCPPTypeObject())
@@ -9764,18 +9729,9 @@ FName URigVMController::AddExposedPin(const FName& InPinName, ERigVMPinDirection
 			}
 		}
 	}
-	*/
-	if (CPPTypeObject)
-	{
-		if(CPPTypeObject == RigVMTypeUtils::GetWildCardCPPTypeObject())
-		{
-			ReportError(TEXT("Cannot expose pins of wildcard type."));
-			return NAME_None;
-		}
-	}
-    	
-	// only allow one IO / input exposed pin of type execute context per direction
-	if(InDirection != ERigVMPinDirection::Output)
+	
+	// only allow one IO / input exposed pin of type execute context per direction
+	if(InDirection != ERigVMPinDirection::Output)
 	{
 		if(const UScriptStruct* CPPTypeStruct = Cast<UScriptStruct>(CPPTypeObject))
 		{
@@ -9851,6 +9807,7 @@ FName URigVMController::AddExposedPin(const FName& InPinName, ERigVMPinDirection
 		EntryNode = NewObject<URigVMFunctionEntryNode>(Graph, TEXT("Entry"));
 		Graph->Nodes.Add(EntryNode);
 		RefreshFunctionPins(EntryNode, false);
+		EntryNode->InitializeFilteredPermutations();
 
 		Notify(ERigVMGraphNotifType::NodeAdded, EntryNode);
 	}
@@ -9861,12 +9818,47 @@ FName URigVMController::AddExposedPin(const FName& InPinName, ERigVMPinDirection
 		ReturnNode = NewObject<URigVMFunctionReturnNode>(Graph, TEXT("Return"));
 		Graph->Nodes.Add(ReturnNode);
 		RefreshFunctionPins(ReturnNode, false);
+		ReturnNode->InitializeFilteredPermutations();
 
 		Notify(ERigVMGraphNotifType::NodeAdded, ReturnNode);
 	}
 
 	RefreshFunctionPins(EntryNode, true);
 	RefreshFunctionPins(ReturnNode, true);
+
+	// If the type is not a wildcard, add pin as template argument
+	if (InCPPType != RigVMTypeUtils::GetWildCardCPPType() && InCPPType != RigVMTypeUtils::GetWildCardArrayCPPType())
+	{
+		URigVMPin* InterfacePin = nullptr;
+		if (EntryNode)
+		{
+			InterfacePin = EntryNode->FindPin(Pin->GetName());
+		}
+		if (!InterfacePin)
+		{
+			if (ReturnNode)
+			{
+				InterfacePin = ReturnNode->FindPin(Pin->GetName());
+			}
+		}
+
+		// Add preferred type
+		{
+			URigVMTemplateNode* InterfaceNode = Cast<URigVMTemplateNode>(InterfacePin->GetNode());
+			TArray<FRigVMTemplatePreferredType> NewPreferredTypes = InterfaceNode->PreferredPermutationPairs;
+
+			NewPreferredTypes.Add(FRigVMTemplatePreferredType(*InterfacePin->GetName(), FRigVMRegistry::Get().GetTypeIndexFromCPPType(InCPPType)));
+
+			if (bSetupUndoRedo)
+			{
+				ActionStack->AddAction(FRigVMSetPreferredTemplatePermutationsAction(InterfaceNode, NewPreferredTypes));
+			}
+			InterfaceNode->PreferredPermutationPairs = NewPreferredTypes;
+		}	
+		
+		AddArgumentForPin(InterfacePin, nullptr, bSetupUndoRedo);
+	}
+	
 	RefreshFunctionReferences(LibraryNode, bSetupUndoRedo);
 
 	if (bSetupUndoRedo)
@@ -9894,6 +9886,346 @@ FName URigVMController::AddExposedPin(const FName& InPinName, ERigVMPinDirection
 	return PinName;
 }
 
+bool URigVMController::AddArgumentForPin(URigVMPin* InPin, URigVMPin* InToLinkPin, bool bSetupUndoRedo, bool bPrintPythonCommand)
+{
+	if (!IsValidGraph())
+	{
+		return false;
+	}
+
+	if (!bIsTransacting && !IsGraphEditable())
+	{
+		return false;
+	}
+
+	URigVMGraph* Graph = GetGraph();
+	check(Graph);
+	
+	if (Graph != InPin->GetGraph())
+	{
+		return false;
+	}
+
+	URigVMLibraryNode* LibraryNode = Graph->GetTypedOuter<URigVMLibraryNode>();
+	if (!LibraryNode)
+	{
+		return false;
+	}
+
+	if (!InPin->GetNode()->IsA<URigVMFunctionEntryNode>() && !InPin->GetNode()->IsA<URigVMFunctionReturnNode>())
+	{
+		return false;
+	}
+
+	FRigVMBaseAction Action;
+	if (bSetupUndoRedo)
+	{
+		Action.Title = TEXT("Add exposed pin from inner pin");
+		ActionStack->BeginAction(Action);
+	}
+
+	const FRigVMTemplate& CurTemplate = LibraryNode->Template;
+	FRigVMTemplate NewTemplate;
+	{
+		// Find all pins connected to the interface pin
+		TArray<URigVMPin*> ConnectedPins = InPin->GetLinkedSourcePins();
+		ConnectedPins.Append(InPin->GetLinkedTargetPins());
+		TArray<URigVMLink*> ConnectedLinks;
+
+		for (URigVMPin* ConnectedPin : ConnectedPins)
+		{
+			ConnectedLinks.Append(ConnectedPin->GetLinks().FilterByPredicate([&](URigVMLink* Link)
+			{
+				return Link->GetSourcePin() == InPin || Link->GetTargetPin() == InPin;
+			}));
+		}
+		for (URigVMNode* Node : Graph->GetNodes())
+		{
+			if (URigVMVariableNode* VariableNode = Cast<URigVMVariableNode>(Node))
+			{
+				if (VariableNode->GetVariableName() == *InPin->GetName())
+				{
+					URigVMPin* ValuePin = VariableNode->GetValuePin();
+					ConnectedPins.Append(ValuePin->GetLinkedSourcePins());
+					ConnectedPins.Append(ValuePin->GetLinkedTargetPins());
+
+					ConnectedLinks.Append(ValuePin->GetLinks());
+				}
+			}
+		}
+	
+		URigVMPin* PinToLink = InToLinkPin;
+		FRigVMRegistry& Registry = FRigVMRegistry::Get();
+
+		// Find the types shared by all the connected pins
+		TArray<TRigVMTypeIndex> TypeIndices;
+		{
+			// If no linked pin provided, find one of the connected pins that has a reduced number of filtered types
+			if (PinToLink == nullptr)
+			{
+				for (URigVMPin* Pin : ConnectedPins)
+				{
+					if (URigVMTemplateNode* TemplateNode = Cast<URigVMTemplateNode>(Pin->GetNode()))
+					{
+						if (TemplateNode->GetFilteredPermutationsIndices().Num() < 10)
+						{
+							PinToLink = Pin;
+							break;
+						}
+					}
+				}
+			}
+		
+			// Initialize types to the pin to link
+			if (PinToLink)
+			{
+				if (URigVMTemplateNode* ToLinkTemplateNode = Cast<URigVMTemplateNode>(PinToLink->GetNode()))
+				{
+					TypeIndices = ToLinkTemplateNode->GetFilteredTypesForPin(PinToLink);
+
+					// Avoid adding arguments with too many filtered types due to connections with unfiltered if, select, reroute or array nodes
+					if (TypeIndices.Num() > 10)
+					{
+						if (bSetupUndoRedo)
+						{
+							ActionStack->CancelAction(Action, this);
+						}
+						return false;
+					}
+				}
+				else
+				{
+					TypeIndices.Add(Registry.GetTypeIndexFromCPPType(PinToLink->GetCPPType()));
+				}
+			}
+			else
+			{
+				// If we couldn't find a pin to link (even through the connections) try to find a preferred type in interface nodes
+				TRigVMTypeIndex TypeIndex = INDEX_NONE;
+				for (int32 EntryReturn=0; EntryReturn<2 && TypeIndex == INDEX_NONE; ++EntryReturn)
+				{
+					URigVMTemplateNode* Node = (EntryReturn == 0) ? (URigVMTemplateNode*)Graph->GetEntryNode() : (URigVMTemplateNode*)Graph->GetReturnNode();
+					if (Node)
+					{
+						for (int32 i=0; i<Node->PreferredPermutationPairs.Num(); ++i)
+						{
+							FString ArgName;
+							if (Node->PreferredPermutationPairs[i].GetArgument() == InPin->GetFName())
+							{
+								TypeIndex = Node->PreferredPermutationPairs[i].GetTypeIndex();
+								break;
+							}
+						}
+					}
+				}
+
+				if (TypeIndex != INDEX_NONE)
+				{
+					TypeIndices.Add(TypeIndex);
+				}
+				else
+				{
+					if (bSetupUndoRedo)
+					{
+						ActionStack->CancelAction(Action, this);
+					}
+					return false;
+				}				
+			}
+
+			// Reduce types to the ones supported by all connections
+			for (int32 i=0; i<ConnectedPins.Num(); ++i)
+			{
+				URigVMPin* ConnectedPin = ConnectedPins[i];
+				if (ConnectedPin != PinToLink)
+				{
+					if (URigVMTemplateNode* ConnectedTemplateNode = Cast<URigVMTemplateNode>(ConnectedPin->GetNode()))
+					{
+						TypeIndices = TypeIndices.FilterByPredicate([&](int32 Type)
+						{
+							return ConnectedTemplateNode->FilteredSupportsType(ConnectedPin, Type);
+						});		
+					}
+					else
+					{
+						// This is a weird situation. If the pin was connected, but was not an argument it usually means
+						// it was connected to an if, select, array or reroute that has not been type-filtered
+						check(false);
+					}
+				}
+			}
+		}
+		if (PinToLink)
+		{
+			ConnectedPins.AddUnique(PinToLink);
+		}
+
+		ensure(!TypeIndices.IsEmpty());
+
+		// Temporarily disconnect connected links
+		if (CurTemplate.NumArguments() > 0)
+		{
+			DetachLinksFromPinObjects(&ConnectedLinks);
+		}
+	
+		// Resolve ConnectedPins to each possible filtered type and see how that affects other exposed pins	
+		TArray<TArray<TArray<TRigVMTypeIndex>>> TypeToPermutations; // #types_to_resolve x #resolved_permutations x #arguments
+		ensure(TypeIndices.Num() < 10);
+		for (TRigVMTypeIndex& TypeIndex : TypeIndices)
+		{
+			TArray<TArray<TRigVMTypeIndex>> PermutationTypes;
+			if (CurTemplate.Arguments.Num() > 0)
+			{
+				TGuardValue<bool> GuardNotifications(bSuspendNotifications, true);
+				TGuardValue<bool> GuardOuterRecomputing(bSuspendRecomputingOuterTemplateFilters, true);
+				OpenUndoBracket(TEXT("Resolving to type"));
+
+				for (URigVMPin* ConnectedPin : ConnectedPins)
+				{
+					if (ConnectedPin->IsWildCard())
+					{
+						ResolveWildCardPin(ConnectedPin, TypeIndex, true);
+					}					
+				}
+
+				for (int32 i=0; i<CurTemplate.NumPermutations(); ++i)
+				{
+					TArray<TRigVMTypeIndex> ArgTypes;					
+					for (const FRigVMTemplateArgument& Arg : CurTemplate.Arguments)
+					{
+						ArgTypes.Add(Arg.TypeIndices[i]);
+					}
+					PermutationTypes.Emplace(ArgTypes);
+				}
+
+				check(CancelUndoBracket());
+			}
+			TypeToPermutations.Emplace(PermutationTypes);
+		}
+
+		// Reconnect  links
+		if (CurTemplate.NumArguments() > 0)
+		{
+			ReattachLinksToPinObjects(false, &ConnectedLinks);
+		}
+		
+		FRigVMTemplateArgument NewArgument;
+		NewArgument.Index = CurTemplate.Arguments.Num();
+		NewArgument.Name = InPin->GetFName();
+		NewArgument.Direction = LibraryNode->FindPin(InPin->GetName())->GetDirection();
+	
+		// Add the arguments
+		for (int32 i=0; i<CurTemplate.Arguments.Num(); ++i)
+		{
+			FRigVMTemplateArgument Arg;
+			Arg.Index = i;
+			Arg.Name = CurTemplate.Arguments[i].Name;
+			Arg.Direction = CurTemplate.Arguments[i].Direction;
+			NewTemplate.Arguments.Add(Arg);
+		}
+		NewTemplate.Arguments.Add(NewArgument);
+
+		// Types
+		for (int32 NewArgTypeIndex=0; NewArgTypeIndex<TypeIndices.Num(); NewArgTypeIndex++)
+		{
+			TRigVMTypeIndex& NewArgType = TypeIndices[NewArgTypeIndex];			
+
+			if (CurTemplate.NumArguments() > 0)
+			{
+				for (int32 PermutationIndex=0; PermutationIndex<TypeToPermutations[NewArgTypeIndex].Num(); ++PermutationIndex)
+				{
+					int32 GlobalPermutationIndex = NewTemplate.Permutations.Num();
+					NewTemplate.Permutations.Add(INDEX_NONE);
+					for (int32 ArgIndex=0; ArgIndex<NewTemplate.NumArguments()-1; ++ArgIndex)
+					{
+						TRigVMTypeIndex ArgType = TypeToPermutations[NewArgTypeIndex][PermutationIndex][ArgIndex];
+						NewTemplate.Arguments[ArgIndex].TypeIndices.Add(ArgType);
+						if (TArray<int32>* PermArray = NewTemplate.Arguments[ArgIndex].TypeToPermutations.Find(ArgType))
+						{
+							PermArray->Add(GlobalPermutationIndex);
+						}
+						else
+						{
+							NewTemplate.Arguments[ArgIndex].TypeToPermutations.Add(ArgType, {GlobalPermutationIndex});
+						}
+					}
+					NewTemplate.Arguments.Last().TypeIndices.Add(NewArgType);
+					if (TArray<int32>* PermArray = NewTemplate.Arguments.Last().TypeToPermutations.Find(NewArgType))
+					{
+						PermArray->Add(GlobalPermutationIndex);
+					}
+					else
+					{
+						NewTemplate.Arguments.Last().TypeToPermutations.Add(NewArgType, {GlobalPermutationIndex});
+					}				
+				}
+			}
+			else
+			{
+				int32 GlobalPermutationIndex = NewTemplate.Permutations.Num();
+				NewTemplate.Permutations.Add(INDEX_NONE);
+				NewTemplate.Arguments.Last().TypeIndices.Add(NewArgType);
+				if (TArray<int32>* PermArray = NewTemplate.Arguments.Last().TypeToPermutations.Find(NewArgType))
+				{
+					PermArray->Add(GlobalPermutationIndex);
+				}
+				else
+				{
+					NewTemplate.Arguments.Last().TypeToPermutations.Add(NewArgType, {GlobalPermutationIndex});
+				}
+			}
+		}
+	}
+	
+	NewTemplate.ComputeNotationFromArguments(LibraryNode->GetName());
+	
+	if (bSetupUndoRedo)
+	{
+		FRigVMControllerGraphGuard GraphGuard(this, LibraryNode->GetTypedOuter<URigVMGraph>(), bSetupUndoRedo);
+		ActionStack->AddAction(FRigVMSetLibraryTemplateAction(LibraryNode, NewTemplate));
+	}
+	
+	LibraryNode->Template = NewTemplate;	
+	
+	// Update Entry and Return filtered permutations to all permutations
+	for (int32 EntryReturn=0; EntryReturn<2; ++EntryReturn)
+	{
+		URigVMTemplateNode* Node = (EntryReturn == 0) ? (URigVMTemplateNode*)Graph->GetEntryNode() : (URigVMTemplateNode*)Graph->GetReturnNode();
+		if (Node)
+		{
+			TArray<int32> OldPermutations = Node->FilteredPermutations;
+
+			Node->FilteredPermutations.SetNumUninitialized(NewTemplate.NumPermutations());
+			for (int32 i=0; i<Node->FilteredPermutations.Num(); ++i)
+			{
+				Node->FilteredPermutations[i] = i;
+			}
+
+			if (bSetupUndoRedo)
+			{
+				ActionStack->AddAction(FRigVMSetTemplateFilteredPermutationsAction(Node, OldPermutations));
+			}
+
+			UpdateTemplateNodePinTypes(Node, bSetupUndoRedo);
+			Notify(ERigVMGraphNotifType::NodeDescriptionChanged, Node);
+		}
+	}
+
+	if (!bIsTransacting && !bSuspendRecomputingOuterTemplateFilters)
+	{
+		FRigVMControllerGraphGuard GraphGuard(this, LibraryNode->GetTypedOuter<URigVMGraph>(), bSetupUndoRedo);
+		RecomputeAllTemplateFilteredPermutations(bSetupUndoRedo);
+		Notify(ERigVMGraphNotifType::NodeDescriptionChanged, LibraryNode);		
+	}
+
+	if (bSetupUndoRedo)
+	{
+		ActionStack->EndAction(Action);
+	}
+	
+	return true;
+}
+
 bool URigVMController::RemoveExposedPin(const FName& InPinName, bool bSetupUndoRedo, bool bPrintPythonCommand)
 {
 	if (!IsValidGraph())
@@ -9952,29 +10284,45 @@ bool URigVMController::RemoveExposedPin(const FName& InPinName, bool bSetupUndoR
 
 	bool bSuccessfullyRemovedPin = false;
 	{
-		FRigVMControllerGraphGuard GraphGuard(this, LibraryNode->GetGraph(), bSetupUndoRedo);
-		bSuccessfullyRemovedPin = RemovePin(Pin, bSetupUndoRedo, true);
-	}
+		TGuardValue<bool> GuardRecomputeFilters(bSuspendRecomputingTemplateFilters, true);
+		TGuardValue<bool> GuardRecomputeOuterFilters(bSuspendRecomputingOuterTemplateFilters, true);
+		{
+			FRigVMControllerGraphGuard GraphGuard(this, LibraryNode->GetGraph(), bSetupUndoRedo);
+			bSuccessfullyRemovedPin = RemovePin(Pin, bSetupUndoRedo, true);
+		}
 
-	TArray<URigVMVariableNode*> NodesToRemove;
-	for (URigVMNode* Node : Graph->GetNodes())
-	{
-		if (URigVMVariableNode* VariableNode = Cast<URigVMVariableNode>(Node))
+		TArray<URigVMVariableNode*> NodesToRemove;
+		for (URigVMNode* Node : Graph->GetNodes())
 		{
-			if (VariableNode->GetVariableName() == InPinName)
+			if (URigVMVariableNode* VariableNode = Cast<URigVMVariableNode>(Node))
 			{
-				NodesToRemove.Add(VariableNode);
+				if (VariableNode->GetVariableName() == InPinName)
+				{
+					NodesToRemove.Add(VariableNode);
+				}
 			}
 		}
+		for (int32 i=NodesToRemove.Num()-1; i >= 0; --i)
+		{
+			RemoveNode(NodesToRemove[i], bSetupUndoRedo);
+		}
+
+		RefreshFunctionPins(Graph->GetEntryNode(), true);
+		RefreshFunctionPins(Graph->GetReturnNode(), true);
+		RefreshFunctionReferences(LibraryNode, false);
 	}
-	for (int32 i=NodesToRemove.Num()-1; i >= 0; --i)
+
+	TArray<URigVMTemplateNode*> InterfaceNodes = {Graph->GetEntryNode(), Graph->GetReturnNode()};
+	// Remove preferred types in interface nodes
 	{
-		RemoveNode(NodesToRemove[i], bSetupUndoRedo);
+		for (URigVMTemplateNode* InterfaceNode : InterfaceNodes)
+		{
+			RemovePreferredType(InterfaceNode, InPinName, bSetupUndoRedo);						
+		}
 	}
 
-	RefreshFunctionPins(Graph->GetEntryNode(), true);
-	RefreshFunctionPins(Graph->GetReturnNode(), true);
-	RefreshFunctionReferences(LibraryNode, false);
+	// Recompute filtered types and library template
+	RecomputeAllTemplateFilteredPermutations(bSetupUndoRedo);
 
 	if (bSetupUndoRedo)
 	{
@@ -10120,20 +10468,53 @@ bool URigVMController::RenameExposedPin(const FName& InOldPinName, const FName&
 		return false;
 	}
 
-	if (URigVMFunctionEntryNode* EntryNode = Graph->GetEntryNode())
+	TArray<URigVMTemplateNode*> InterfaceNodes = {Graph->GetEntryNode(), Graph->GetReturnNode()};
+	for (URigVMTemplateNode* InterfaceNode : InterfaceNodes)
 	{
-		if (URigVMPin* EntryPin = EntryNode->FindPin(InOldPinName.ToString()))
+		if (InterfaceNode)
 		{
-			Local::RenamePin(this, EntryPin, PinName);
+			if (URigVMPin* InterfacePin = InterfaceNode->FindPin(InOldPinName.ToString()))
+			{
+				TArray<FRigVMTemplatePreferredType> NewPreferredPermutationsTypes = InterfaceNode->PreferredPermutationPairs;
+				for (FRigVMTemplatePreferredType& PreferredType : NewPreferredPermutationsTypes)
+				{
+					if (PreferredType.GetArgument() == InOldPinName)
+					{
+						PreferredType.Argument = PinName;
+						break;
+					}
+				}
+
+				if (bSetupUndoRedo)
+				{
+					ActionStack->AddAction(FRigVMSetPreferredTemplatePermutationsAction(InterfaceNode, NewPreferredPermutationsTypes));
+				}
+
+				InterfaceNode->PreferredPermutationPairs = NewPreferredPermutationsTypes;
+				
+				Local::RenamePin(this, InterfacePin, PinName);
+			}
 		}
 	}
 
-	if (URigVMFunctionReturnNode* ReturnNode = Graph->GetReturnNode())
+	// Find Argument and rename if it exists
 	{
-		if (URigVMPin* ReturnPin = ReturnNode->FindPin(InOldPinName.ToString()))
+		FRigVMTemplate NewTemplate = LibraryNode->Template;
+		for (FRigVMTemplateArgument& Argument : NewTemplate.Arguments)
+		{
+			if (Argument.Name == InOldPinName)
+			{
+				Argument.Name = PinName;
+				break;
+			}
+		}
+		NewTemplate.ComputeNotationFromArguments(LibraryNode->GetName());
+		if (bSetupUndoRedo)
 		{
-			Local::RenamePin(this, ReturnPin, PinName);
+			FRigVMControllerGraphGuard Guard(this, LibraryNode->GetGraph(), bSetupUndoRedo);
+			ActionStack->AddAction(FRigVMSetLibraryTemplateAction(LibraryNode, NewTemplate));
 		}
+		LibraryNode->Template = NewTemplate;
 	}
 
 	if (URigVMFunctionLibrary* FunctionLibrary = Cast<URigVMFunctionLibrary>(LibraryNode->GetGraph()))
@@ -10213,11 +10594,20 @@ bool URigVMController::ChangeExposedPinType(const FName& InPinName, const FStrin
 	{
 		return false;
 	}
+
+	// We do not allow unresolving exposed pins
+	if (InCPPType == RigVMTypeUtils::GetWildCardCPPType())
+	{
+		ReportError(TEXT("Cannot change exposed pin type to wildcard."));
+		return false;
+	}
 	
 	// only allow one exposed pin of type execute context per direction
+	UObject* CPPTypeObject = nullptr;
 	if (!InCPPTypeObjectPath.IsNone())
 	{
-		if(UObject* CPPTypeObject = URigVMPin::FindObjectFromCPPTypeObjectPath<UObject>(InCPPTypeObjectPath.ToString()))
+		CPPTypeObject = URigVMPin::FindObjectFromCPPTypeObjectPath<UObject>(InCPPTypeObjectPath.ToString());
+		if(CPPTypeObject)
 		{
 			if(const UScriptStruct* CPPTypeStruct = Cast<UScriptStruct>(CPPTypeObject))
 			{
@@ -10259,35 +10649,121 @@ bool URigVMController::ChangeExposedPinType(const FName& InPinName, const FStrin
 		ActionStack->BeginAction(Action);
 	}
 
+	FRigVMRegistry& Registry = FRigVMRegistry::Get();
+
+	// If the pin does not support the type, first break all links in the contained graph to this pin (and subpins)
+	TArray<URigVMTemplateNode*> InterfaceNodes = {Graph->GetEntryNode(), Graph->GetReturnNode()};
+	if (!LibraryNode->SupportsType(Pin, Registry.GetTypeIndexFromCPPType(InCPPType)))
 	{
-		FRigVMControllerGraphGuard GraphGuard(this, LibraryNode->GetGraph(), bSetupUndoRedo);
-		if (!ChangePinType(Pin, InCPPType, InCPPTypeObjectPath, bSetupUndoRedo, bSetupOrphanPins))
 		{
-			if (bSetupUndoRedo)
+			TGuardValue<bool> GuardRecompute (bSuspendRecomputingTemplateFilters, true);
+			TGuardValue<bool> GuardRecomputeOuter (bSuspendRecomputingOuterTemplateFilters, true);
+
+			// Remove preferred types for this pin
+			for (URigVMTemplateNode* InterfaceNode : InterfaceNodes)
 			{
-				ActionStack->CancelAction(Action, this);
+				if (InterfaceNode)
+				{
+					if (URigVMPin* InterfacePin = InterfaceNode->FindPin(Pin->GetName()))
+					{
+						RemovePreferredType(InterfaceNode, InPinName, bSetupUndoRedo);
+					}										
+				}
+			}
+		
+			// Break all links to this pin
+			{
+				TArray<URigVMLink*> InterfacePinLinks;	
+				TArray<URigVMPin*> ExtendedInterfacePins;
+				for (URigVMNode* Node : Graph->GetNodes())
+				{
+					if (Node->IsA<URigVMFunctionEntryNode>() || Node->IsA<URigVMFunctionReturnNode>())
+					{
+						if (URigVMPin* InterfacePin = Node->FindPin(Pin->GetName()))
+						{
+							ExtendedInterfacePins.Add(InterfacePin);
+						}
+					}
+					else if (URigVMVariableNode* VariableNode = Cast<URigVMVariableNode>(Node))
+					{
+						if (VariableNode->GetVariableName() == InPinName)
+						{
+							ExtendedInterfacePins.Add(VariableNode->GetValuePin());
+						}
+					}
+				}
+		
+				for (URigVMPin* InterfacePin : ExtendedInterfacePins)
+				{
+					TArray<URigVMPin*> PinsToProcess;
+					PinsToProcess.Add(InterfacePin);
+					for (int32 i=0; i<PinsToProcess.Num(); ++i)
+					{
+						InterfacePinLinks.Append(PinsToProcess[i]->GetLinks());
+						PinsToProcess.Append(PinsToProcess[i]->GetSubPins());
+					}
+				}		
+				for (int32 i=0; i<InterfacePinLinks.Num(); ++i)
+				{
+					BreakLink(InterfacePinLinks[i]->GetSourcePin(), InterfacePinLinks[i]->GetTargetPin(), bSetupUndoRedo);
+				}
 			}
-			return false;
 		}
-		RemoveUnusedOrphanedPins(LibraryNode, true);
+
+		RecomputeAllTemplateFilteredPermutations(bSetupUndoRedo);
 	}
 
-	if (URigVMFunctionEntryNode* EntryNode = Graph->GetEntryNode())
+	// Resolve interface pins to the input type
+	for (URigVMTemplateNode* InterfaceNode : InterfaceNodes)
 	{
-		const TArray<URigVMLink*> Links = EntryNode->GetLinks();
-		DetachLinksFromPinObjects(&Links, true);
-		RepopulatePinsOnNode(EntryNode, true, true, bSetupOrphanPins);
-		ReattachLinksToPinObjects(true, &Links, true, bSetupOrphanPins);
-		RemoveUnusedOrphanedPins(EntryNode, true);
+		if (InterfaceNode)
+		{
+			if (URigVMPin* InterfacePin = InterfaceNode->FindPin(Pin->GetName()))
+			{
+				if (!ResolveWildCardPin(InterfacePin, FRigVMTemplateArgumentType(*InCPPType, CPPTypeObject), bSetupUndoRedo))
+				{
+					if (bSetupUndoRedo)
+					{
+						ActionStack->CancelAction(Action, this);
+					}
+					return false;
+				}
+			}
+		}
 	}
 	
-	if (URigVMFunctionReturnNode* ReturnNode = Graph->GetReturnNode())
+		
 	{
-		const TArray<URigVMLink*> Links = ReturnNode->GetLinks();
-		DetachLinksFromPinObjects(&Links, true);
-		RepopulatePinsOnNode(ReturnNode, true, true, bSetupOrphanPins);
-		ReattachLinksToPinObjects(true, &Links, true, bSetupOrphanPins);
-		RemoveUnusedOrphanedPins(ReturnNode, true);
+		bool bSuccessChangingType;
+		{
+			FRigVMControllerGraphGuard GraphGuard(this, LibraryNode->GetGraph(), bSetupUndoRedo);
+			bSuccessChangingType = ChangePinType(Pin, InCPPType, InCPPTypeObjectPath, bSetupUndoRedo, bSetupOrphanPins);
+
+			if (bSuccessChangingType)
+			{
+				RemoveUnusedOrphanedPins(LibraryNode, true);
+			}
+		}
+		if (!bSuccessChangingType)
+		{
+			if (bSetupUndoRedo)
+			{
+				ActionStack->CancelAction(Action, this);
+			}
+			return false;
+		}
+	}
+
+	for (URigVMTemplateNode* InterfaceNode : InterfaceNodes)
+	{
+		if (InterfaceNode)
+		{
+			const TArray<URigVMLink*> Links = InterfaceNode->GetLinks();
+			DetachLinksFromPinObjects(&Links, true);
+			RepopulatePinsOnNode(InterfaceNode, true, true, bSetupOrphanPins);
+			ReattachLinksToPinObjects(true, &Links, true, bSetupOrphanPins);
+			RemoveUnusedOrphanedPins(InterfaceNode, true);
+		}
 	}
 
 	if (URigVMFunctionLibrary* FunctionLibrary = Cast<URigVMFunctionLibrary>(LibraryNode->GetGraph()))
@@ -10464,6 +10940,7 @@ URigVMFunctionReferenceNode* URigVMController::AddFunctionReferenceNode(URigVMLi
 	FunctionRefNode->Position = InNodePosition;
 	FunctionRefNode->SetReferencedNode(InFunctionDefinition);
 	Graph->Nodes.Add(FunctionRefNode);
+	FunctionRefNode->InitializeFilteredPermutations();
 
 	FRigVMControllerCompileBracketScope CompileScope(this);
 	
@@ -10656,13 +11133,14 @@ URigVMLibraryNode* URigVMController::AddFunctionToLibrary(const FName& InFunctio
 	if (bMutable)
 	{
 		UScriptStruct* ExecuteContextStruct = Graph->GetExecuteContextStruct();
-		URigVMPin* ExecutePin = NewObject<URigVMPin>(CollapseNode, FRigVMStruct::ExecuteContextName);
-		ExecutePin->DisplayName = FRigVMStruct::ExecuteName;
-		ExecutePin->CPPType = FString::Printf(TEXT("F%s"), *ExecuteContextStruct->GetName());
-		ExecutePin->CPPTypeObject = ExecuteContextStruct;
-		ExecutePin->CPPTypeObjectPath = *ExecutePin->CPPTypeObject->GetPathName();
-		ExecutePin->Direction = ERigVMPinDirection::IO;
-		AddNodePin(CollapseNode, ExecutePin);
+		
+		FRigVMControllerGraphGuard GraphGuard(this, CollapseNode->ContainedGraph, bSetupUndoRedo);
+		AddExposedPin(FRigVMStruct::ExecuteContextName,
+			ERigVMPinDirection::IO,
+			FString::Printf(TEXT("F%s"), *ExecuteContextStruct->GetName()),
+			*ExecuteContextStruct->GetPathName(),
+			FString(),
+			bSetupUndoRedo);
 	}
 
 	Notify(ERigVMGraphNotifType::NodeAdded, CollapseNode);
@@ -10671,17 +11149,27 @@ URigVMLibraryNode* URigVMController::AddFunctionToLibrary(const FName& InFunctio
 		FRigVMControllerGraphGuard GraphGuard(this, CollapseNode->GetContainedGraph(), false);
 		TGuardValue<bool> GuardEditGraph(CollapseNode->ContainedGraph->bEditable, true);
 
-		URigVMFunctionEntryNode* EntryNode = NewObject<URigVMFunctionEntryNode>(CollapseNode->ContainedGraph, TEXT("Entry"));
-		CollapseNode->ContainedGraph->Nodes.Add(EntryNode);
-		EntryNode->Position = FVector2D(-250.f, 0.f);
-		RefreshFunctionPins(EntryNode, false);
-		Notify(ERigVMGraphNotifType::NodeAdded, EntryNode);
+		URigVMNode* EntryNode = CollapseNode->ContainedGraph->FindNode(TEXT("Entry"));
+		URigVMNode* ReturnNode = CollapseNode->ContainedGraph->FindNode(TEXT("Return"));
+
+		if (EntryNode == nullptr)
+		{
+			EntryNode = NewObject<URigVMFunctionEntryNode>(CollapseNode->ContainedGraph, TEXT("Entry"));
+			CollapseNode->ContainedGraph->Nodes.Add(EntryNode);
+			RefreshFunctionPins(EntryNode, false);
+			Notify(ERigVMGraphNotifType::NodeAdded, EntryNode);
+		}
+
+		if (ReturnNode == nullptr)
+		{
+			ReturnNode = NewObject<URigVMFunctionReturnNode>(CollapseNode->ContainedGraph, TEXT("Return"));
+			CollapseNode->ContainedGraph->Nodes.Add(ReturnNode);
+			RefreshFunctionPins(ReturnNode, false);
+			Notify(ERigVMGraphNotifType::NodeAdded, ReturnNode);
+		}
 
-		URigVMFunctionReturnNode* ReturnNode = NewObject<URigVMFunctionReturnNode>(CollapseNode->ContainedGraph, TEXT("Return"));
-		CollapseNode->ContainedGraph->Nodes.Add(ReturnNode);
+		EntryNode->Position = FVector2D(-250.f, 0.f);
 		ReturnNode->Position = FVector2D(250.f, 0.f);
-		RefreshFunctionPins(ReturnNode, false);
-		Notify(ERigVMGraphNotifType::NodeAdded, ReturnNode);
 
 		if (bMutable)
 		{
@@ -13651,7 +14139,7 @@ int32 URigVMController::DetachLinksFromPinObjects(const TArray<URigVMLink*>* InL
 	return Links.Num();
 }
 
-int32 URigVMController::ReattachLinksToPinObjects(bool bFollowCoreRedirectors, const TArray<URigVMLink*>* InLinks, bool bNotify, bool bSetupOrphanedPins)
+int32 URigVMController::ReattachLinksToPinObjects(bool bFollowCoreRedirectors, const TArray<URigVMLink*>* InLinks, bool bNotify, bool bSetupOrphanedPins, bool bAllowNonArgumentLinks)
 {
 	URigVMGraph* Graph = GetGraph();
 	check(Graph);
@@ -13716,7 +14204,7 @@ int32 URigVMController::ReattachLinksToPinObjects(bool bFollowCoreRedirectors, c
 				continue;
 			}
 
-			if (!URigVMPin::CanLink(SourcePin, TargetPin, nullptr, nullptr))
+			if (!URigVMPin::CanLink(SourcePin, TargetPin, nullptr, nullptr, ERigVMPinDirection::IO, bAllowNonArgumentLinks))
 			{
 				if(SourcePin->GetNode()->HasOrphanedPins() && bSetupOrphanedPins)
 				{
@@ -13839,7 +14327,7 @@ int32 URigVMController::ReattachLinksToPinObjects(bool bFollowCoreRedirectors, c
 			{
 				FRigVMControllerGraphGuard GraphGuard(this, CollapseNode->GetContainedGraph(), false);
 				TGuardValue<bool> GuardEditGraph(CollapseNode->ContainedGraph->bEditable, true);
-				ReattachLinksToPinObjects(bFollowCoreRedirectors, nullptr, bNotify, bSetupOrphanedPins);
+				ReattachLinksToPinObjects(bFollowCoreRedirectors, nullptr, bNotify, bSetupOrphanedPins, bAllowNonArgumentLinks);
 			}
 		}
 	}
@@ -15071,7 +15559,7 @@ bool URigVMController::FullyResolveTemplateNode(URigVMTemplateNode* InNode, int3
 
 	const FRigVMFunction* ResolvedFunction = Template->GetPermutation(InputPermutation);
 	const TArray<int32> PermutationIndices = {InputPermutation};
-	InNode->FilteredPermutations = PermutationIndices;
+	//InNode->FilteredPermutations = PermutationIndices;
 
 	// find all existing pins that we may need to change
 	TArray<FRigVMTemplateArgument> MissingPins;
@@ -15327,6 +15815,38 @@ bool URigVMController::PrepareTemplatePinForType(URigVMPin* InPin, const TArray<
 		return false;
 	}
 
+	if (InTypeIndices.IsEmpty())
+	{
+		return false;
+	}
+
+	FRigVMRegistry& Registry = FRigVMRegistry::Get();
+
+	// It might be a subpin of a struct, in that case we just want to make sure
+	// the type required matches the one in the pin
+	if (InPin->IsStructMember())
+	{
+		int32 PinTypeIndex = InPin->GetTypeIndex();
+		
+		// Check if any of the types allows the connection
+		for (const int32& Type : InTypeIndices)
+		{
+			if (Registry.CanMatchTypes(Type, PinTypeIndex, true))
+			{
+				return true;
+			}
+		}
+		
+		return false;		
+	}
+
+	FRigVMBaseAction Action;
+	Action.Title = TEXT("Prepare template pin for type");
+	if (bSetupUndoRedo)
+	{
+		ActionStack->BeginAction(Action);		
+	}
+
 	// If the pin is an element of an array, get the array pin and convert the types
 	URigVMPin* RootPin = InPin;
 	TArray<TRigVMTypeIndex> Types = InTypeIndices;
@@ -15339,17 +15859,6 @@ bool URigVMController::PrepareTemplatePinForType(URigVMPin* InPin, const TArray<
 		}
 	}
 
-	// It might be a subpin of a struct, in that case we just want to make sure
-	// the type required matches the one in the pin
-	if (!TemplateNode->GetTemplate()->FindArgument(RootPin->GetFName()))
-	{
-		if (Types.Num() == 1 && FRigVMRegistry::Get().CanMatchTypes(Types[0], RootPin->GetTypeIndex(), true))
-		{
-			return true;
-		}
-		return false;
-	}
-
 	bool bFilteredSupportsType = false;
 	for (const TRigVMTypeIndex& Type : Types)
 	{
@@ -15361,6 +15870,7 @@ bool URigVMController::PrepareTemplatePinForType(URigVMPin* InPin, const TArray<
 	}
 
 	bool bSupportsType = false;
+	TArray<FRigVMTemplatePreferredType> OldPreferredTypes;
 	if (!bFilteredSupportsType)
 	{
 		if (!TemplateNode->PreferredPermutationPairs.IsEmpty())
@@ -15369,6 +15879,7 @@ bool URigVMController::PrepareTemplatePinForType(URigVMPin* InPin, const TArray<
 			{
 				ActionStack->AddAction(FRigVMSetPreferredTemplatePermutationsAction(TemplateNode, {}));
 			}
+			OldPreferredTypes = TemplateNode->PreferredPermutationPairs;
 			TemplateNode->PreferredPermutationPairs = {};
 		}
 		
@@ -15384,20 +15895,104 @@ bool URigVMController::PrepareTemplatePinForType(URigVMPin* InPin, const TArray<
 
 	if (!bFilteredSupportsType && !bSupportsType)
 	{
+		if (bSetupUndoRedo)
+		{
+			ActionStack->CancelAction(Action, this);
+		}
 		return false;
 	}
 
 	if (bFilteredSupportsType)
 	{
-		if (TemplateNode->PinNeedsFilteredTypesUpdate(RootPin, Types))
 		{
-			if (UpdateFilteredPermutations(RootPin, Types, bSetupUndoRedo))
+			TGuardValue<FRigVMController_RequestBreakLinksDialogDelegate> GuardDelegate(RequestBreakLinksDialogDelegate, nullptr);
+			OpenUndoBracket(FString::Printf(TEXT("Resolve wildcard pin %s"), *InPin->GetPinPath()));
+			if (TemplateNode->PinNeedsFilteredTypesUpdate(RootPin, Types))
+			{
+				if (UpdateFilteredPermutations(RootPin, Types, bSetupUndoRedo))
+				{
+					UpdateTemplateNodePinTypes(TemplateNode, bSetupUndoRedo);
+					PropagateTemplateFilteredTypes(TemplateNode, bSetupUndoRedo);
+				}
+			}
+		}
+
+		// Even if the filter supports type, we might find we break links in parent graphs
+		// Find break link actions
+		TArray<URigVMLink*> LinksToBreak;
+		TArray<TPair<FString, FString>> InconsistentLinks;
+		TArray<FSoftObjectPath> InconsistentLinksGraph;
+		{
+			TArray<FRigVMActionKey> Actions = ActionStack->BracketActions.Last()->SubActions;
+			for (int32 i=0; i<Actions.Num(); ++i)
+			{
+				const FRigVMActionKey& ActionKey = Actions[i];
+				FRigVMActionWrapper Wrapper(ActionKey);
+				if (Wrapper.GetAction()->GetScriptStruct() == FRigVMBreakLinkAction::StaticStruct())
+				{
+					const FRigVMBreakLinkAction* BreakLinkAction = (const FRigVMBreakLinkAction*)Wrapper.GetAction();
+					InconsistentLinks.Add(TPair<FString,FString>(BreakLinkAction->OutputPinPath, BreakLinkAction->InputPinPath));
+					InconsistentLinksGraph.Add(BreakLinkAction->GraphPath);
+				}
+
+				Actions.Append(Wrapper.GetAction()->SubActions);
+			}
+		}
+
+		if (InconsistentLinks.IsEmpty())
+		{
+			CloseUndoBracket();
+			if (bSetupUndoRedo)
 			{
-				UpdateTemplateNodePinTypes(TemplateNode, bSetupUndoRedo);
-				PropagateTemplateFilteredTypes(TemplateNode, bSetupUndoRedo);
+				ActionStack->EndAction(Action);
 			}
+			return true;
 		}
-		return true;
+
+		check(CancelUndoBracket());
+		for (int32 i=0; i<InconsistentLinks.Num(); ++i)
+		{
+			TSoftObjectPtr<URigVMGraph> GraphPtr(InconsistentLinksGraph[i]);
+			LinksToBreak.AddUnique(GraphPtr.Get()->FindLink(FString::Printf(TEXT("%s -> %s"), *InconsistentLinks[i].Key, *InconsistentLinks[i].Value)));
+		}
+		
+		// Warn the user the links that would be broken
+		bool bBreakLinks = true;
+		bool bConsultedBreak = false;
+		if(!bIsTransacting && RequestBreakLinksDialogDelegate.IsBound())
+		{
+			bBreakLinks = RequestBreakLinksDialogDelegate.Execute(LinksToBreak);
+			bConsultedBreak = true;
+		}
+
+		if (bBreakLinks)
+		{
+			// Break links
+			for (URigVMLink* Link : LinksToBreak)
+			{
+				FRigVMControllerGraphGuard GraphGuard(this, Link->GetGraph(), bSetupUndoRedo);
+				if (!bConsultedBreak && !bIsTransacting && !bSuspendNotifications)
+				{
+					ReportWarningf(TEXT("The link between %s was broken due to inconsistent types"), *Link->GetPinPathRepresentation());
+				}
+				BreakLink(Link->GetSourcePin(), Link->GetTargetPin(), bSetupUndoRedo);
+			}
+
+			ensure(LinksToBreak.Num() > 0);
+			PrepareTemplatePinForType(InPin, InTypeIndices, bSetupUndoRedo);
+
+			if (bSetupUndoRedo)
+			{
+				ActionStack->EndAction(Action);
+			}
+			return true;
+		}
+
+		if (bSetupUndoRedo)
+		{
+			ActionStack->CancelAction(Action, this);
+		}
+		return false;
 	}
 	else
 	{
@@ -15443,7 +16038,7 @@ bool URigVMController::PrepareTemplatePinForType(URigVMPin* InPin, const TArray<
 
 				Actions.Append(Wrapper.GetAction()->SubActions);
 			}
-			CancelUndoBracket();
+			check(CancelUndoBracket());
 
 			for (TPair<FString, FString>& Pair : InconsistentLinks)
 			{
@@ -15457,11 +16052,40 @@ bool URigVMController::PrepareTemplatePinForType(URigVMPin* InPin, const TArray<
 		{
 			if (!TemplateNode->PreferredPermutationPairs.IsEmpty())
 			{
+				ensure(false);
+				if (bSetupUndoRedo)
+				{
+					ActionStack->CancelAction(Action, this);
+				}
 				return false;
 			}
 			
 			UnresolveTemplateNodes({TemplateNode}, bSetupUndoRedo);
 			PrepareTemplatePinForType(InPin, InTypeIndices, bSetupUndoRedo);
+
+			// Reapply the valid preferred types
+			for (int32 i=0; i<OldPreferredTypes.Num(); ++i)
+			{
+				if (URigVMPin* Pin = TemplateNode->FindPin(OldPreferredTypes[i].Argument.ToString()))
+				{
+					if (Pin->IsWildCard())
+					{
+						if (!TemplateNode->FilteredSupportsType(Pin, OldPreferredTypes[i].TypeIndex))
+						{
+							continue;
+						}
+					}
+				
+					if (Pin->IsWildCard() || Pin->GetTypeIndex() == OldPreferredTypes[i].TypeIndex)
+					{
+						ResolveWildCardPin(Pin, OldPreferredTypes[i].TypeIndex, bSetupUndoRedo);
+					}
+				}
+			}
+			if (bSetupUndoRedo)
+			{
+				ActionStack->EndAction(Action);
+			}
 			return true;
 		}
 
@@ -15488,31 +16112,54 @@ bool URigVMController::PrepareTemplatePinForType(URigVMPin* InPin, const TArray<
 
 			ensure(LinksToBreak.Num() > 0);
 			PrepareTemplatePinForType(InPin, InTypeIndices, bSetupUndoRedo);
+
+			// Reapply the valid preferred types
+			for (int32 i=0; i<OldPreferredTypes.Num(); ++i)
+			{
+				if (URigVMPin* Pin = TemplateNode->FindPin(OldPreferredTypes[i].Argument.ToString()))
+				{
+					if (Pin->IsWildCard())
+					{
+						if (!TemplateNode->FilteredSupportsType(Pin, OldPreferredTypes[i].TypeIndex))
+						{
+							continue;
+						}
+					}
+				
+					if (Pin->IsWildCard() || Pin->GetTypeIndex() == OldPreferredTypes[i].TypeIndex)
+					{
+						ResolveWildCardPin(Pin, OldPreferredTypes[i].TypeIndex, bSetupUndoRedo);
+					}
+				}
+			}
+			if (bSetupUndoRedo)
+			{
+				ActionStack->EndAction(Action);
+			}
 			return true;
 		}
 	}
 
+	if (bSetupUndoRedo)
+	{
+		ActionStack->CancelAction(Action, this);
+	}
 	return false;
 }
 
-TArray<TRigVMTypeIndex> URigVMController::GetWildcardFilteredTypeIndices(URigVMPin* InPin)
+TArray<TRigVMTypeIndex> URigVMController::GetFilteredTypes(URigVMPin* InPin)
 {
-	TArray<TRigVMTypeIndex> Types;
-	if (!InPin->IsWildCard())
-	{
-		Types.Add(InPin->GetTypeIndex());
-		return Types;
-	}
-
-	if (URigVMTemplateNode* TemplateNode = Cast<URigVMTemplateNode>(InPin->GetNode()))
+	if (!InPin->IsStructMember())
 	{
-		if (!TemplateNode->IsSingleton())
+		if (URigVMTemplateNode* TemplateNode = Cast<URigVMTemplateNode>(InPin->GetNode()))
 		{
-			Types = TemplateNode->GetFilteredTypesForPin(InPin);
+			if (!TemplateNode->IsSingleton())
+			{
+				return TemplateNode->GetFilteredTypesForPin(InPin);
+			}
 		}
 	}
-	
-	return Types;
+	return {InPin->GetTypeIndex()};
 }
 
 bool URigVMController::ResolveWildCardPin(const FString& InPinPath, const FString& InCPPType, const FName& InCPPTypeObjectPath, bool bSetupUndoRedo, bool bPrintPythonCommand)
@@ -15608,18 +16255,50 @@ bool URigVMController::ResolveWildCardPin(const FString& InPinPath, TRigVMTypeIn
 bool URigVMController::ResolveWildCardPin(URigVMPin* InPin, TRigVMTypeIndex InTypeIndex, bool bSetupUndoRedo,
 	bool bPrintPythonCommand)
 {
-	if(InPin->IsWildCard())
+	if (InPin->IsStructMember())
 	{
-		ensure(InPin->GetNode()->IsA<URigVMTemplateNode>());
-		
-		FRigVMBaseAction Action;
-		if (bSetupUndoRedo)
-		{
-			Action.Title = TEXT("Resolve Wildcard Pin");
+		return false;
+	}
+
+	URigVMPin* RootPin = InPin;
+	TRigVMTypeIndex Type = InTypeIndex;	
+	while (RootPin->IsArrayElement())
+	{
+		RootPin = InPin->GetParentPin();
+		Type = FRigVMRegistry::Get().GetArrayTypeFromBaseTypeIndex(Type);
+	}
+	
+	if(RootPin->IsWildCard())
+	{
+		ensure(RootPin->GetNode()->IsA<URigVMTemplateNode>());
+		
+		FRigVMBaseAction Action;
+		if (bSetupUndoRedo)
+		{
+			Action.Title = TEXT("Resolve Wildcard Pin");
 			ActionStack->BeginAction(Action);
 		}
 		
-		URigVMTemplateNode* TemplateNode = CastChecked<URigVMTemplateNode>(InPin->GetNode());
+		URigVMTemplateNode* TemplateNode = CastChecked<URigVMTemplateNode>(RootPin->GetNode());
+
+		// Add preferred type
+		AddPreferredType(TemplateNode, RootPin->GetFName(), Type, bSetupUndoRedo);
+
+		if (TemplateNode->IsA<URigVMFunctionEntryNode>() || TemplateNode->IsA<URigVMFunctionReturnNode>())
+		{
+			if (TemplateNode->GetTemplate()->FindArgument(RootPin->GetFName()) == nullptr)
+			{
+				if (!AddArgumentForPin(RootPin, nullptr, bSetupUndoRedo))
+				{
+					if (bSetupUndoRedo)
+					{
+						ActionStack->CancelAction(Action, this);
+					}
+					return false;
+				}
+			}
+		}
+		
 		if (!PrepareTemplatePinForType(InPin, {InTypeIndex}, bSetupUndoRedo))
 		{
 			if (bSetupUndoRedo)
@@ -15628,14 +16307,16 @@ bool URigVMController::ResolveWildCardPin(URigVMPin* InPin, TRigVMTypeIndex InTy
 			}
 			return false;
 		}
+				
+		
 
-		TArray<FRigVMTemplatePreferredType> NewPreferredPermutationTypes = TemplateNode->PreferredPermutationPairs;
-		NewPreferredPermutationTypes.Add({FRigVMTemplatePreferredType(InPin->GetFName(), InTypeIndex)});
-		if (bSetupUndoRedo)
+		URigVMGraph* Graph = GetGraph();
+		if (URigVMGraph* ParentGraph = Graph->GetParentGraph())
 		{
-			ActionStack->AddAction(FRigVMSetPreferredTemplatePermutationsAction(TemplateNode, NewPreferredPermutationTypes));
+			FRigVMControllerGraphGuard GraphGuard(this, ParentGraph, bSetupUndoRedo);
+			UpdateLibraryTemplate(Graph->GetTypedOuter<URigVMLibraryNode>(), bSetupUndoRedo);
 		}
-		TemplateNode->PreferredPermutationPairs = NewPreferredPermutationTypes;
+		
 		if (bSetupUndoRedo)
 		{
 			ActionStack->EndAction(Action);
@@ -15643,6 +16324,32 @@ bool URigVMController::ResolveWildCardPin(URigVMPin* InPin, TRigVMTypeIndex InTy
 		
 		return true;		
 	}
+	else if (FRigVMRegistry::Get().CanMatchTypes(RootPin->GetTypeIndex(), Type, true))
+	{
+		if (URigVMTemplateNode* TemplateNode = Cast<URigVMTemplateNode>(RootPin->GetNode()))
+		{
+			TRigVMTypeIndex TypeIndex = TemplateNode->GetPreferredType(RootPin->GetFName());
+			if (TypeIndex == Type && RootPin->GetTypeIndex() == Type)
+			{
+				return true;
+			}
+			
+			if (TypeIndex != Type && TypeIndex != INDEX_NONE)
+			{
+				RemovePreferredType(TemplateNode, RootPin->GetFName(), bSetupUndoRedo);
+				TypeIndex = INDEX_NONE;
+			}
+			if (TypeIndex == INDEX_NONE)
+			{
+				AddPreferredType(TemplateNode, RootPin->GetFName(), Type, bSetupUndoRedo);
+			}
+	
+			if (RootPin->GetTypeIndex() != Type)
+			{
+				UpdateTemplateNodePinTypes(TemplateNode, bSetupUndoRedo, false);
+			}
+		}
+	}
 	else
 	{
 		OpenUndoBracket(TEXT("Resolving wildcard pin"));
@@ -15676,12 +16383,21 @@ bool URigVMController::UpdateFilteredPermutations(URigVMPin* InPin, URigVMPin* I
 	{
 		if (bSetupUndoRedo)
 		{
-			FRigVMSetTemplateFilteredPermutationsAction Action(Node, InPin->FindLinkForPin(InLinkedPin), OldPermutations);
+			FRigVMSetTemplateFilteredPermutationsAction Action(Node, OldPermutations);
 			ActionStack->BeginAction(Action);
 			ActionStack->AddAction(Action);
 			ActionStack->EndAction(Action);
 		}
 
+		// If updating the filtered permutations of entry or return nodes, update the library template
+		if (Node->IsA<URigVMFunctionEntryNode>() || Node->IsA<URigVMFunctionReturnNode>())
+		{
+			if (URigVMLibraryNode* LibraryNode = Node->GetTypedOuter<URigVMLibraryNode>())
+			{
+				UpdateLibraryTemplate(LibraryNode, bSetupUndoRedo);
+			}
+		}
+
 		return true;
 	}
 
@@ -15706,11 +16422,20 @@ bool URigVMController::UpdateFilteredPermutations(URigVMPin* InPin, const TArray
 	{
 		if (bSetupUndoRedo)
 		{
-			FRigVMSetTemplateFilteredPermutationsAction Action(Node, nullptr, OldPermutations);
+			FRigVMSetTemplateFilteredPermutationsAction Action(Node, OldPermutations);
 			ActionStack->BeginAction(Action);
 			ActionStack->AddAction(Action);
 			ActionStack->EndAction(Action);
 		}
+		
+		// If updating the filtered permutations of entry or return nodes, update the library template
+        if (Node->IsA<URigVMFunctionEntryNode>() || Node->IsA<URigVMFunctionReturnNode>())
+        {
+        	if (URigVMLibraryNode* LibraryNode = Node->GetTypedOuter<URigVMLibraryNode>())
+        	{
+        		UpdateLibraryTemplate(LibraryNode, bSetupUndoRedo);
+        	}
+        }
 
 		return true;
 	}
@@ -15718,10 +16443,28 @@ bool URigVMController::UpdateFilteredPermutations(URigVMPin* InPin, const TArray
 	return false;
 }
 
-bool URigVMController::UpdateTemplateNodePinTypes(URigVMTemplateNode* InNode, bool bSetupUndoRedo)
+bool URigVMController::UpdateTemplateNodePinTypes(URigVMTemplateNode* InNode, bool bSetupUndoRedo, bool bInitializeDefaultValue)
 {
+	check(InNode->GetGraph() == GetGraph());
 	bool bAnyTypeChanged = false;
 
+	TArray<int32> FilteredPermutations = InNode->GetFilteredPermutationsIndices();
+
+	// Try to pick permutations that respects the preferred types
+	TArray<int32> PreferredPermutations = InNode->FindPermutationsForTypes(InNode->PreferredPermutationPairs, false);
+	FilteredPermutations = FilteredPermutations.FilterByPredicate([&](const int32& OtherPermutation) { return PreferredPermutations.Contains(OtherPermutation); });
+	if (FilteredPermutations.IsEmpty())
+	{
+		FilteredPermutations = InNode->GetFilteredPermutationsIndices();
+
+		// Try to pick one permutations that respects the preferred types (with casting)
+		PreferredPermutations = InNode->FindPermutationsForTypes(InNode->PreferredPermutationPairs, true);
+		FilteredPermutations = FilteredPermutations.FilterByPredicate([&](const int32& OtherPermutation) { return PreferredPermutations.Contains(OtherPermutation); });
+	}
+	int32 ReducedToPermutation = INDEX_NONE;
+
+	const FRigVMTemplate* Template = InNode->GetTemplate();
+
 	// First unresolve any pins that need unresolving, then resolve everything else
 	for (int32 UnresolveResolve=0; UnresolveResolve<2; ++UnresolveResolve)
 	{
@@ -15733,60 +16476,88 @@ bool URigVMController::UpdateTemplateNodePinTypes(URigVMTemplateNode* InNode, bo
 				continue;
 			}
 			
+			bool bNeedsToUnresolve = false;
 			if (Pin->GetDirection() == ERigVMPinDirection::Hidden)
 			{
 				continue;
 			}
 		
-			TArray<TRigVMTypeIndex> Types = InNode->GetFilteredTypesForPin(Pin);
+			TArray<TRigVMTypeIndex> Types;
+			Types.Reserve(FilteredPermutations.Num());
+			for (const int32& Permutation : FilteredPermutations)
+			{
+				if (const FRigVMTemplateArgument* Argument = Template->FindArgument(Pin->GetFName()))
+				{
+					Types.AddUnique(Argument->TypeIndices[Permutation]);
+				}
+			}
+			
 			if (Types.IsEmpty())
 			{
-				continue;
+				// We might be in the case where the pin doesn't own an argument yet, but the node has a preferred type
+				if (!InNode->GetTemplate()->FindArgument(Pin->GetFName()))
+				{
+					for (int32 i=0; i<InNode->PreferredPermutationPairs.Num(); ++i)
+					{
+						if (Pin->GetFName() == InNode->PreferredPermutationPairs[i].Argument)
+						{
+							Types = {InNode->PreferredPermutationPairs[i].TypeIndex};
+							break;
+						}
+					}
+				}
+
+				if (Types.IsEmpty())
+				{
+					bNeedsToUnresolve = true;
+				}
 			}
 
+			// Try to reduce the number of types to one
 			if (Types.Num() > 1)
 			{
-			
-				bool bCanReduceToSingleType = true;
-				for (int32 i=1; i<Types.Num(); ++i)
+				TRigVMTypeIndex PreferredType = Pin->GetTypeIndex();
+				if (ReducedToPermutation != INDEX_NONE)
 				{
-					if (!FRigVMRegistry::Get().CanMatchTypes(Types[0], Types[i], true))
+					if (const FRigVMTemplateArgument* Argument = Template->FindArgument(Pin->GetFName()))
 					{
-						bCanReduceToSingleType = false;
-						break;
+						PreferredType = Argument->TypeIndices[ReducedToPermutation];
 					}
 				}
-			
-				// Select the same Pin type if it has one
-				if (bCanReduceToSingleType)
+				InNode->TryReduceTypesToSingle(Types, PreferredType);
+				if (Types.Num() == 1)
 				{
-					int32 PreferredIndex = 0;
-					if (!Pin->IsWildCard())
+					// Make sure the other pins use the same permutation
+					ReducedToPermutation = *FilteredPermutations.FindByPredicate([&](const int32& Permutation)
 					{
-						for (int32 i=0; i<Types.Num(); ++i)
+						if (const FRigVMTemplateArgument* Argument = Template->FindArgument(Pin->GetFName()))
 						{
-							if (Types[i] == Pin->GetTypeIndex())
-							{
-								PreferredIndex = i;
-								break;
-							}
+							return Types[0] == Argument->TypeIndices[Permutation];
 						}
-					}
-					Types = {Types[PreferredIndex]};
+						return false;
+					});
 				}
 			}
-		
-			if (Types.Num() > 1)
+			
+			if (bNeedsToUnresolve || Types.Num() > 1)
 			{
 				// Unresolve
 				if (Pin->HasInjectedNodes())
 				{
 					EjectNodeFromPin(Pin, bSetupUndoRedo);
 				}
-			
+				
 				const FRigVMTemplateArgument* Argument = InNode->GetTemplate()->FindArgument(*Pin->GetName());
-				const FRigVMTemplateArgument::EArrayType ArrayType = Argument->GetArrayType();
-			
+				FRigVMTemplateArgument::EArrayType ArrayType;
+				if (Argument)
+				{
+					ArrayType = Argument->GetArrayType();
+				}
+				else
+				{
+					ArrayType = Pin->IsArray() ? FRigVMTemplateArgument::EArrayType::EArrayType_ArrayValue : FRigVMTemplateArgument::EArrayType::EArrayType_SingleValue;
+				}
+				
 				FString CPPType = RigVMTypeUtils::GetWildCardCPPType();
 				UObject* CPPObjectType = RigVMTypeUtils::GetWildCardCPPTypeObject();
 
@@ -15802,7 +16573,7 @@ bool URigVMController::UpdateTemplateNodePinTypes(URigVMTemplateNode* InNode, bo
 				if (Pin->GetCPPType() != CPPType || Pin->GetCPPTypeObject() != CPPObjectType)
 				{
 					bAnyTypeChanged = true;
-					ChangePinType(Pin, CPPType, CPPObjectType, bSetupUndoRedo, false, false, false);
+					ChangePinType(Pin, CPPType, CPPObjectType, bSetupUndoRedo, false, false, false, bInitializeDefaultValue);
 				}
 			}
 			else if (Types.Num() == 1)
@@ -15811,7 +16582,22 @@ bool URigVMController::UpdateTemplateNodePinTypes(URigVMTemplateNode* InNode, bo
 				if (Pin->GetTypeIndex() != Types[0])
 				{
 					bAnyTypeChanged = true;
-					ChangePinType(Pin, Types[0], bSetupUndoRedo, false, false, false);
+					ChangePinType(Pin, Types[0], bSetupUndoRedo, false, false, false, bInitializeDefaultValue);
+				}
+
+				// Once an interface pin is resolved, add it as a preferred type
+				if (InNode->IsA<URigVMFunctionEntryNode>() || InNode->IsA<URigVMFunctionReturnNode>())
+				{
+					TRigVMTypeIndex TypeIndex = InNode->GetPreferredType(Pin->GetFName());
+					if (TypeIndex != Types[0] && TypeIndex != INDEX_NONE)
+					{
+						RemovePreferredType(InNode, Pin->GetFName(), bSetupUndoRedo);
+						TypeIndex = INDEX_NONE;
+					}
+					if (TypeIndex == INDEX_NONE)
+					{
+						AddPreferredType(InNode, Pin->GetFName(), Types[0], bSetupUndoRedo);
+					}
 				}
 			}
 			else
@@ -15837,6 +16623,17 @@ bool URigVMController::PropagateTemplateFilteredTypes(URigVMTemplateNode* InNode
 			bool bIsTemplate = false;
 			if (URigVMTemplateNode* OtherTemplate = Cast<URigVMTemplateNode>(OtherPin->GetNode()))
 			{
+				if (OtherTemplate->IsA<URigVMFunctionEntryNode>() || OtherTemplate->IsA<URigVMFunctionReturnNode>())
+				{
+					if (!OtherTemplate->GetTemplate()->FindArgument(OtherPin->GetFName()))
+					{
+						if (!AddArgumentForPin(OtherPin, Pin, bSetupUndoRedo))
+						{
+							continue;
+						}
+					}
+				}				
+				
 				if (!OtherTemplate->IsSingleton())
 				{
 					bIsTemplate = true;
@@ -15861,45 +16658,445 @@ bool URigVMController::PropagateTemplateFilteredTypes(URigVMTemplateNode* InNode
 				}								
 			}
 
-			if (!bIsTemplate)
+			if (!bIsTemplate)
+			{
+				if (!InNode->FilteredSupportsType(Pin, OtherPin->GetTypeIndex()))
+				{
+					URigVMLink* Link = Pin->FindLinkForPin(OtherPin);
+					ensureMsgf(!ActionStack->BracketActions.IsEmpty(), TEXT("Unexpected link broken %s in package %s"), *Link->GetPinPathRepresentation(), *GetPackage()->GetPathName());
+					BreakLink(Link->GetSourcePin(), Link->GetTargetPin(), bSetupUndoRedo);
+					return false;
+				}
+			}
+		}
+		return true;
+	};
+	
+	for(URigVMPin* Pin : InNode->GetPins())
+	{
+		if (!InNode->GetTemplate()->FindArgument(Pin->GetFName()))
+		{
+			continue;
+		}
+		
+		if (!UpdateAndPropagate(Pin))
+		{
+			return false;
+		}
+
+		if (Pin->IsArray())
+		{
+			if (Pin->GetSubPins().Num() > 0)
+			{
+				for (URigVMPin* SubPin : Pin->GetSubPins())
+				{
+					if (!UpdateAndPropagate(SubPin))
+					{
+						return false;
+					}
+				}
+			}
+		}
+	}
+	return true;
+}
+
+bool URigVMController::AddPreferredType(URigVMTemplateNode* InNode, const FName& InPinName, const int32 InPreferredTypeIndex, bool bSetupUndoRedo)
+{
+	check(InNode->FindPin(InPinName.ToString()));
+
+	FRigVMTemplatePreferredType* PreferredType = InNode->PreferredPermutationPairs.FindByPredicate([&](const FRigVMTemplatePreferredType& Other)
+	{
+		return Other.Argument == InPinName;
+	});
+
+	if (PreferredType)
+	{
+		return PreferredType->TypeIndex == InPreferredTypeIndex;		
+	}
+
+	TArray<FRigVMTemplatePreferredType> NewPreferredTypes = InNode->PreferredPermutationPairs;
+	NewPreferredTypes.Add(FRigVMTemplatePreferredType(InPinName, InPreferredTypeIndex));
+	if (bSetupUndoRedo)
+	{
+		ActionStack->AddAction(FRigVMSetPreferredTemplatePermutationsAction(InNode, NewPreferredTypes));
+	}
+	InNode->PreferredPermutationPairs = NewPreferredTypes;
+	return true;
+}
+
+bool URigVMController::RemovePreferredType(URigVMTemplateNode* InNode, const FName& InPinName, bool bSetupUndoRedo)
+{
+	int32 FoundIndex = INDEX_NONE;
+	for (int32 i=0; i<InNode->PreferredPermutationPairs.Num(); ++i)
+	{
+		if (InNode->PreferredPermutationPairs[i].Argument == InPinName)
+		{
+			FoundIndex = i;
+			break;
+		}
+	}
+	
+	if (FoundIndex == INDEX_NONE)
+	{
+		return false;
+	}
+
+	TArray<FRigVMTemplatePreferredType> NewPreferredTypes = InNode->PreferredPermutationPairs;
+	NewPreferredTypes.RemoveAt(FoundIndex);	
+	
+	if (bSetupUndoRedo)
+	{
+		ActionStack->AddAction(FRigVMSetPreferredTemplatePermutationsAction(InNode, NewPreferredTypes));
+	}
+	InNode->PreferredPermutationPairs = NewPreferredTypes;
+	return true;
+}
+
+bool URigVMController::ShouldPinOwnArgument(URigVMPin* InPin)
+{
+	if (!IsValidGraph())
+	{
+		return false;
+	}
+
+	if (!bIsTransacting && !IsGraphEditable())
+	{
+		return false;
+	}
+
+	URigVMLibraryNode* LibraryNode = Cast<URigVMLibraryNode>(InPin->GetNode());
+	if (LibraryNode == nullptr)
+	{
+		return false;
+	}
+
+	URigVMGraph* ContainedGraph = LibraryNode->GetContainedGraph();
+	check(ContainedGraph);
+
+	if (ContainedGraph->IsRootGraph())
+	{
+		return false;
+	}
+
+	URigVMFunctionEntryNode* EntryNode = ContainedGraph->GetEntryNode();
+	URigVMFunctionReturnNode* ReturnNode = ContainedGraph->GetReturnNode();
+
+	auto IsPinUsingArgument = [](URigVMPin* TestPin)
+	{
+		// If any subpins that are struct members are connected, then it should own an argument
+		TArray<URigVMPin*> ToProcessPins = TestPin->GetSubPins();
+		for (int32 i=0; i<ToProcessPins.Num(); ++i)
+		{
+			URigVMPin* CurPin = ToProcessPins[i];
+			if (CurPin->IsStructMember() && CurPin->IsLinked())
+			{
+				return true;
+			}
+			ToProcessPins.Append(CurPin->GetSubPins());
+		}
+		
+		for (int32 i=0; i<2; ++i)
+		{
+			TArray<URigVMPin*> ConnectedPins = (i==0) ? TestPin->GetLinkedSourcePins() : TestPin->GetLinkedTargetPins();
+			for (URigVMPin* ConnectedPin : ConnectedPins)
+			{
+				if (URigVMTemplateNode* TemplateNode = Cast<URigVMTemplateNode>(ConnectedPin->GetNode()))
+				{
+					if (TemplateNode->GetFilteredPermutationsIndices().Num() < 10)
+					{
+						return true;
+					}
+				}
+				else
+				{					
+					return true;
+				}
+			}								
+		}
+		return false;
+	};
+
+	bool bIsInputPin = false;
+	for (int32 i=0; i<2; ++i)
+	{
+		if (URigVMTemplateNode* InterfaceNode = (i == 0) ? (URigVMTemplateNode*)EntryNode : (URigVMTemplateNode*)ReturnNode)
+		{
+			if (URigVMPin* InterfacePin = InterfaceNode->FindPin(InPin->GetName()))
+			{
+				for (const FRigVMTemplatePreferredType& PreferredType : InterfaceNode->PreferredPermutationPairs)
+				{
+					if (PreferredType.Argument == InPin->GetFName())
+					{
+						return true;
+					}
+				}
+				if (IsPinUsingArgument(InterfacePin))
+				{
+					return true;
+				}
+				if (i == 0)
+				{
+					bIsInputPin = true;
+				}
+			}
+		}
+	}
+
+	// Find input variable nodes
+	for (URigVMNode* Node : ContainedGraph->GetNodes())
+	{
+		if (URigVMVariableNode* VariableNode = Cast<URigVMVariableNode>(Node))
+		{
+			if (VariableNode->GetVariableName() == InPin->GetFName())
+			{
+				if (IsPinUsingArgument(VariableNode->GetValuePin()))
+				{
+					return true;
+				}
+			}
+		}
+	}
+
+	return false;
+}
+
+bool URigVMController::UpdateLibraryTemplate(URigVMLibraryNode* LibraryNode, bool bSetupUndoRedo)
+{
+	if (!IsValidGraph())
+	{
+		return false;
+	}
+
+	if (!bIsTransacting && !IsGraphEditable())
+	{
+		return false;
+	}
+
+	if (!LibraryNode)
+	{
+		return false;
+	}
+
+	URigVMGraph* ContainedGraph = LibraryNode->GetContainedGraph();
+	check(ContainedGraph);
+
+	if (ContainedGraph->IsRootGraph())
+	{
+		return false;
+	}
+
+	URigVMFunctionEntryNode* EntryNode = ContainedGraph->GetEntryNode();
+	URigVMFunctionReturnNode* ReturnNode = ContainedGraph->GetReturnNode();
+	const TArray<URigVMTemplateNode*> InterfaceNodes = {EntryNode, ReturnNode};
+	const FRigVMTemplate* CurTemplate = LibraryNode->GetTemplate();
+
+	// Remove unused arguments from template
+	TArray<FName> ArgumentsToRemove;
+	{
+		for (const FRigVMTemplateArgument& CurArgument : CurTemplate->Arguments)
+		{
+			URigVMPin* LibraryPin = LibraryNode->FindPin(CurArgument.Name.ToString());
+			if (!LibraryPin || !ShouldPinOwnArgument(LibraryPin))
+			{
+				ArgumentsToRemove.Add(CurArgument.Name);
+			}
+		}
+	}
+
+	// If interface node contains a preferred type that differs from the type in the template, keep the interface type
+	TMap<FName, TRigVMTypeIndex> ArgPreferredType;
+	for (URigVMTemplateNode* InterfaceNode : InterfaceNodes)
+	{
+		if (InterfaceNode)
+		{
+			for (int32 i=0; i<InterfaceNode->PreferredPermutationPairs.Num(); ++i)
+			{
+				const FRigVMTemplatePreferredType& PreferredType = InterfaceNode->PreferredPermutationPairs[i];
+				if (!ArgumentsToRemove.Contains(PreferredType.Argument) && !ArgPreferredType.Contains(PreferredType.Argument))
+				{
+					ArgPreferredType.Add(PreferredType.Argument, PreferredType.TypeIndex);
+				}			
+			}
+		}
+	}
+
+	FRigVMTemplate NewTemplate;
+	NewTemplate.Index = CurTemplate->Index;
+	
+	NewTemplate.Arguments.Reserve(CurTemplate->Arguments.Num());
+	int32 ArgIndex = 0;
+	for (const FRigVMTemplateArgument& CurArgument : CurTemplate->Arguments)
+	{
+		if (ArgumentsToRemove.Contains(CurArgument.Name))
+		{
+			continue;
+		}
+		
+		FRigVMTemplateArgument NewArgument;
+		NewArgument.Index = ArgIndex++;
+		NewArgument.Name = CurArgument.Name;
+		NewArgument.Direction = CurArgument.Direction;
+		NewTemplate.Arguments.Add(NewArgument);
+	}
+
+	// The filtered permutations of the entry and return nodes reference the library template permutations
+	// so they can be combined
+	TArray<int32> CombinedFilteredPermutations;
+	{
+		if (EntryNode)
+		{
+			CombinedFilteredPermutations = EntryNode->FilteredPermutations;
+		}
+		if (ReturnNode)
+		{
+			CombinedFilteredPermutations = CombinedFilteredPermutations.FilterByPredicate([&](int32 Index)
+			{
+				return ReturnNode->FilteredPermutations.Contains(Index);
+			});			
+		}
+	}
+
+	TArray<FString> AllPermutationTypes;
+	for (int32 PermutationIndex=0; PermutationIndex<CombinedFilteredPermutations.Num(); ++PermutationIndex)
+	{
+		int32 OldIndex = CombinedFilteredPermutations[PermutationIndex];
+
+		// For each permutation, create a string that identifies all arguments with their types
+		FString PermutationTypesStr;
+		TMap<const FName, TRigVMTypeIndex> TypesToAdd;
+		for (const FRigVMTemplateArgument& CurArgument : CurTemplate->Arguments)
+		{
+			if (ArgumentsToRemove.Contains(CurArgument.Name))
+			{
+				continue;
+			}
+
+			if (TRigVMTypeIndex* TypeIndex = ArgPreferredType.Find(CurArgument.Name))
+			{
+				TypesToAdd.Add(CurArgument.Name, *TypeIndex);
+				const int32 TypeInt = *TypeIndex;				
+				PermutationTypesStr += FString::Printf(TEXT("%s:%d;"), *CurArgument.Name.ToString(), TypeInt);	
+			}
+			else
+			{
+				TypesToAdd.Add(CurArgument.Name, CurArgument.TypeIndices[OldIndex]);
+				const int32 TypeInt = CurArgument.TypeIndices[OldIndex];
+				PermutationTypesStr += FString::Printf(TEXT("%s:%d;"), *CurArgument.Name.ToString(), TypeInt);
+			}
+		}
+		if (AllPermutationTypes.Contains(PermutationTypesStr))
+		{
+			continue;
+		}
+
+		AllPermutationTypes.Add(PermutationTypesStr);
+
+		// If this is a new permutation, add it to the template 
+		int32 NewArgumentIndex=0;
+		for (TPair<const FName, TRigVMTypeIndex> TypePair : TypesToAdd)
+		{
+			FRigVMTemplateArgument& NewArgument = NewTemplate.Arguments[NewArgumentIndex++];
+			const TRigVMTypeIndex& TypeIndex = TypePair.Value; 
+			NewArgument.TypeIndices.Add(TypeIndex);
+			if (TArray<int32>* Permutations = NewArgument.TypeToPermutations.Find(TypeIndex))
+			{
+				Permutations->Add(PermutationIndex);
+			}
+			else
+			{
+				NewArgument.TypeToPermutations.Add(TypeIndex, {PermutationIndex});
+			}
+		}
+		NewTemplate.Permutations.Add(INDEX_NONE);
+	}
+
+	// Add a permutation if non was created but we are about to add arguments
+	if (NewTemplate.Permutations.IsEmpty() && !ArgPreferredType.IsEmpty())
+	{
+		NewTemplate.Permutations.Add(INDEX_NONE);
+	}
+
+	// Add any arguments with preferred types that are missing
+	for (TPair<FName, TRigVMTypeIndex>& ArgPair : ArgPreferredType)
+	{
+		bool bMissing = true;
+		for (FRigVMTemplateArgument& Argument : NewTemplate.Arguments)
+		{
+			if (Argument.Name == ArgPair.Key)
 			{
-				if (!InNode->FilteredSupportsType(Pin, OtherPin->GetTypeIndex()))
-				{
-					URigVMLink* Link = Pin->FindLinkForPin(OtherPin);
-					ensureMsgf(!ActionStack->BracketActions.IsEmpty(), TEXT("Unexpected link broken %s in package %s"), *Link->GetPinPathRepresentation(), *GetPackage()->GetPathName());
-					BreakLink(Link->GetSourcePin(), Link->GetTargetPin(), bSetupUndoRedo);
-					return false;
-				}
+				bMissing = false;
+				break;
 			}
 		}
-		return true;
-	};
-	
-	for(URigVMPin* Pin : InNode->GetPins())
-	{
-		if (!UpdateAndPropagate(Pin))
+
+		if (bMissing)
 		{
-			return false;
+			FRigVMTemplateArgument NewArgument;
+			NewArgument.Index = NewTemplate.Arguments.Num();
+			NewArgument.Name = ArgPair.Key;
+			NewArgument.Direction = LibraryNode->FindPin(ArgPair.Key.ToString())->GetDirection();
+
+			TArray<int32> Permutations;
+			NewArgument.TypeIndices.Reserve(NewTemplate.NumPermutations());
+			Permutations.Reserve(NewTemplate.NumPermutations());
+			for (int32 i=0; i<NewTemplate.NumPermutations(); ++i)
+			{
+				NewArgument.TypeIndices.Add(ArgPair.Value);
+				Permutations.Add(i);
+			}
+			NewArgument.TypeToPermutations.Add(ArgPair.Value, Permutations);
+			NewTemplate.Arguments.Add(NewArgument);			
 		}
+	}
 
-		if (Pin->IsArray())
+	NewTemplate.ComputeNotationFromArguments(LibraryNode->GetName());
+
+	if (bSetupUndoRedo)
+	{
+		FRigVMControllerGraphGuard GraphGuard(this, LibraryNode->GetTypedOuter<URigVMGraph>(), bSetupUndoRedo);
+		ActionStack->AddAction(FRigVMSetLibraryTemplateAction(LibraryNode, NewTemplate));		
+	}
+	LibraryNode->Template = NewTemplate;
+	Notify(ERigVMGraphNotifType::NodeDescriptionChanged, LibraryNode);
+
+	// Update Entry and Return filtered permutations to all permutations
+	{
+		FRigVMControllerGraphGuard GraphGuard(this, ContainedGraph, bSetupUndoRedo);
+		for (int32 EntryReturn=0; EntryReturn<2; ++EntryReturn)
 		{
-			if (Pin->GetSubPins().Num() > 0)
+			URigVMTemplateNode* Node = (EntryReturn == 0) ? (URigVMTemplateNode*)EntryNode : (URigVMTemplateNode*)ReturnNode;
+			if (Node)
 			{
-				for (URigVMPin* SubPin : Pin->GetSubPins())
+				TArray<int32> OldPermutations = Node->FilteredPermutations;
+				Node->FilteredPermutations.SetNumUninitialized(NewTemplate.NumPermutations());
+				for (int32 i=0; i<Node->FilteredPermutations.Num(); ++i)
 				{
-					if (!UpdateAndPropagate(SubPin))
-					{
-						return false;
-					}
+					Node->FilteredPermutations[i] = i;
+				}		
+
+				if (bSetupUndoRedo)
+				{
+					ActionStack->AddAction(FRigVMSetTemplateFilteredPermutationsAction(Node, OldPermutations));
 				}
+
+				UpdateTemplateNodePinTypes(Node, bSetupUndoRedo);
+				Notify(ERigVMGraphNotifType::NodeDescriptionChanged, Node);
 			}
 		}
 	}
+
+	// Update outer graphs
+	if (!bIsTransacting && !bSuspendRecomputingOuterTemplateFilters)
+	{
+		FRigVMControllerGraphGuard GraphGuard(this, LibraryNode->GetGraph(), bSetupUndoRedo);
+		RecomputeAllTemplateFilteredPermutations(bSetupUndoRedo);
+	}
+
 	return true;
 }
 
-bool URigVMController::RecomputeAllTemplateFilteredTypes(bool bSetupUndoRedo)
+bool URigVMController::RecomputeAllTemplateFilteredPermutations(bool bSetupUndoRedo)
 {
 	if (!IsValidGraph())
 	{
@@ -15925,18 +17122,13 @@ bool URigVMController::RecomputeAllTemplateFilteredTypes(bool bSetupUndoRedo)
 				continue;
 			}
 
-			if (!TemplateNode->PreferredPermutationPairs.IsEmpty())
-			{
-				continue;
-			}
-
 			if (const FRigVMTemplate* Template = TemplateNode->GetTemplate())
 			{
 				for (int32 i=0; i<Template->NumArguments(); ++i)
 				{
 					const FRigVMTemplateArgument* Argument = Template->GetArgument(i);
 					URigVMPin* Pin = TemplateNode->FindPin(Argument->GetName().ToString());
-					if (Pin->IsWildCard())
+					if (!Pin || Pin->IsWildCard())
 					{
 						continue;
 					}
@@ -15947,137 +17139,236 @@ bool URigVMController::RecomputeAllTemplateFilteredTypes(bool bSetupUndoRedo)
 		}
 	}
 
-	// Initialize all filtered permutations to unresolved (except for preferred permutation)
-	InitializeAllTemplateFiltersInGraph(bSetupUndoRedo, false);
-
-	// Apply all links to update filtered permutations
-	TArray<URigVMLink*> SortedLinks = Graph->GetLinks();
-
-	// Solve for unit nodes first, other templates are more expensive to solve. This way we are avoiding solving
-	// links like (reroute-reroute) before their filtered permutations are reduced
-	Algo::Sort(SortedLinks, [](const URigVMLink* A, const URigVMLink* B)
+	bool bAnyTypeChanged = false;
+	bool bAnyLinkBroken = false;
 	{
-		const URigVMPin* APinSource = A->GetSourcePin();
-		const URigVMPin* APinTarget = A->GetTargetPin();
-		const URigVMPin* BPinSource = B->GetSourcePin();
-		const URigVMPin* BPinTarget = B->GetTargetPin();
-		if (!IsValid(APinSource) || !IsValid(APinTarget) || !IsValid(BPinSource) || !IsValid(BPinTarget))
-		{
-			return false;
-		}
-	
-		const bool bASourceIsUnitNode = APinSource->GetNode()->IsA<URigVMUnitNode>();
-		const bool bATargetIsUnitNode = APinTarget->GetNode()->IsA<URigVMUnitNode>();
-		const bool bBSourceIsUnitNode = BPinSource->GetNode()->IsA<URigVMUnitNode>();
-		const bool bBTargetIsUnitNode = BPinTarget->GetNode()->IsA<URigVMUnitNode>();
+		TGuardValue<bool> GuardSuspendRecomputeOuterFilters(bSuspendRecomputingOuterTemplateFilters, true);
 		
-		if (bASourceIsUnitNode && bATargetIsUnitNode && (!bBSourceIsUnitNode || !bBTargetIsUnitNode))
-		{
-			return true;
-		}
+		// Initialize all filtered permutations to unresolved (except for preferred permutation)
+		// This will also remove all arguments from library templates
+		InitializeAllTemplateFiltersInGraph(bSetupUndoRedo, false);	
 
-		if ((bASourceIsUnitNode || bATargetIsUnitNode) && (!bBSourceIsUnitNode && !bBTargetIsUnitNode))
+		// Apply all links to update filtered permutations
+		TArray<URigVMLink*> SortedLinks = Graph->GetLinks().FilterByPredicate([](URigVMLink* Link)
 		{
-			return true;
-		}
+			// Filter out detached links
+			return Link->GetSourcePin()->IsLinkedTo(Link->GetTargetPin());
+		});
 
-		return false;		
-	});
+		// Solve for unit nodes first, other templates are more expensive to solve. This way we are avoiding solving
+		// links like (reroute-reroute) before their filtered permutations are reduced
+		Algo::Sort(SortedLinks, [](const URigVMLink* A, const URigVMLink* B)
+		{
+			const URigVMPin* APinSource = A->GetSourcePin();
+			const URigVMPin* APinTarget = A->GetTargetPin();
+			const URigVMPin* BPinSource = B->GetSourcePin();
+			const URigVMPin* BPinTarget = B->GetTargetPin();
+			if (!IsValid(APinSource) || !IsValid(APinTarget) || !IsValid(BPinSource) || !IsValid(BPinTarget))
+			{
+				return false;
+			}
 	
-	for (int32 i=0; i<SortedLinks.Num(); ++i)
-	{
-		URigVMLink* Link = SortedLinks[i];
+			const bool bASourceIsUnitNode = APinSource->GetNode()->IsA<URigVMUnitNode>();
+			const bool bATargetIsUnitNode = APinTarget->GetNode()->IsA<URigVMUnitNode>();
+			const bool bBSourceIsUnitNode = BPinSource->GetNode()->IsA<URigVMUnitNode>();
+			const bool bBTargetIsUnitNode = BPinTarget->GetNode()->IsA<URigVMUnitNode>();
 		
-		URigVMPin* OutputPin = Link->GetSourcePin();
-		URigVMPin* InputPin = Link->GetTargetPin();
+			if (bASourceIsUnitNode && bATargetIsUnitNode && (!bBSourceIsUnitNode || !bBTargetIsUnitNode))
+			{
+				return true;
+			}
+
+			if ((bASourceIsUnitNode || bATargetIsUnitNode) && (!bBSourceIsUnitNode && !bBTargetIsUnitNode))
+			{
+				return true;
+			}
+
+			return false;		
+		});
 
-		if (!IsValid(OutputPin) || !IsValid(InputPin))
+		// Keep a list of removed links so that we can only notify adding links when necessary
+		TArray<TPair<FString, FString>> LinksRemoved;
+		FDelegateHandle Delegate = OnModified().AddLambda([&LinksRemoved](ERigVMGraphNotifType NotifType, URigVMGraph* Graph, UObject* Subject) {
+			if (NotifType == ERigVMGraphNotifType::LinkAdded)
+			{
+				if (URigVMLink* Link = Cast<URigVMLink>(Subject))
+				{
+					LinksRemoved.Remove(TPair<FString, FString>(Link->GetSourcePin()->GetPinPath(), Link->GetTargetPin()->GetPinPath()));
+				}
+			}
+			else if (NotifType == ERigVMGraphNotifType::LinkRemoved)
+			{
+				if (URigVMLink* Link = Cast<URigVMLink>(Subject))
+				{
+					LinksRemoved.Add(TPair<FString, FString>(Link->GetSourcePin()->GetPinPath(), Link->GetTargetPin()->GetPinPath()));
+				}
+			}
+		});
+			
+		DetachLinksFromPinObjects(&SortedLinks);
+		TArray<URigVMLink*> LinksToRemove;
+		for (int32 i=0; i<SortedLinks.Num(); ++i)
 		{
-			continue;
-		}
+			URigVMLink* Link = SortedLinks[i];
 		
-		// If pin is a struct member, we should resolve to that type
-		if (URigVMTemplateNode* OutputNode = Cast<URigVMTemplateNode>(OutputPin->GetNode()))
-		{
-			if (!OutputNode->IsSingleton() && OutputNode->PreferredPermutationPairs.IsEmpty())
+			URigVMPin* OutputPin = Link->GetSourcePin();
+			URigVMPin* InputPin = Link->GetTargetPin();
+
+			if (!IsValid(OutputPin) || !IsValid(InputPin))
+			{
+				continue;
+			}
+
 			{
-				if (OutputPin->IsStructMember())
+				ActionStack->OpenUndoBracket(TEXT("Prepare template for type"));
+				bool bLinkRemoved = false;
+				for (int32 EdgeI=0; EdgeI<2 && !bLinkRemoved; ++EdgeI)
 				{
-					URigVMPin* RootPin = OutputPin->GetRootPin();
-					const TRigVMTypeIndex TypeIndex = TypesBeforeRecomputing.FindChecked(RootPin);
-					if (UpdateFilteredPermutations(RootPin, {TypeIndex}, bSetupUndoRedo))
+					URigVMPin* EdgePin = (EdgeI==0) ? OutputPin : InputPin;
+					if (EdgePin->IsStructMember())
 					{
-						PropagateTemplateFilteredTypes(OutputNode, bSetupUndoRedo);
+						if (URigVMTemplateNode* TemplateNode = Cast<URigVMTemplateNode>(EdgePin->GetNode()))
+						{
+							URigVMPin* RootPin = EdgePin->GetRootPin();
+							if (int32* TypeIndex = TypesBeforeRecomputing.Find(RootPin))
+							{
+								if (TemplateNode->FilteredSupportsType(RootPin, *TypeIndex))
+								{
+									PrepareTemplatePinForType(RootPin, {*TypeIndex}, bSetupUndoRedo);
+								}
+								else
+								{
+									LinksToRemove.Add(Link);
+									bLinkRemoved = true;
+									break;
+								}
+							}
+						}			
 					}
 				}
+				if (bLinkRemoved)
+				{
+					ActionStack->CancelUndoBracket(this);
+					break;
+				}
+				else
+				{
+					ActionStack->CloseUndoBracket(this);
+				}
 			}
-		}
 
-		// If pin is a struct member, we should resolve to that type
-		if (URigVMTemplateNode* InputNode = Cast<URigVMTemplateNode>(InputPin->GetNode()))
-		{
-			if (!InputNode->IsSingleton())
+			// Propagate filtered permutations
+			if (PrepareToLink(OutputPin, InputPin, bSetupUndoRedo))
 			{
-				if (InputPin->IsStructMember() && InputNode->PreferredPermutationPairs.IsEmpty())
+				OutputPin->Links.Add(Link);
+				InputPin->Links.Add(Link);
+
+				if (LinksRemoved.Contains(TPair<FString, FString>(Link->GetSourcePin()->GetPinPath(), Link->GetTargetPin()->GetPinPath())))
 				{
-					URigVMPin* RootPin = InputPin->GetRootPin();
-					const TRigVMTypeIndex TypeIndex = TypesBeforeRecomputing.FindChecked(RootPin);
-					if (UpdateFilteredPermutations(RootPin, {TypeIndex}, bSetupUndoRedo))
-					{
-						PropagateTemplateFilteredTypes(InputNode, bSetupUndoRedo);
-					}
+					Notify(ERigVMGraphNotifType::LinkAdded, Link);
 				}
 			}
+			else
+			{
+				LinksToRemove.Add(Link);
+			}
 		}
 
-		// Propagate filtered types due to this link
-		if (URigVMTemplateNode* OutputNode = Cast<URigVMTemplateNode>(OutputPin->GetNode()))
+		OnModified().Remove(Delegate);
+
+		bAnyLinkBroken = !LinksToRemove.IsEmpty();
 		{
-			if (!OutputNode->IsSingleton())
+			TGuardValue<bool> RecomputeGuard(bSuspendRecomputingTemplateFilters, true);
+			for (URigVMLink* Link : LinksToRemove)
 			{
-				if (OutputNode->PinNeedsFilteredTypesUpdate(OutputPin, InputPin))
+				Link->SourcePin->Links.Add(Link);
+				Link->TargetPin->Links.Add(Link);
+
+				BreakLink(Link->SourcePin, Link->TargetPin, bSetupUndoRedo);
+			}
+		}
+	
+		// Now update all template nodes pin types
+		for (URigVMNode* Node : Graph->GetNodes())
+		{
+			if (URigVMTemplateNode* TemplateNode = Cast<URigVMTemplateNode>(Node))
+			{
+				if (TemplateNode->IsSingleton())
 				{
-					if (UpdateFilteredPermutations(OutputPin, InputPin, bSetupUndoRedo))
-					{
-						PropagateTemplateFilteredTypes(OutputNode, bSetupUndoRedo);
-					}
+					continue;
+				}
+
+				if(UpdateTemplateNodePinTypes(TemplateNode, bSetupUndoRedo))
+				{
+					bAnyTypeChanged = true;
 				}
 			}
-		}				
-		if (URigVMTemplateNode* InputNode = Cast<URigVMTemplateNode>(InputPin->GetNode()))
+		}
+	}
+
+	if (URigVMGraph* ParentGraph = Graph->GetParentGraph())
+	{
+		FRigVMControllerGraphGuard GraphGuard(this, ParentGraph, bSetupUndoRedo);
+		UpdateLibraryTemplate(Graph->GetTypedOuter<URigVMLibraryNode>(), bSetupUndoRedo);
+	}
+
+	
+	return bAnyTypeChanged || bAnyLinkBroken;
+}
+
+bool URigVMController::PrepareToLink(URigVMPin* FirstToResolve, URigVMPin* SecondToResolve, bool bSetupUndoRedo)
+{
+	
+	if (URigVMTemplateNode* FirstTemplateNode = Cast<URigVMTemplateNode>(FirstToResolve->GetNode()))
+	{
+		bool bShouldPrepare = true;
+		if (FirstTemplateNode->IsA<URigVMFunctionEntryNode>() || FirstTemplateNode->IsA<URigVMFunctionReturnNode>())
+		{
+			URigVMPin* RootPin = FirstToResolve->GetRootPin();
+			if (!FirstTemplateNode->GetTemplate()->FindArgument(RootPin->GetFName()))
+			{
+				ensure(RootPin == FirstToResolve);
+				bShouldPrepare = AddArgumentForPin(FirstToResolve, SecondToResolve, bSetupUndoRedo);
+			}
+		}
+		if (!FirstTemplateNode->IsSingleton() && bShouldPrepare)
 		{
-			if (!InputNode->IsSingleton())
+			TArray<TRigVMTypeIndex> InputTypes = GetFilteredTypes(SecondToResolve);
+			if (InputTypes.Num() > 0)
 			{
-				if (InputNode->PinNeedsFilteredTypesUpdate(InputPin, OutputPin))
+				if (!PrepareTemplatePinForType(FirstToResolve, InputTypes, bSetupUndoRedo))
 				{
-					if (UpdateFilteredPermutations(InputPin, OutputPin, bSetupUndoRedo))
-					{
-						PropagateTemplateFilteredTypes(InputNode, bSetupUndoRedo);
-					}
+					return false;
 				}
 			}
-		}				
+		}
 	}
-
-	// Now update all template nodes pin types
-	bool bAnyTypeChanged = false;
-	for (URigVMNode* Node : Graph->GetNodes())
+	if (URigVMTemplateNode* SecondTemplateNode = Cast<URigVMTemplateNode>(SecondToResolve->GetNode()))
 	{
-		if (URigVMTemplateNode* TemplateNode = Cast<URigVMTemplateNode>(Node))
+		bool bShouldPrepare = true;
+		if (SecondTemplateNode->IsA<URigVMFunctionEntryNode>() || SecondTemplateNode->IsA<URigVMFunctionReturnNode>())
 		{
-			if (TemplateNode->IsSingleton())
+			URigVMPin* RootPin = SecondToResolve->GetRootPin();
+			if (!SecondTemplateNode->GetTemplate()->FindArgument(RootPin->GetFName()))
 			{
-				continue;
+				ensure(RootPin == SecondToResolve);
+				bShouldPrepare = AddArgumentForPin(SecondToResolve, FirstToResolve, bSetupUndoRedo);
 			}
-
-			if(UpdateTemplateNodePinTypes(TemplateNode, bSetupUndoRedo))
+		}
+		if (!SecondTemplateNode->IsSingleton() && bShouldPrepare)
+		{
+			TArray<TRigVMTypeIndex> OutTypes = GetFilteredTypes(FirstToResolve);
+			if (OutTypes.Num() > 0)
 			{
-				bAnyTypeChanged = true;
+				if (!PrepareTemplatePinForType(SecondToResolve, OutTypes, bSetupUndoRedo))
+				{
+					return false;
+				}
 			}
 		}
 	}
 
-	return bAnyTypeChanged;
+	return true;
 }
 
 void URigVMController::InitializeFilteredPermutationsFromTemplateTypes()
@@ -16097,9 +17388,52 @@ void URigVMController::InitializeFilteredPermutationsFromTemplateTypes()
 
 	for (URigVMNode* Node : Graph->GetNodes())
 	{
+		if (URigVMLibraryNode* LibraryNode = Cast<URigVMLibraryNode>(Node))
+		{
+			bool bTemplateChanged = false;
+			FRigVMTemplate NewTemplate;
+			NewTemplate.Permutations.Add(INDEX_NONE);
+			
+			TArray<URigVMTemplateNode*> InterfaceNodes = {LibraryNode->GetEntryNode(), LibraryNode->GetReturnNode()};
+			for (URigVMTemplateNode* InterfaceNode : InterfaceNodes)
+			{
+				if (InterfaceNode)
+				{
+					InterfaceNode->PreferredPermutationPairs.Reset();
+					for (URigVMPin* Pin : InterfaceNode->GetPins())
+					{
+						if (!Pin->IsWildCard())
+						{
+							InterfaceNode->PreferredPermutationPairs.Add(FRigVMTemplatePreferredType(Pin->GetFName(), Pin->GetTypeIndex()));
+
+							if (!NewTemplate.FindArgument(Pin->GetFName()))
+							{
+								FRigVMTemplateArgument NewArgument;
+								NewArgument.Name = Pin->GetFName();
+								NewArgument.Index = NewTemplate.NumArguments();
+								NewArgument.Direction = LibraryNode->FindPin(Pin->GetName())->GetDirection();
+								NewArgument.TypeIndices.Add(Pin->GetTypeIndex());
+								NewArgument.TypeToPermutations.Add(Pin->GetTypeIndex(), {0});
+								NewTemplate.Arguments.Add(NewArgument);
+							}
+							
+							bTemplateChanged = true;
+						}
+					}
+					InterfaceNode->FilteredPermutations = {0};
+				}
+			}
+
+			if (bTemplateChanged)
+			{
+				NewTemplate.ComputeNotationFromArguments(LibraryNode->GetName());				
+				LibraryNode->Template = NewTemplate;
+			}
+		}
+		
 		if (URigVMTemplateNode* TemplateNode = Cast<URigVMTemplateNode>(Node))
 		{
-			TemplateNode->InitializeFilteredPermutationsFromTypes();
+			TemplateNode->InitializeFilteredPermutationsFromTypes(false);
 		}
 	}
 }
@@ -16118,14 +17452,124 @@ void URigVMController::InitializeAllTemplateFiltersInGraph(bool bSetupUndoRedo,
 
 	URigVMGraph* Graph = GetGraph();
 	check(Graph);
-	
-	
+
+	// Remove all arguments (that do not have a preferred type) from the template
+	if (URigVMLibraryNode* LibraryNode = Graph->GetTypedOuter<URigVMLibraryNode>())
+	{
+		if (const FRigVMTemplate* Template = LibraryNode->GetTemplate())
+		{
+			FRigVMTemplate NewTemplate;
+			NewTemplate.Permutations.Add(INDEX_NONE);
+
+			TArray<URigVMTemplateNode*> InterfaceNodes = {LibraryNode->GetEntryNode(), LibraryNode->GetReturnNode()};
+			for (const FRigVMTemplateArgument& Argument : Template->Arguments)
+			{
+				// Find out if this argument has a preferred type
+				bool bHasPreferredType = false;
+				TRigVMTypeIndex PreferredTypeIndex;
+				for (URigVMTemplateNode* InterfaceNode : InterfaceNodes)
+				{
+					if (InterfaceNode)
+					{
+						for (int32 i=0; i<InterfaceNode->PreferredPermutationPairs.Num(); ++i)
+						{
+							FRigVMTemplatePreferredType& PreferredType = InterfaceNode->PreferredPermutationPairs[i];
+							if (PreferredType.Argument == Argument.Name)
+							{
+								PreferredTypeIndex = PreferredType.TypeIndex;
+								bHasPreferredType = true;
+								break;
+							}
+						}
+						if (bHasPreferredType)
+						{
+							break;
+						}
+					}
+				}
+
+				if (bHasPreferredType)
+				{
+					FRigVMTemplateArgument NewArgument;
+					NewArgument.Name = Argument.Name;
+					NewArgument.Index = NewTemplate.NumArguments();
+					NewArgument.Direction = Argument.Direction;
+					NewArgument.TypeIndices.Add(PreferredTypeIndex);
+					NewArgument.TypeToPermutations.Add(PreferredTypeIndex, {0});
+					NewTemplate.Arguments.Add(NewArgument);
+				}
+			}
+
+			NewTemplate.ComputeNotationFromArguments(LibraryNode->GetName());
+			
+			if (bSetupUndoRedo)
+			{
+				FRigVMControllerGraphGuard GraphGuard(this, LibraryNode->GetTypedOuter<URigVMGraph>(), bSetupUndoRedo);
+				ActionStack->AddAction(FRigVMSetLibraryTemplateAction(LibraryNode, NewTemplate));
+			}
+			LibraryNode->Template = NewTemplate;
+			Notify(ERigVMGraphNotifType::NodeDescriptionChanged, LibraryNode);
+
+			// Update Entry and Return filtered permutations to all permutations
+			for (URigVMTemplateNode* InterfaceNode : InterfaceNodes)
+			{
+				if (InterfaceNode)
+				{
+					TArray<int32> OldPermutations = InterfaceNode->FilteredPermutations;
+					InterfaceNode->FilteredPermutations.Init(INDEX_NONE, NewTemplate.NumPermutations());
+					
+					if (bSetupUndoRedo)
+					{
+						ActionStack->AddAction(FRigVMSetTemplateFilteredPermutationsAction(InterfaceNode, OldPermutations));
+					}
+
+					if (bChangePinTypes)
+					{
+						UpdateTemplateNodePinTypes(InterfaceNode, bSetupUndoRedo);
+					}
+
+					Notify(ERigVMGraphNotifType::NodeDescriptionChanged, InterfaceNode);
+				}
+			}
+
+			// Update outer graphs
+			if (!bIsTransacting && !bSuspendRecomputingOuterTemplateFilters)
+			{
+				FRigVMControllerGraphGuard GraphGuard(this, LibraryNode->GetGraph(), bSetupUndoRedo);
+				RecomputeAllTemplateFilteredPermutations(bSetupUndoRedo);
+			}
+		}
+	}
+
 	for (URigVMNode* Node : Graph->GetNodes())
 	{
 		if (URigVMTemplateNode* TemplateNode = Cast<URigVMTemplateNode>(Node))
 		{
 			if (!TemplateNode->IsSingleton())
 			{
+				// If the preferred permutation does not fit to the available types in the template, remove them
+				TArray<FRigVMTemplatePreferredType> NewPreferredTypes;
+				for (int32 i=0; i<TemplateNode->PreferredPermutationPairs.Num(); ++i)
+				{
+					const FRigVMTemplatePreferredType& PreferredType = TemplateNode->PreferredPermutationPairs[i];
+					
+					if (TemplateNode->GetTemplate()->FindArgument(PreferredType.Argument) && TemplateNode->GetTemplate()->ArgumentSupportsTypeIndex(PreferredType.Argument, PreferredType.TypeIndex))
+					{
+						NewPreferredTypes.Add(PreferredType);
+					}
+				}
+
+				if (TemplateNode->PreferredPermutationPairs.Num() != NewPreferredTypes.Num())
+				{
+					if (bSetupUndoRedo)
+					{
+						FRigVMSetPreferredTemplatePermutationsAction PreferredTypesAction(TemplateNode, NewPreferredTypes);
+						ActionStack->AddAction(PreferredTypesAction);
+					}
+
+					TemplateNode->PreferredPermutationPairs = NewPreferredTypes;
+				}
+				
 				TArray<int32> OldPermutations = TemplateNode->FilteredPermutations;
 				TemplateNode->InitializeFilteredPermutations();
 				if (bChangePinTypes)
@@ -16134,7 +17578,7 @@ void URigVMController::InitializeAllTemplateFiltersInGraph(bool bSetupUndoRedo,
 				}
 				if (bSetupUndoRedo)
 				{
-					FRigVMSetTemplateFilteredPermutationsAction FilteringAction(TemplateNode, nullptr, OldPermutations);
+					FRigVMSetTemplateFilteredPermutationsAction FilteringAction(TemplateNode, OldPermutations);
 					ActionStack->AddAction(FilteringAction);
 				}
 			}
@@ -16142,7 +17586,7 @@ void URigVMController::InitializeAllTemplateFiltersInGraph(bool bSetupUndoRedo,
 	}
 }
 
-bool URigVMController::ChangePinType(const FString& InPinPath, const FString& InCPPType, const FName& InCPPTypeObjectPath, bool bSetupUndoRedo, bool bSetupOrphanPins, bool bBreakLinks, bool bRemoveSubPins)
+bool URigVMController::ChangePinType(const FString& InPinPath, const FString& InCPPType, const FName& InCPPTypeObjectPath, bool bSetupUndoRedo, bool bSetupOrphanPins, bool bBreakLinks, bool bRemoveSubPins, bool bInitializeDefaultValue)
 {
 	if (!IsValidGraph())
 	{
@@ -16159,13 +17603,13 @@ bool URigVMController::ChangePinType(const FString& InPinPath, const FString& In
 
 	if (URigVMPin* Pin = Graph->FindPin(InPinPath))
 	{
-		return ChangePinType(Pin, InCPPType, InCPPTypeObjectPath, bSetupUndoRedo, bSetupOrphanPins, bBreakLinks, bRemoveSubPins);
+		return ChangePinType(Pin, InCPPType, InCPPTypeObjectPath, bSetupUndoRedo, bSetupOrphanPins, bBreakLinks, bRemoveSubPins, bInitializeDefaultValue);
 	}
 
 	return false;
 }
 
-bool URigVMController::ChangePinType(URigVMPin* InPin, const FString& InCPPType, const FName& InCPPTypeObjectPath, bool bSetupUndoRedo, bool bSetupOrphanPins, bool bBreakLinks, bool bRemoveSubPins)
+bool URigVMController::ChangePinType(URigVMPin* InPin, const FString& InCPPType, const FName& InCPPTypeObjectPath, bool bSetupUndoRedo, bool bSetupOrphanPins, bool bBreakLinks, bool bRemoveSubPins, bool bInitializeDefaultValue)
 {
 	if (!bIsTransacting && !IsGraphEditable())
 	{
@@ -16194,10 +17638,10 @@ bool URigVMController::ChangePinType(URigVMPin* InPin, const FString& InCPPType,
 		}
 	}
 
-	return ChangePinType(InPin, InCPPType, CPPTypeObject, bSetupUndoRedo, bSetupOrphanPins, bBreakLinks, bRemoveSubPins);
+	return ChangePinType(InPin, InCPPType, CPPTypeObject, bSetupUndoRedo, bSetupOrphanPins, bBreakLinks, bRemoveSubPins, bInitializeDefaultValue);
 }
 
-bool URigVMController::ChangePinType(URigVMPin* InPin, const FString& InCPPType, UObject* InCPPTypeObject, bool bSetupUndoRedo, bool bSetupOrphanPins, bool bBreakLinks, bool bRemoveSubPins)
+bool URigVMController::ChangePinType(URigVMPin* InPin, const FString& InCPPType, UObject* InCPPTypeObject, bool bSetupUndoRedo, bool bSetupOrphanPins, bool bBreakLinks, bool bRemoveSubPins, bool bInitializeDefaultValue)
 {
 	if (!bIsTransacting && !IsGraphEditable())
 	{
@@ -16216,10 +17660,10 @@ bool URigVMController::ChangePinType(URigVMPin* InPin, const FString& InCPPType,
 
 	const FRigVMTemplateArgumentType Type(*InCPPType, InCPPTypeObject);
 	const TRigVMTypeIndex TypeIndex = FRigVMRegistry::Get().GetTypeIndex(Type);
-	return ChangePinType(InPin, TypeIndex, bSetupUndoRedo, bSetupOrphanPins, bBreakLinks, bRemoveSubPins);
+	return ChangePinType(InPin, TypeIndex, bSetupUndoRedo, bSetupOrphanPins, bBreakLinks, bRemoveSubPins, bInitializeDefaultValue);
 }
 
-bool URigVMController::ChangePinType(URigVMPin* InPin, TRigVMTypeIndex InTypeIndex, bool bSetupUndoRedo, bool bSetupOrphanPins, bool bBreakLinks, bool bRemoveSubPins)
+bool URigVMController::ChangePinType(URigVMPin* InPin, TRigVMTypeIndex InTypeIndex, bool bSetupUndoRedo, bool bSetupOrphanPins, bool bBreakLinks, bool bRemoveSubPins, bool bInitializeDefaultValue)
 {
 	if (!bIsTransacting && !IsGraphEditable())
 	{
@@ -16231,15 +17675,22 @@ bool URigVMController::ChangePinType(URigVMPin* InPin, TRigVMTypeIndex InTypeInd
 		return false;
 	}
 
+	check(InPin->GetGraph() == GetGraph());
+
 	// only allow valid pin cpp types on template nodes
 	TRigVMTypeIndex TypeIndex = InTypeIndex;
 	if(URigVMTemplateNode* TemplateNode = Cast<URigVMTemplateNode>(InPin->GetNode()))
 	{
-		if(!TemplateNode->SupportsType(InPin, InTypeIndex, &TypeIndex))
+		if (!bIsTransacting)
 		{
-			const FRigVMTemplateArgumentType Type = FRigVMRegistry::Get().GetType(InTypeIndex);
-			ReportErrorf(TEXT("ChangePinType: %s doesn't support type '%s'."), *InPin->GetPinPath(), *Type.CPPType.ToString());
-			return false;
+			if (TemplateNode->IsA<URigVMUnitNode>() || TemplateNode->IsA<URigVMDispatchNode>())
+			{
+				if(!TemplateNode->SupportsType(InPin, InTypeIndex))
+				{
+					ReportErrorf(TEXT("ChangePinType: %s doesn't support type '%s'."), *InPin->GetPinPath(), *FRigVMRegistry::Get().GetType(InTypeIndex).CPPType.ToString());
+					return false;
+				}
+			}
 		}
 
 		// change the pin's type to be of an array as well
@@ -16328,14 +17779,17 @@ bool URigVMController::ChangePinType(URigVMPin* InPin, TRigVMTypeIndex InTypeInd
 	InPin->CPPTypeObjectPath = Type.GetCPPTypeObjectPath();
 	InPin->CPPTypeObject = Type.CPPTypeObject;
 	InPin->bIsDynamicArray = FRigVMRegistry::Get().IsArrayType(TypeIndex);
-	// we might want to use GetPinInitialDefaultValue here for a better default value
-	InPin->DefaultValue = FString();
 
-	if(InPin->IsRootPin() && !InPin->IsWildCard())
+	if (bInitializeDefaultValue)
 	{
-		if(URigVMTemplateNode* TemplateNode = Cast<URigVMTemplateNode>(InPin->GetNode()))
+		InPin->DefaultValue = FString();
+
+		if(InPin->IsRootPin() && !InPin->IsWildCard())
 		{
-			InPin->DefaultValue = TemplateNode->GetInitialDefaultValueForPin(InPin->GetFName());
+			if(URigVMTemplateNode* TemplateNode = Cast<URigVMTemplateNode>(InPin->GetNode()))
+			{
+				InPin->DefaultValue = TemplateNode->GetInitialDefaultValueForPin(InPin->GetFName());
+			}
 		}
 	}
 
@@ -16361,7 +17815,7 @@ bool URigVMController::ChangePinType(URigVMPin* InPin, TRigVMTypeIndex InTypeInd
 			{
 				continue;
 			}
-			ChangePinType(SubPin, BaseTypeIndex, bSetupUndoRedo, bSetupOrphanPins, bBreakLinks, bRemoveSubPins);
+			ChangePinType(SubPin, BaseTypeIndex, bSetupUndoRedo, bSetupOrphanPins, bBreakLinks, bRemoveSubPins, bInitializeDefaultValue);
 		}
 	}
 
@@ -16507,9 +17961,12 @@ bool URigVMController::ChangePinType(URigVMPin* InPin, TRigVMTypeIndex InTypeInd
 
 			if (URigVMTemplateNode* TemplateNode = Cast<URigVMTemplateNode>(InPin->GetNode()))
 			{
-				// Figure out the permutation from the pin types. During undo, the filtered permutations are not
-				// reliable as to which permutation we are resolving to.
-				FullyResolveTemplateNode(TemplateNode, INDEX_NONE, bSetupUndoRedo);
+				if (TemplateNode->IsA<URigVMUnitNode>() || TemplateNode->IsA<URigVMDispatchNode>())
+				{
+					// Figure out the permutation from the pin types. During undo, the filtered permutations are not
+					// reliable as to which permutation we are resolving to.
+					FullyResolveTemplateNode(TemplateNode, INDEX_NONE, bSetupUndoRedo);
+				}
 			}
 		}
 	}
@@ -16518,7 +17975,7 @@ bool URigVMController::ChangePinType(URigVMPin* InPin, TRigVMTypeIndex InTypeInd
 	// graph views know to invalidate the node title text widget
 	Notify(ERigVMGraphNotifType::NodeDescriptionChanged, InPin->GetNode());
 
-	// in cases were we are just changing the type we have to let the
+	// in cases where we are just changing the type we have to let the
 	// clients know that the links are still there
 	if(!bSetupOrphanPins && !bBreakLinks && !bRemoveSubPins)
 	{
diff --git a/Engine/Source/Developer/RigVMDeveloper/Private/RigVMModel/RigVMControllerActions.cpp b/Engine/Source/Developer/RigVMDeveloper/Private/RigVMModel/RigVMControllerActions.cpp
--- a/Engine/Source/Developer/RigVMDeveloper/Private/RigVMModel/RigVMControllerActions.cpp
+++ b/Engine/Source/Developer/RigVMDeveloper/Private/RigVMModel/RigVMControllerActions.cpp
@@ -789,7 +789,7 @@ FRigVMRemoveNodeAction::FRigVMRemoveNodeAction(URigVMNode* InNode, URigVMControl
 		{
 			InverseAction.AddAction(FRigVMAddTemplateNodeAction(TemplateNode), InController);
 			InverseAction.AddAction(FRigVMSetPreferredTemplatePermutationsAction(TemplateNode, TemplateNode->PreferredPermutationPairs), InController);
-			InverseAction.AddAction(FRigVMSetTemplateFilteredPermutationsAction(TemplateNode, nullptr, {}), InController);
+			InverseAction.AddAction(FRigVMSetTemplateFilteredPermutationsAction(TemplateNode, {}), InController);
 
 			for (URigVMPin* Pin : TemplateNode->GetPins())
 			{
@@ -1366,9 +1366,8 @@ bool FRigVMSetPinDefaultValueAction::Redo(URigVMController* InController)
 	return FRigVMBaseAction::Redo(InController);
 }
 
-FRigVMSetTemplateFilteredPermutationsAction::FRigVMSetTemplateFilteredPermutationsAction(URigVMTemplateNode* InNode, URigVMLink* InLink, const TArray<int32>& InOldFilteredPermutations)
+FRigVMSetTemplateFilteredPermutationsAction::FRigVMSetTemplateFilteredPermutationsAction(URigVMTemplateNode* InNode, const TArray<int32>& InOldFilteredPermutations)
 : NodePath(InNode->GetNodePath())
-, LinkPath(InLink ? InLink->GetPinPathRepresentation() : FString())
 , OldFilteredPermutations(InOldFilteredPermutations)
 , NewFilteredPermutations(InNode->GetFilteredPermutationsIndices())
 {
@@ -1469,6 +1468,106 @@ bool FRigVMSetPreferredTemplatePermutationsAction::Redo(URigVMController* InCont
 	return FRigVMBaseAction::Redo(InController);
 }
 
+FRigVMSetLibraryTemplateAction::FRigVMSetLibraryTemplateAction(URigVMLibraryNode* InNode, FRigVMTemplate& InNewTemplate)
+: NodePath(InNode->GetNodePath())
+{
+	FMemoryWriter ArchiveWriterOld(OldTemplateBytes);
+	InNode->Template.Serialize(ArchiveWriterOld);
+	FMemoryWriter ArchiveWriterNew(NewTemplateBytes);
+	InNewTemplate.Serialize(ArchiveWriterNew);	
+}
+
+bool FRigVMSetLibraryTemplateAction::Undo(URigVMController* InController)
+{
+	if(!FRigVMBaseAction::Undo(InController))
+	{
+		return false;
+	}
+
+	if (URigVMLibraryNode* LibraryNode = Cast<URigVMLibraryNode>(InController->GetGraph()->FindNode(NodePath)))
+	{
+		FMemoryReader ArchiveReader(OldTemplateBytes);
+		FRigVMTemplate NewTemplate;
+		NewTemplate.Serialize(ArchiveReader);
+		LibraryNode->Template = NewTemplate;
+
+		// Update type to permutations map for each argument
+		for (FRigVMTemplateArgument& Argument : LibraryNode->Template.Arguments)
+		{
+			Argument.TypeToPermutations.Reset();
+			for (int32 PermutationIndex=0; PermutationIndex<Argument.TypeIndices.Num(); ++PermutationIndex)
+			{
+				const int32& TypeIndex = Argument.TypeIndices[PermutationIndex];
+				if (TArray<int32>* Permutations = Argument.TypeToPermutations.Find(TypeIndex))
+				{
+					Permutations->Add(PermutationIndex);
+				}
+				else
+				{
+					Argument.TypeToPermutations.Add(TypeIndex, {PermutationIndex});
+				}
+			}
+		}
+		
+		InController->Notify(ERigVMGraphNotifType::NodeDescriptionChanged, LibraryNode);
+		if (URigVMTemplateNode* EntryNode = Cast<URigVMTemplateNode>(LibraryNode->GetEntryNode()))
+		{
+			InController->Notify(ERigVMGraphNotifType::NodeDescriptionChanged, EntryNode);
+		}
+		if (URigVMTemplateNode* ReturnNode = Cast<URigVMTemplateNode>(LibraryNode->GetReturnNode()))
+		{
+			InController->Notify(ERigVMGraphNotifType::NodeDescriptionChanged, ReturnNode);
+		}
+		
+		return true;
+	}
+	return false;
+}
+
+bool FRigVMSetLibraryTemplateAction::Redo(URigVMController* InController)
+{
+	URigVMLibraryNode* LibraryNode = Cast<URigVMLibraryNode>(InController->GetGraph()->FindNode(NodePath));
+	if (!LibraryNode)
+	{
+		return false;
+	}
+
+	FMemoryReader ArchiveReader(NewTemplateBytes);
+	FRigVMTemplate NewTemplate;
+	NewTemplate.Serialize(ArchiveReader);
+	LibraryNode->Template = NewTemplate;
+	
+	// Update type to permutations map for each argument
+	for (FRigVMTemplateArgument& Argument : LibraryNode->Template.Arguments)
+	{
+		Argument.TypeToPermutations.Reset();
+		for (int32 PermutationIndex=0; PermutationIndex<Argument.TypeIndices.Num(); ++PermutationIndex)
+		{
+			const int32& TypeIndex = Argument.TypeIndices[PermutationIndex];
+			if (TArray<int32>* Permutations = Argument.TypeToPermutations.Find(TypeIndex))
+			{
+				Permutations->Add(PermutationIndex);
+			}
+			else
+			{
+				Argument.TypeToPermutations.Add(TypeIndex, {PermutationIndex});
+			}
+		}
+	}
+
+	InController->Notify(ERigVMGraphNotifType::NodeDescriptionChanged, LibraryNode);
+	if (URigVMTemplateNode* EntryNode = Cast<URigVMTemplateNode>(LibraryNode->GetEntryNode()))
+	{
+		InController->Notify(ERigVMGraphNotifType::NodeDescriptionChanged, EntryNode);
+	}
+	if (URigVMTemplateNode* ReturnNode = Cast<URigVMTemplateNode>(LibraryNode->GetReturnNode()))
+	{
+		InController->Notify(ERigVMGraphNotifType::NodeDescriptionChanged, ReturnNode);
+	}
+	
+	return FRigVMBaseAction::Redo(InController);
+}
+
 FRigVMInsertArrayPinAction::FRigVMInsertArrayPinAction(URigVMPin* InArrayPin, int32 InIndex, const FString& InNewDefaultValue)
 : ArrayPinPath(InArrayPin->GetPinPath())
 , Index(InIndex)
@@ -1547,6 +1646,7 @@ FRigVMBreakLinkAction::FRigVMBreakLinkAction(URigVMPin* InOutputPin, URigVMPin*
 : OutputPinPath(InOutputPin->GetPinPath())
 , InputPinPath(InInputPin->GetPinPath())
 {
+	GraphPath = TSoftObjectPtr<URigVMGraph>(Cast<URigVMGraph>(InOutputPin->GetGraph())).GetUniqueID();
 }
 
 bool FRigVMBreakLinkAction::Undo(URigVMController* InController)
diff --git a/Engine/Source/Developer/RigVMDeveloper/Private/RigVMModel/RigVMPin.cpp b/Engine/Source/Developer/RigVMDeveloper/Private/RigVMModel/RigVMPin.cpp
--- a/Engine/Source/Developer/RigVMDeveloper/Private/RigVMModel/RigVMPin.cpp
+++ b/Engine/Source/Developer/RigVMDeveloper/Private/RigVMModel/RigVMPin.cpp
@@ -1633,7 +1633,7 @@ URigVMGraph* URigVMPin::GetGraph() const
 	return nullptr;
 }
 
-bool URigVMPin::CanLink(URigVMPin* InSourcePin, URigVMPin* InTargetPin, FString* OutFailureReason, const FRigVMByteCode* InByteCode, ERigVMPinDirection InUserLinkDirection)
+bool URigVMPin::CanLink(URigVMPin* InSourcePin, URigVMPin* InTargetPin, FString* OutFailureReason, const FRigVMByteCode* InByteCode, ERigVMPinDirection InUserLinkDirection, bool bInAllowNonArgumentPins)
 {
 	if (InSourcePin == nullptr || InTargetPin == nullptr)
 	{
@@ -1763,11 +1763,11 @@ bool URigVMPin::CanLink(URigVMPin* InSourcePin, URigVMPin* InTargetPin, FString*
 				
 				if(IsPinValidForTypeChange(InSourcePin, false, InUserLinkDirection))
 				{
-					return TemplateNodeSupportsType(InSourcePin, InTargetPin->GetTypeIndex(), OutFailureReason);
+					bCPPTypesDiffer = !TemplateNodeSupportsType(InSourcePin, InTargetPin->GetTypeIndex(), OutFailureReason);
 				}
 				else if(IsPinValidForTypeChange(InTargetPin, true, InUserLinkDirection))
 				{
-					return TemplateNodeSupportsType(InTargetPin, InSourcePin->GetTypeIndex(), OutFailureReason);
+					bCPPTypesDiffer = !TemplateNodeSupportsType(InTargetPin, InSourcePin->GetTypeIndex(), OutFailureReason);
 				}
 			}
 		}
@@ -1800,6 +1800,46 @@ bool URigVMPin::CanLink(URigVMPin* InSourcePin, URigVMPin* InTargetPin, FString*
 		return false;
 	}
 
+	if (!bInAllowNonArgumentPins)
+	{
+		if (URigVMTemplateNode* SourceTemplateNode = Cast<URigVMTemplateNode>(SourceNode))
+		{
+			if (!SourceNode->IsA<URigVMFunctionEntryNode>() && !SourceNode->IsA<URigVMFunctionReturnNode>())
+			{
+				if (const FRigVMTemplate* Template = SourceTemplateNode->GetTemplate())
+				{
+					URigVMPin* RootPin = InSourcePin->GetRootPin();
+					if (!Template->FindArgument(RootPin->GetFName()))
+					{
+						if (OutFailureReason)
+						{
+							*OutFailureReason = FString::Printf(TEXT("Library pin %s supported types need to be reduced."), *RootPin->GetPinPath(true));
+						}
+						return false;
+					}
+				}
+			}
+		}
+		if (URigVMTemplateNode* TargetTemplateNode = Cast<URigVMTemplateNode>(TargetNode))
+		{
+			if (!TargetNode->IsA<URigVMFunctionEntryNode>() && !TargetNode->IsA<URigVMFunctionReturnNode>())
+			{
+				if (const FRigVMTemplate* Template = TargetTemplateNode->GetTemplate())
+				{
+					URigVMPin* RootPin = InTargetPin->GetRootPin();
+					if (!Template->FindArgument(RootPin->GetFName()))
+					{
+						if (OutFailureReason)
+						{
+							*OutFailureReason = FString::Printf(TEXT("Library pin %s supported types need to be reduced."), *RootPin->GetPinPath(true));
+						}
+						return false;
+					}
+				}
+			}
+		}
+	}
+
 	// only allow to link to specified input / output pins on an injected node
 	if (URigVMInjectionInfo* SourceInjectionInfo = SourceNode->GetInjectionInfo())
 	{
diff --git a/Engine/Source/Developer/RigVMDeveloper/Public/RigVMModel/Nodes/RigVMFunctionEntryNode.h b/Engine/Source/Developer/RigVMDeveloper/Public/RigVMModel/Nodes/RigVMFunctionEntryNode.h
--- a/Engine/Source/Developer/RigVMDeveloper/Public/RigVMModel/Nodes/RigVMFunctionEntryNode.h
+++ b/Engine/Source/Developer/RigVMDeveloper/Public/RigVMModel/Nodes/RigVMFunctionEntryNode.h
@@ -2,7 +2,7 @@
 
 #pragma once
 
-#include "RigVMModel/RigVMNode.h"
+#include "RigVMModel/Nodes/RigVMTemplateNode.h"
 #include "RigVMFunctionEntryNode.generated.h"
 
 /**
@@ -10,16 +10,26 @@
  * input pins of the library node for links within.
  */
 UCLASS(BlueprintType)
-class RIGVMDEVELOPER_API URigVMFunctionEntryNode : public URigVMNode
+class RIGVMDEVELOPER_API URigVMFunctionEntryNode : public URigVMTemplateNode
 {
 	GENERATED_BODY()
 
 public:
 
+	// Override template node functions
+	virtual UScriptStruct* GetScriptStruct() const override { return nullptr; }
+	virtual const FRigVMTemplate* GetTemplate() const override;
+	virtual FName GetNotation() const override;
+
 	// Override node functions
 	virtual FLinearColor GetNodeColor() const override;
 	virtual bool IsDefinedAsVarying() const override;
 
+	// URigVMNode interface
+	virtual FString GetNodeTitle() const override;
+	virtual  FText GetToolTipText() const override;
+	virtual FText GetToolTipTextForPin(const URigVMPin* InPin) const override;
+
 private:
 
 	friend class URigVMController;
diff --git a/Engine/Source/Developer/RigVMDeveloper/Public/RigVMModel/Nodes/RigVMFunctionReferenceNode.h b/Engine/Source/Developer/RigVMDeveloper/Public/RigVMModel/Nodes/RigVMFunctionReferenceNode.h
--- a/Engine/Source/Developer/RigVMDeveloper/Public/RigVMModel/Nodes/RigVMFunctionReferenceNode.h
+++ b/Engine/Source/Developer/RigVMDeveloper/Public/RigVMModel/Nodes/RigVMFunctionReferenceNode.h
@@ -31,6 +31,7 @@ public:
 	virtual URigVMFunctionLibrary* GetLibrary() const override;
 	virtual URigVMGraph* GetContainedGraph() const override;
 	virtual TArray<FRigVMExternalVariable> GetExternalVariables() const override;
+	virtual const FRigVMTemplate* GetTemplate() const override;
 	// end URigVMLibraryNode interface
 
 	URigVMLibraryNode* GetReferencedNode() const;
diff --git a/Engine/Source/Developer/RigVMDeveloper/Public/RigVMModel/Nodes/RigVMFunctionReturnNode.h b/Engine/Source/Developer/RigVMDeveloper/Public/RigVMModel/Nodes/RigVMFunctionReturnNode.h
--- a/Engine/Source/Developer/RigVMDeveloper/Public/RigVMModel/Nodes/RigVMFunctionReturnNode.h
+++ b/Engine/Source/Developer/RigVMDeveloper/Public/RigVMModel/Nodes/RigVMFunctionReturnNode.h
@@ -2,7 +2,7 @@
 
 #pragma once
 
-#include "RigVMModel/RigVMNode.h"
+#include "RigVMModel/Nodes/RigVMTemplateNode.h"
 #include "RigVMFunctionReturnNode.generated.h"
 
 /**
@@ -10,15 +10,25 @@
  * output pins of the library node for links within.
  */
 UCLASS(BlueprintType)
-class RIGVMDEVELOPER_API URigVMFunctionReturnNode : public URigVMNode
+class RIGVMDEVELOPER_API URigVMFunctionReturnNode : public URigVMTemplateNode
 {
 	GENERATED_BODY()
 
 public:
 
+	// Override template node functions
+	virtual UScriptStruct* GetScriptStruct() const override { return nullptr; }
+	virtual const FRigVMTemplate* GetTemplate() const override;
+	virtual FName GetNotation() const override;
+
 	// Override node functions
 	virtual FLinearColor GetNodeColor() const override;
 	virtual bool IsDefinedAsVarying() const override;
+	
+	// URigVMNode interface
+	virtual FString GetNodeTitle() const override;
+	virtual  FText GetToolTipText() const override;
+	virtual FText GetToolTipTextForPin(const URigVMPin* InPin) const override;	
 
 private:
 
diff --git a/Engine/Source/Developer/RigVMDeveloper/Public/RigVMModel/Nodes/RigVMLibraryNode.h b/Engine/Source/Developer/RigVMDeveloper/Public/RigVMModel/Nodes/RigVMLibraryNode.h
--- a/Engine/Source/Developer/RigVMDeveloper/Public/RigVMModel/Nodes/RigVMLibraryNode.h
+++ b/Engine/Source/Developer/RigVMDeveloper/Public/RigVMModel/Nodes/RigVMLibraryNode.h
@@ -2,7 +2,7 @@
 
 #pragma once
 
-#include "RigVMModel/RigVMNode.h"
+#include "RigVMModel/Nodes/RigVMTemplateNode.h"
 #include "RigVMLibraryNode.generated.h"
 
 class URigVMGraph;
@@ -17,17 +17,26 @@ class URigVMFunctionLibrary;
  * referenced function (RigVMFunctionNode).
  */
 UCLASS(BlueprintType)
-class RIGVMDEVELOPER_API URigVMLibraryNode : public URigVMNode
+class RIGVMDEVELOPER_API URigVMLibraryNode : public URigVMTemplateNode
 {
 	GENERATED_BODY()
 
 public:
+	URigVMLibraryNode();
 
 	// Override node functions
 	virtual bool IsDefinedAsConstant() const override;
 	virtual bool IsDefinedAsVarying() const override;
-	//virtual int32 GetInstructionVisitedCount(URigVM* InVM, const FRigVMASTProxy& InProxy = FRigVMASTProxy(), bool bConsolidatePerNode = false) const override; 
+	//virtual int32 GetInstructionVisitedCount(URigVM* InVM, const FRigVMASTProxy& InProxy = FRigVMASTProxy(), bool bConsolidatePerNode = false) const override;
+	
+	// Override template node functions
+	virtual UScriptStruct* GetScriptStruct() const override { return nullptr; }
+	virtual const FRigVMTemplate* GetTemplate() const override;
+	virtual FName GetNotation() const override;
 
+	// URigVMNode interface
+	virtual  FText GetToolTipText() const override;
+	
 	// Library node interface
 	virtual FString GetNodeCategory() const { return FString(); }
 	virtual FString GetNodeKeywords() const { return FString(); }
@@ -48,6 +57,8 @@ public:
 
 protected:
 
+	FRigVMTemplate Template;
+	
 	virtual TArray<int32> GetInstructionsForVMImpl(URigVM* InVM, const FRigVMASTProxy& InProxy = FRigVMASTProxy()) const override; 
 	const static TArray<URigVMNode*> EmptyNodes;
 	const static TArray<URigVMLink*> EmptyLinks;
@@ -55,5 +66,6 @@ protected:
 private:
 
 	friend class URigVMController;
+	friend struct FRigVMSetLibraryTemplateAction;
 };
 
diff --git a/Engine/Source/Developer/RigVMDeveloper/Public/RigVMModel/Nodes/RigVMTemplateNode.h b/Engine/Source/Developer/RigVMDeveloper/Public/RigVMModel/Nodes/RigVMTemplateNode.h
--- a/Engine/Source/Developer/RigVMDeveloper/Public/RigVMModel/Nodes/RigVMTemplateNode.h
+++ b/Engine/Source/Developer/RigVMDeveloper/Public/RigVMModel/Nodes/RigVMTemplateNode.h
@@ -37,6 +37,11 @@ public:
 		return TypeIndex;
 #endif
 	}
+
+	bool operator==(const FRigVMTemplatePreferredType& Other) const
+	{
+		return Argument == Other.Argument && TypeIndex == Other.TypeIndex;
+	}
 	
 	void UpdateStringFromIndex();
 	void UpdateIndexFromString();
@@ -53,6 +58,7 @@ protected:
 	FString TypeString;
 
 	friend class URigVMTemplateNode;
+	friend class URigVMController;
 };
 
 /**
@@ -79,6 +85,10 @@ public:
 	virtual FName GetMethodName() const;
 	virtual  FText GetToolTipText() const override;
 	virtual FText GetToolTipTextForPin(const URigVMPin* InPin) const override;
+	virtual TArray<URigVMPin*> GetAggregateInputs() const override;
+	virtual TArray<URigVMPin*> GetAggregateOutputs() const override;
+	TArray<URigVMPin*> GetAggregatePins(const ERigVMPinDirection& InDirection) const;
+	virtual FName GetNextAggregateName(const FName& InLastAggregatePinName) const override;
 
 	// Returns the UStruct for this unit node
 	// (the struct declaring the RIGVM_METHOD)
@@ -127,6 +137,10 @@ public:
 	// returns the filtered types of this pin
 	TArray<TRigVMTypeIndex> GetFilteredTypesForPin(URigVMPin* InPin) const;
 
+	// Tries to reduce the input types to a single type, if all are compatible
+	// Will prioritize the InPreferredType if available
+	bool TryReduceTypesToSingle(TArray<TRigVMTypeIndex>& InOutTypes, const TRigVMTypeIndex InPreferredType = TRigVMTypeIndex()) const;	
+
 	// returns true if updating pin filters with InTypes would result in different filters 
 	bool PinNeedsFilteredTypesUpdate(URigVMPin* InPin, const TArray<TRigVMTypeIndex>& InTypeIndices);
 	bool PinNeedsFilteredTypesUpdate(URigVMPin* InPin, URigVMPin* LinkedPin);
@@ -139,7 +153,7 @@ public:
 	void InitializeFilteredPermutations();
 
 	// Initializes the filtered permutations and preferred permutation from the types of the pins
-	void InitializeFilteredPermutationsFromTypes();
+	void InitializeFilteredPermutationsFromTypes(bool bAllowCasting);
 
 	// Converts the preferred types per pin from index to string
 	void ConvertPreferredTypesToString();
@@ -147,6 +161,9 @@ public:
 	// Converts the preferred types per pin from string to indices
 	void ConvertPreferredTypesToTypeIndex();
 
+	// Returns the preferred type, or an invalid type if non was found
+	TRigVMTypeIndex GetPreferredType(const FName& ArgumentName) const;
+
 protected:
 
 	virtual void InvalidateCache() override;
@@ -164,7 +181,7 @@ protected:
 	UPROPERTY()
 	FString ResolvedFunctionName;
 
-	// Indicates a preferred permuation using the types of the arguments
+	// Indicates a preferred permutation using the types of the arguments
 	// Each element is in the format "ArgumentName:CPPType"
 	UPROPERTY()
 	TArray<FString> PreferredPermutationTypes_DEPRECATED;
@@ -184,5 +201,6 @@ protected:
 	friend struct FRigVMSetTemplateFilteredPermutationsAction;
 	friend struct FRigVMSetPreferredTemplatePermutationsAction;
 	friend struct FRigVMRemoveNodeAction;
+	friend class FRigVMTemplatesLibraryNode3Test;
 };
 
diff --git a/Engine/Source/Developer/RigVMDeveloper/Public/RigVMModel/RigVMController.h b/Engine/Source/Developer/RigVMDeveloper/Public/RigVMModel/RigVMController.h
--- a/Engine/Source/Developer/RigVMDeveloper/Public/RigVMModel/RigVMController.h
+++ b/Engine/Source/Developer/RigVMDeveloper/Public/RigVMModel/RigVMController.h
@@ -284,7 +284,7 @@ public:
 	// Turns a resolved templated node(s) back into its template.
 	UFUNCTION(BlueprintCallable, Category = RigVMController)
 	bool UnresolveTemplateNodes(const TArray<FName>& InNodeNames, bool bSetupUndoRedo = true, bool bPrintPythonCommand = false);
-	bool UnresolveTemplateNodes(const TArray<URigVMTemplateNode*>& InNodes, bool bSetupUndoRedo);
+	bool UnresolveTemplateNodes(const TArray<URigVMNode*>& InNodes, bool bSetupUndoRedo);
 
 	// Upgrades a set of nodes with each corresponding next known version
 	UFUNCTION(BlueprintCallable, Category = RigVMController)
@@ -859,7 +859,7 @@ public:
 	static URigVMBuildData* GetBuildData(bool bCreateIfNeeded = true);
 
 	int32 DetachLinksFromPinObjects(const TArray<URigVMLink*>* InLinks = nullptr, bool bNotify = false);
-	int32 ReattachLinksToPinObjects(bool bFollowCoreRedirectors = false, const TArray<URigVMLink*>* InLinks = nullptr, bool bNotify = false, bool bSetupOrphanedPins = false);
+	int32 ReattachLinksToPinObjects(bool bFollowCoreRedirectors = false, const TArray<URigVMLink*>* InLinks = nullptr, bool bNotify = false, bool bSetupOrphanedPins = false, bool bAllowNonArgumentLinks = false);
 	void AddPinRedirector(bool bInput, bool bOutput, const FString& OldPinPath, const FString& NewPinPath);
 
 	// Removes nodes which went stale.
@@ -876,8 +876,14 @@ public:
 	bool RemoveUnusedOrphanedPins(URigVMNode* InNode, bool bNotify);
 
 	// Initializes and recomputes the filtered permutations of all template nodes in the graph
-	// Returns true if any pin has change it's type
-	bool RecomputeAllTemplateFilteredTypes(bool bSetupUndoRedo);
+	// Returns true if any pin has change it's type or link was broken
+	bool RecomputeAllTemplateFilteredPermutations(bool bSetupUndoRedo);
+
+	// Update the template of a subgraph with the filtered permutations of the interface nodes
+	bool UpdateLibraryTemplate(URigVMLibraryNode* LibraryNode, bool bSetupUndoRedo);
+
+	// Update filtered permutations, and propagate both ways of the link before adding this link
+	bool PrepareToLink(URigVMPin* FirstToResolve, URigVMPin* SecondToResolve, bool bSetupUndoRedo);
 
 	// Try to initialize the filterd permutations from the pin types
 	void InitializeFilteredPermutationsFromTemplateTypes();
@@ -897,6 +903,9 @@ public:
 	// A flag that can be used to turn off pin default value validation if necessary
 	bool bValidatePinDefaults;
 
+	// A flag to suspend the recomputation of filtered permutations of outer graphs
+	bool bSuspendRecomputingOuterTemplateFilters;
+
 	const FRigVMByteCode* GetCurrentByteCode() const;
 
 	void ReportInfo(const FString& InMessage) const;
@@ -1056,23 +1065,37 @@ private:
 
 	// Get filtered types for a wildcard node. If template node, that means just returning its filtered wildcard types, but if it's another type of node (select, if, rereoute), iterate
 	// its connections to figure out the filtered types
-	TArray<TRigVMTypeIndex> GetWildcardFilteredTypeIndices(URigVMPin* InPin);
+	TArray<TRigVMTypeIndex> GetFilteredTypes(URigVMPin* InPin);
 	
 	// Updates the permutations allowed without having to break any links
 	bool UpdateFilteredPermutations(URigVMPin* InPin, URigVMPin* InLinkedPin, bool bSetupUndoRedo);
 	bool UpdateFilteredPermutations(URigVMPin* InPin, const TArray<TRigVMTypeIndex>& InTypeIndices, bool bSetupUndoRedo);
 
 	// Changes Pin types if filtered types of a pin are unique
-	bool UpdateTemplateNodePinTypes(URigVMTemplateNode* InNode, bool bSetupUndoRedo);
+	bool UpdateTemplateNodePinTypes(URigVMTemplateNode* InNode, bool bSetupUndoRedo, bool bInitializeDefaultValue = true);
 
 	// Reduces the filtered permutations of all templates in the graph to comply with the types filtered by InNode
 	// Returns false if a link had to be broken
 	bool PropagateTemplateFilteredTypes(URigVMTemplateNode* InNode, bool bSetupUndoRedo);
 
-	bool ChangePinType(const FString& InPinPath, const FString& InCPPType, const FName& InCPPTypeObjectPath, bool bSetupUndoRedo, bool bSetupOrphanPins = true, bool bBreakLinks = true, bool bRemoveSubPins = true);
-	bool ChangePinType(URigVMPin* InPin, const FString& InCPPType, const FName& InCPPTypeObjectPath, bool bSetupUndoRedo, bool bSetupOrphanPins = true, bool bBreakLinks = true, bool bRemoveSubPins = true);
-	bool ChangePinType(URigVMPin* InPin, const FString& InCPPType,UObject* InCPPTypeObject, bool bSetupUndoRedo, bool bSetupOrphanPins = true, bool bBreakLinks = true, bool bRemoveSubPins = true);
-	bool ChangePinType(URigVMPin* InPin, TRigVMTypeIndex InTypeIndex, bool bSetupUndoRedo, bool bSetupOrphanPins = true, bool bBreakLinks = true, bool bRemoveSubPins = true);
+	// Adds a preferred type for the pin
+	// Returns false if the pin already has a different type
+	bool AddPreferredType(URigVMTemplateNode* InNode, const FName& InPinName, const int32 InPreferredTypeIndex, bool bSetupUndoRedo);
+
+	// Removes preferred type
+	// Returns true if the preferred type was found and removed
+	bool RemovePreferredType(URigVMTemplateNode* InNode, const FName& InPinName, bool bSetupUndoRedo);
+
+	// Returns true if the pin is connected, and the filtered types is reduced (not infinite like reroute, if, select or array nodes)
+	bool ShouldPinOwnArgument(URigVMPin* InPin);
+
+	// Given a Entry or Return pin and a potential pin to link it to, try to add the first pin as an argument of the library's template
+	bool AddArgumentForPin(URigVMPin* InPin, URigVMPin* InToLinkPin, bool bSetupUndoRedo = true, bool bPrintPythonCommand = false);
+
+	bool ChangePinType(const FString& InPinPath, const FString& InCPPType, const FName& InCPPTypeObjectPath, bool bSetupUndoRedo, bool bSetupOrphanPins = true, bool bBreakLinks = true, bool bRemoveSubPins = true, bool bInitializeDefaultValue = true);
+	bool ChangePinType(URigVMPin* InPin, const FString& InCPPType, const FName& InCPPTypeObjectPath, bool bSetupUndoRedo, bool bSetupOrphanPins = true, bool bBreakLinks = true, bool bRemoveSubPins = true, bool bInitializeDefaultValue = true);
+	bool ChangePinType(URigVMPin* InPin, const FString& InCPPType,UObject* InCPPTypeObject, bool bSetupUndoRedo, bool bSetupOrphanPins = true, bool bBreakLinks = true, bool bRemoveSubPins = true, bool bInitializeDefaultValue = true);
+	bool ChangePinType(URigVMPin* InPin, TRigVMTypeIndex InTypeIndex, bool bSetupUndoRedo, bool bSetupOrphanPins = true, bool bBreakLinks = true, bool bRemoveSubPins = true, bool bInitializeDefaultValue = true);
 
 #if WITH_EDITOR
 	void RewireLinks(URigVMPin* OldPin, URigVMPin* NewPin, bool bAsInput, bool bSetupUndoRedo, TArray<URigVMLink*> InLinks = TArray<URigVMLink*>());
diff --git a/Engine/Source/Developer/RigVMDeveloper/Public/RigVMModel/RigVMControllerActions.h b/Engine/Source/Developer/RigVMDeveloper/Public/RigVMModel/RigVMControllerActions.h
--- a/Engine/Source/Developer/RigVMDeveloper/Public/RigVMModel/RigVMControllerActions.h
+++ b/Engine/Source/Developer/RigVMDeveloper/Public/RigVMModel/RigVMControllerActions.h
@@ -1127,7 +1127,7 @@ struct FRigVMSetTemplateFilteredPermutationsAction : public FRigVMBaseAction
 public:
 
 	FRigVMSetTemplateFilteredPermutationsAction() {}
-	FRigVMSetTemplateFilteredPermutationsAction(URigVMTemplateNode* InNode, URigVMLink* InLink, const TArray<int32>& InOldFilteredPermutations);
+	FRigVMSetTemplateFilteredPermutationsAction(URigVMTemplateNode* InNode, const TArray<int32>& InOldFilteredPermutations);
 	virtual ~FRigVMSetTemplateFilteredPermutationsAction() {};
 	virtual UScriptStruct* GetScriptStruct() const override { return FRigVMSetTemplateFilteredPermutationsAction::StaticStruct(); }
 	virtual bool Merge(const FRigVMBaseAction* Other);
@@ -1137,9 +1137,6 @@ public:
 	UPROPERTY()
 	FString NodePath;
 
-	UPROPERTY()
-	FString LinkPath;
-
 	UPROPERTY()
 	TArray<int32> OldFilteredPermutations;
 
@@ -1175,6 +1172,32 @@ public:
 	TArray<FRigVMTemplatePreferredType> NewPreferredPermutationTypes;
 };
 
+/**
+ * An action setting the filtered permutations on a template node
+ */
+USTRUCT()
+struct FRigVMSetLibraryTemplateAction : public FRigVMBaseAction
+{
+	GENERATED_BODY()
+
+public:
+
+	FRigVMSetLibraryTemplateAction() {}
+	FRigVMSetLibraryTemplateAction(URigVMLibraryNode* InNode, FRigVMTemplate& InNewTemplate);
+	virtual UScriptStruct* GetScriptStruct() const override { return FRigVMSetTemplateFilteredPermutationsAction::StaticStruct(); }
+	virtual bool Undo(URigVMController* InController) override;
+	virtual bool Redo(URigVMController* InController) override;
+
+	UPROPERTY()
+	FString NodePath;
+
+	UPROPERTY()
+	TArray<uint8> OldTemplateBytes;
+
+	UPROPERTY()
+	TArray<uint8> NewTemplateBytes;
+};
+
 /**
  * An action inserting a new array pin in the graph.
  */
@@ -1276,6 +1299,9 @@ public:
 	virtual bool Undo(URigVMController* InController) override;
 	virtual bool Redo(URigVMController* InController) override;
 
+	UPROPERTY()
+	FSoftObjectPath GraphPath;
+
 	UPROPERTY()
 	FString OutputPinPath;
 
diff --git a/Engine/Source/Developer/RigVMDeveloper/Public/RigVMModel/RigVMPin.h b/Engine/Source/Developer/RigVMDeveloper/Public/RigVMModel/RigVMPin.h
--- a/Engine/Source/Developer/RigVMDeveloper/Public/RigVMModel/RigVMPin.h
+++ b/Engine/Source/Developer/RigVMDeveloper/Public/RigVMModel/RigVMPin.h
@@ -384,7 +384,7 @@ public:
 
 	// Returns true is the two provided source and target Pins
 	// can be linked to one another.
-	static bool CanLink(URigVMPin* InSourcePin, URigVMPin* InTargetPin, FString* OutFailureReason, const FRigVMByteCode* InByteCode, ERigVMPinDirection InUserLinkDirection = ERigVMPinDirection::IO);
+	static bool CanLink(URigVMPin* InSourcePin, URigVMPin* InTargetPin, FString* OutFailureReason, const FRigVMByteCode* InByteCode, ERigVMPinDirection InUserLinkDirection = ERigVMPinDirection::IO, bool bInAllowWildcard = false);
 
 	// Returns true if this pin has injected nodes
 	bool HasInjectedNodes() const { return InjectionInfos.Num() > 0; }
diff --git a/Engine/Source/Runtime/RigVM/Private/RigVMCore/RigVMTemplate.cpp b/Engine/Source/Runtime/RigVM/Private/RigVMCore/RigVMTemplate.cpp
--- a/Engine/Source/Runtime/RigVM/Private/RigVMCore/RigVMTemplate.cpp
+++ b/Engine/Source/Runtime/RigVM/Private/RigVMCore/RigVMTemplate.cpp
@@ -89,6 +89,80 @@ FRigVMTemplateArgument::FRigVMTemplateArgument(const FName& InName, ERigVMPinDir
 	UpdateTypeToPermutations();
 }
 
+void FRigVMTemplateArgument::Serialize(FArchive& Ar)
+{
+	Ar << Index;
+	Ar << Name;
+	Ar << Direction;
+
+#if UE_RIGVM_DEBUG_TYPEINDEX
+	if (Ar.IsSaving())
+	{
+		TArray<int32> TypeIndicesInt;
+		TypeIndicesInt.Reserve(TypeIndices.Num());
+		for (TRigVMTypeIndex& TypeIndex : TypeIndices)
+		{
+			TypeIndicesInt.Add(TypeIndex.GetIndex());
+		}
+		Ar << TypeIndicesInt;
+	}
+	else if (Ar.IsLoading())
+	{
+		TArray<int32> TypeIndicesInt;
+		Ar << TypeIndicesInt;
+		TypeIndices.Reset();
+		TypeIndices.SetNum(TypeIndicesInt.Num());
+		static FRigVMRegistry& Registry = FRigVMRegistry::Get();
+		for (int32 i=0; i<TypeIndicesInt.Num(); ++i)
+		{
+			TypeIndices[i].Index = TypeIndicesInt[i];
+			TypeIndices[i].Name = Registry.GetType(TypeIndices[i]).CPPType;
+		}
+	}
+#else
+	Ar << TypeIndices;
+#endif
+
+	if (Ar.IsLoading())
+	{
+		for (int32 i=0; i<TypeIndices.Num(); ++i)
+		{
+			if (TArray<int32>* Permutations = TypeToPermutations.Find(TypeIndices[i]))
+			{
+				Permutations->Add(i);
+			}
+			else
+			{
+				TypeToPermutations.Add(TypeIndices[i], {i});
+			}
+		}
+	}
+
+	if (Ar.IsSaving())
+	{
+		int32 CategoriesNum = TypeCategories.Num();
+		Ar << CategoriesNum;
+		for (int32 i=0; i<CategoriesNum; ++i)
+		{
+			uint8 CategoryIndex = (int8)TypeCategories[i]; 
+			Ar << CategoryIndex;
+		}
+	}
+	else if (Ar.IsLoading())
+	{
+		int32 CategoriesNum = 0;
+		Ar << CategoriesNum;
+		TypeCategories.SetNumUninitialized(CategoriesNum);
+		
+		for (int32 i=0; i<CategoriesNum; ++i)
+		{
+			uint8 CategoryIndex = INDEX_NONE;
+			Ar << CategoryIndex;
+			TypeCategories[i] = (FRigVMTemplateArgument::ETypeCategory) CategoryIndex;
+		}
+	}	
+}
+
 void FRigVMTemplateArgument::UpdateTypeToPermutations()
 {
 	for(int32 TypeIndex=0;TypeIndex<TypeIndices.Num();TypeIndex++)
@@ -255,6 +329,38 @@ FRigVMTemplate::FRigVMTemplate()
 
 }
 
+void FRigVMTemplate::Serialize(FArchive& Ar)
+{
+	Ar << Index;
+	Ar << Notation;
+	Ar << Permutations;
+
+	if (Ar.IsSaving())
+	{
+		int32 ArgumentsNum = Arguments.Num();
+		Ar << ArgumentsNum;
+
+		for (FRigVMTemplateArgument& Argument : Arguments)
+		{
+			Argument.Serialize(Ar);
+		}
+	}
+	else if(Ar.IsLoading())
+	{
+		int32 ArgumentsNum = 0;
+		Ar << ArgumentsNum;
+
+		Arguments.Reset();
+		Arguments.Reserve(ArgumentsNum);
+		for (int32 i=0; i<ArgumentsNum; ++i)
+		{
+			FRigVMTemplateArgument Argument;
+			Argument.Serialize(Ar);
+			Arguments.Add(Argument);
+		}
+	}	
+}
+
 FRigVMTemplate::FRigVMTemplate(UScriptStruct* InStruct, const FString& InTemplateName, int32 InFunctionIndex)
 	: Index(INDEX_NONE)
 	, Notation(NAME_None)
@@ -410,6 +516,21 @@ FString FRigVMTemplate::GetArgumentNotation(const FRigVMTemplateArgument& InArgu
 		*InArgument.GetName().ToString());
 }
 
+void FRigVMTemplate::ComputeNotationFromArguments(const FString& InTemplateName)
+{
+	TArray<FString> ArgumentNotations;			
+	for (FRigVMTemplateArgument& Argument : Arguments)
+	{
+		if(FRigVMTemplate::IsValidArgumentForTemplate(Argument))
+		{
+			ArgumentNotations.Add(FRigVMTemplate::GetArgumentNotation(Argument));
+		}
+	}
+
+	const FString NotationStr = FString::Printf(TEXT("%s(%s)"), *InTemplateName, *FString::Join(ArgumentNotations, TEXT(",")));
+	Notation = *NotationStr;
+}
+
 TArray<UStruct*> FRigVMTemplate::GetSuperStructs(UStruct* InStruct, bool bIncludeLeaf)
 {
 	// Create an array of structs, ordered super -> child struct
@@ -430,7 +551,6 @@ TArray<UStruct*> FRigVMTemplate::GetSuperStructs(UStruct* InStruct, bool bInclud
 	{
 		SuperStructs.Remove(SuperStructs.Last());
 	}
-
 	return SuperStructs;
 }
 
@@ -977,7 +1097,7 @@ bool FRigVMTemplate::Resolve(FTypeMap& InOutTypes, TArray<int32>& OutPermutation
 		}
 		else if (const TRigVMTypeIndex* InputType = InputTypes.Find(Argument.Name))
 		{
-			int32 MatchedType = *InputType;
+			TRigVMTypeIndex MatchedType = *InputType;
 			bool bFoundMatch = false;
 			bool bFoundPerfectMatch = false;
 			for (int32 PermutationIndex = 0; PermutationIndex < Argument.TypeIndices.Num(); PermutationIndex++)
diff --git a/Engine/Source/Runtime/RigVM/Public/RigVMCore/RigVMTemplate.h b/Engine/Source/Runtime/RigVM/Public/RigVMCore/RigVMTemplate.h
--- a/Engine/Source/Runtime/RigVM/Public/RigVMCore/RigVMTemplate.h
+++ b/Engine/Source/Runtime/RigVM/Public/RigVMCore/RigVMTemplate.h
@@ -139,7 +139,7 @@ struct RIGVM_API FRigVMTemplateArgument
 		EArrayType_Invalid
 	};
 
-	enum ETypeCategory
+	enum ETypeCategory : uint8
 	{
 		ETypeCategory_Execute,
 		ETypeCategory_SingleAnyValue,
@@ -170,6 +170,9 @@ struct RIGVM_API FRigVMTemplateArgument
 	FRigVMTemplateArgument(const FName& InName, ERigVMPinDirection InDirection, const TArray<TRigVMTypeIndex>& InTypeIndices);
 	FRigVMTemplateArgument(const FName& InName, ERigVMPinDirection InDirection, const TArray<ETypeCategory>& InTypeCategories);
 
+	// Serialize
+	void Serialize(FArchive& Ar);
+
 	// returns the name of the argument
 	const FName& GetName() const { return Name; }
 
@@ -218,6 +221,7 @@ protected:
 	friend class URigVMController;
 	friend struct FRigVMRegistry;
 	friend struct FRigVMStructUpgradeInfo;
+	friend struct FRigVMSetLibraryTemplateAction;
 };
 
 /**
@@ -236,6 +240,9 @@ public:
 	// Default constructor
 	FRigVMTemplate();
 
+	// Serialize
+	void Serialize(FArchive& Ar);
+
 	// returns true if this is a valid template
 	bool IsValid() const;
 
@@ -302,6 +309,9 @@ public:
 	// returns the notation of an argument
 	static FString GetArgumentNotation(const FRigVMTemplateArgument& InArgument);
 
+	// recomputes the notation from its arguments
+	void ComputeNotationFromArguments(const FString& InTemplateName);
+
 	// returns an array of structs in the inheritance order of a given struct
 	static TArray<UStruct*> GetSuperStructs(UStruct* InStruct, bool bIncludeLeaf = true);
 
@@ -371,5 +381,6 @@ private:
 	friend struct FRigVMRegistry;
 	friend class URigVMController;
 	friend class URigVMLibraryNode;
+	friend struct FRigVMSetLibraryTemplateAction;
 	friend struct FRigVMDispatchFactory;
 };
diff --git a/Engine/Source/Runtime/RigVM/Public/RigVMCore/RigVMTypeIndex.h b/Engine/Source/Runtime/RigVM/Public/RigVMCore/RigVMTypeIndex.h
--- a/Engine/Source/Runtime/RigVM/Public/RigVMCore/RigVMTypeIndex.h
+++ b/Engine/Source/Runtime/RigVM/Public/RigVMCore/RigVMTypeIndex.h
@@ -77,7 +77,8 @@ protected:
 	FName Name;
 	int32 Index;
 
-	friend struct FRigVMRegistry; 
+	friend struct FRigVMRegistry;
+	friend struct FRigVMTemplateArgument;
 };
 #else
 typedef int32 TRigVMTypeIndex;
