Control Rig: Fix ensure hit when openning rigs with FRigPose pins

diff --git a/Engine/Source/Developer/RigVMDeveloper/Private/RigVMModel/RigVMController.cpp b/Engine/Source/Developer/RigVMDeveloper/Private/RigVMModel/RigVMController.cpp
--- a/Engine/Source/Developer/RigVMDeveloper/Private/RigVMModel/RigVMController.cpp
+++ b/Engine/Source/Developer/RigVMDeveloper/Private/RigVMModel/RigVMController.cpp
@@ -16783,6 +16783,7 @@ bool URigVMController::ShouldPinOwnArgument(URigVMPin* InPin)
 	URigVMFunctionReturnNode* ReturnNode = ContainedGraph->GetReturnNode();
 
 	auto IsPinUsingArgument = [](URigVMPin* TestPin)
+	TMap<URigVMPin*, TRigVMTypeIndex> TypesBeforeRecomputing;
 	{
 		// If any subpins that are struct members are connected, then it should own an argument
 		TArray<URigVMPin*> ToProcessPins = TestPin->GetSubPins();
