Control Rig: Adding aggregates for template nodes

diff --git a/Engine/Plugins/Animation/ControlRig/Source/ControlRig/Private/Units/Math/RigUnit_MathBool.h b/Engine/Plugins/Animation/ControlRig/Source/ControlRig/Private/Units/Math/RigUnit_MathBool.h
--- a/Engine/Plugins/Animation/ControlRig/Source/ControlRig/Private/Units/Math/RigUnit_MathBool.h
+++ b/Engine/Plugins/Animation/ControlRig/Source/ControlRig/Private/Units/Math/RigUnit_MathBool.h
@@ -63,6 +63,26 @@ struct CONTROLRIG_API FRigUnit_MathBoolBinaryOp : public FRigUnit_MathBoolBase
 	bool Result;
 };
 
+USTRUCT(meta=(Abstract))
+struct CONTROLRIG_API FRigUnit_MathBoolBinaryAggregateOp : public FRigUnit_MathBoolBase
+{
+	GENERATED_BODY()
+
+	FRigUnit_MathBoolBinaryAggregateOp()
+	{
+		A = B = Result = 0.f;
+	}
+
+	UPROPERTY(meta=(Input, Aggregate))
+	bool A;
+
+	UPROPERTY(meta=(Input, Aggregate))
+	bool B;
+
+	UPROPERTY(meta=(Output, Aggregate))
+	bool Result;
+};
+
 /**
  * Returns true
  */
@@ -113,7 +133,7 @@ struct CONTROLRIG_API FRigUnit_MathBoolNot : public FRigUnit_MathBoolUnaryOp
  * Returns true if both conditions are true
  */
 USTRUCT(meta=(DisplayName="And", TemplateName="And", Keywords="&&"))
-struct CONTROLRIG_API FRigUnit_MathBoolAnd : public FRigUnit_MathBoolBinaryOp
+struct CONTROLRIG_API FRigUnit_MathBoolAnd : public FRigUnit_MathBoolBinaryAggregateOp
 {
 	GENERATED_BODY()
 
@@ -137,7 +157,7 @@ struct CONTROLRIG_API FRigUnit_MathBoolNand : public FRigUnit_MathBoolBinaryOp
  * Returns true if one of the conditions is true
  */
 USTRUCT(meta=(DisplayName="Or", TemplateName="Or", Keywords="||"))
-struct CONTROLRIG_API FRigUnit_MathBoolOr : public FRigUnit_MathBoolBinaryOp
+struct CONTROLRIG_API FRigUnit_MathBoolOr : public FRigUnit_MathBoolBinaryAggregateOp
 {
 	GENERATED_BODY()
 
diff --git a/Engine/Plugins/Animation/ControlRig/Source/ControlRig/Private/Units/Math/RigUnit_MathColor.h b/Engine/Plugins/Animation/ControlRig/Source/ControlRig/Private/Units/Math/RigUnit_MathColor.h
--- a/Engine/Plugins/Animation/ControlRig/Source/ControlRig/Private/Units/Math/RigUnit_MathColor.h
+++ b/Engine/Plugins/Animation/ControlRig/Source/ControlRig/Private/Units/Math/RigUnit_MathColor.h
@@ -32,6 +32,26 @@ struct CONTROLRIG_API FRigUnit_MathColorBinaryOp : public FRigUnit_MathColorBase
 	FLinearColor Result;
 };
 
+USTRUCT(meta=(Abstract))
+struct CONTROLRIG_API FRigUnit_MathColorBinaryAggregateOp : public FRigUnit_MathColorBase
+{
+	GENERATED_BODY()
+
+	FRigUnit_MathColorBinaryAggregateOp()
+	{
+		A = B = Result = FLinearColor::Black;
+	}
+
+	UPROPERTY(meta=(Input, Aggregate))
+	FLinearColor A;
+
+	UPROPERTY(meta=(Input, Aggregate))
+	FLinearColor B;
+
+	UPROPERTY(meta=(Output, Aggregate))
+	FLinearColor Result;
+};
+
 /**
  * Makes a vector from a single float
  */
@@ -59,7 +79,7 @@ struct CONTROLRIG_API FRigUnit_MathColorFromFloat : public FRigUnit_MathColorBas
  * Returns the sum of the two values
  */
 USTRUCT(meta=(DisplayName="Add", TemplateName="Add", Keywords="Sum,+"))
-struct CONTROLRIG_API FRigUnit_MathColorAdd : public FRigUnit_MathColorBinaryOp
+struct CONTROLRIG_API FRigUnit_MathColorAdd : public FRigUnit_MathColorBinaryAggregateOp
 {
 	GENERATED_BODY()
 
@@ -83,7 +103,7 @@ struct CONTROLRIG_API FRigUnit_MathColorSub : public FRigUnit_MathColorBinaryOp
  * Returns the product of the two values
  */
 USTRUCT(meta=(DisplayName="Multiply", TemplateName="Multiply", Keywords="Product,*"))
-struct CONTROLRIG_API FRigUnit_MathColorMul : public FRigUnit_MathColorBinaryOp
+struct CONTROLRIG_API FRigUnit_MathColorMul : public FRigUnit_MathColorBinaryAggregateOp
 {
 	GENERATED_BODY()
 
diff --git a/Engine/Plugins/Animation/ControlRig/Source/ControlRig/Private/Units/Math/RigUnit_MathDouble.h b/Engine/Plugins/Animation/ControlRig/Source/ControlRig/Private/Units/Math/RigUnit_MathDouble.h
--- a/Engine/Plugins/Animation/ControlRig/Source/ControlRig/Private/Units/Math/RigUnit_MathDouble.h
+++ b/Engine/Plugins/Animation/ControlRig/Source/ControlRig/Private/Units/Math/RigUnit_MathDouble.h
@@ -62,6 +62,27 @@ struct CONTROLRIG_API FRigUnit_MathDoubleBinaryOp : public FRigUnit_MathDoubleBa
 	double Result;
 };
 
+
+USTRUCT(meta=(Abstract))
+struct CONTROLRIG_API FRigUnit_MathDoubleBinaryAggregateOp : public FRigUnit_MathDoubleBase
+{
+	GENERATED_BODY()
+
+	FRigUnit_MathDoubleBinaryAggregateOp()
+	{
+		A = B = Result = 0.0;
+	}
+
+	UPROPERTY(meta=(Input, Aggregate))
+	double A;
+
+	UPROPERTY(meta=(Input, Aggregate))
+	double B;
+
+	UPROPERTY(meta=(Output, Aggregate))
+	double Result;
+};
+
 /**
  * Returns PI
  */
@@ -114,7 +135,7 @@ struct CONTROLRIG_API FRigUnit_MathDoubleConstE : public FRigUnit_MathDoubleCons
  * Returns the sum of the two values
  */
 USTRUCT(meta=(DisplayName="Add", TemplateName="Add", Keywords="Sum,+"))
-struct CONTROLRIG_API FRigUnit_MathDoubleAdd : public FRigUnit_MathDoubleBinaryOp
+struct CONTROLRIG_API FRigUnit_MathDoubleAdd : public FRigUnit_MathDoubleBinaryAggregateOp
 {
 	GENERATED_BODY()
 
@@ -138,7 +159,7 @@ struct CONTROLRIG_API FRigUnit_MathDoubleSub : public FRigUnit_MathDoubleBinaryO
  * Returns the product of the two values
  */
 USTRUCT(meta=(DisplayName="Multiply", TemplateName="Multiply", Keywords="Product,*"))
-struct CONTROLRIG_API FRigUnit_MathDoubleMul : public FRigUnit_MathDoubleBinaryOp
+struct CONTROLRIG_API FRigUnit_MathDoubleMul : public FRigUnit_MathDoubleBinaryAggregateOp
 {
 	GENERATED_BODY()
 
@@ -191,7 +212,7 @@ struct CONTROLRIG_API FRigUnit_MathDoubleMod : public FRigUnit_MathDoubleBinaryO
  * Returns the smaller of the two values
  */
 USTRUCT(meta=(DisplayName="Minimum", TemplateName="Minimum"))
-struct CONTROLRIG_API FRigUnit_MathDoubleMin : public FRigUnit_MathDoubleBinaryOp
+struct CONTROLRIG_API FRigUnit_MathDoubleMin : public FRigUnit_MathDoubleBinaryAggregateOp
 {
 	GENERATED_BODY()
 
@@ -203,7 +224,7 @@ struct CONTROLRIG_API FRigUnit_MathDoubleMin : public FRigUnit_MathDoubleBinaryO
  * Returns the larger of the two values
  */
 USTRUCT(meta=(DisplayName="Maximum", TemplateName="Maximum"))
-struct CONTROLRIG_API FRigUnit_MathDoubleMax : public FRigUnit_MathDoubleBinaryOp
+struct CONTROLRIG_API FRigUnit_MathDoubleMax : public FRigUnit_MathDoubleBinaryAggregateOp
 {
 	GENERATED_BODY()
 
diff --git a/Engine/Plugins/Animation/ControlRig/Source/ControlRig/Private/Units/Math/RigUnit_MathFloat.h b/Engine/Plugins/Animation/ControlRig/Source/ControlRig/Private/Units/Math/RigUnit_MathFloat.h
--- a/Engine/Plugins/Animation/ControlRig/Source/ControlRig/Private/Units/Math/RigUnit_MathFloat.h
+++ b/Engine/Plugins/Animation/ControlRig/Source/ControlRig/Private/Units/Math/RigUnit_MathFloat.h
@@ -62,6 +62,26 @@ struct CONTROLRIG_API FRigUnit_MathFloatBinaryOp : public FRigUnit_MathFloatBase
 	float Result;
 };
 
+USTRUCT(meta=(Abstract))
+struct CONTROLRIG_API FRigUnit_MathFloatBinaryAggregateOp : public FRigUnit_MathFloatBase
+{
+	GENERATED_BODY()
+
+	FRigUnit_MathFloatBinaryAggregateOp()
+	{
+		A = B = Result = 0.f;
+	}
+
+	UPROPERTY(meta=(Input, Aggregate))
+	float A;
+
+	UPROPERTY(meta=(Input, Aggregate))
+	float B;
+
+	UPROPERTY(meta=(Output, Aggregate))
+	float Result;
+};
+
 /**
  * Returns PI
  */
@@ -114,7 +134,7 @@ struct CONTROLRIG_API FRigUnit_MathFloatConstE : public FRigUnit_MathFloatConsta
  * Returns the sum of the two values
  */
 USTRUCT(meta=(DisplayName="Add", TemplateName="Add", Keywords="Sum,+"))
-struct CONTROLRIG_API FRigUnit_MathFloatAdd : public FRigUnit_MathFloatBinaryOp
+struct CONTROLRIG_API FRigUnit_MathFloatAdd : public FRigUnit_MathFloatBinaryAggregateOp
 {
 	GENERATED_BODY()
 
@@ -138,7 +158,7 @@ struct CONTROLRIG_API FRigUnit_MathFloatSub : public FRigUnit_MathFloatBinaryOp
  * Returns the product of the two values
  */
 USTRUCT(meta=(DisplayName="Multiply", TemplateName="Multiply", Keywords="Product,*"))
-struct CONTROLRIG_API FRigUnit_MathFloatMul : public FRigUnit_MathFloatBinaryOp
+struct CONTROLRIG_API FRigUnit_MathFloatMul : public FRigUnit_MathFloatBinaryAggregateOp
 {
 	GENERATED_BODY()
 
@@ -191,7 +211,7 @@ struct CONTROLRIG_API FRigUnit_MathFloatMod : public FRigUnit_MathFloatBinaryOp
  * Returns the smaller of the two values
  */
 USTRUCT(meta=(DisplayName="Minimum", TemplateName="Minimum"))
-struct CONTROLRIG_API FRigUnit_MathFloatMin : public FRigUnit_MathFloatBinaryOp
+struct CONTROLRIG_API FRigUnit_MathFloatMin : public FRigUnit_MathFloatBinaryAggregateOp
 {
 	GENERATED_BODY()
 
@@ -203,7 +223,7 @@ struct CONTROLRIG_API FRigUnit_MathFloatMin : public FRigUnit_MathFloatBinaryOp
  * Returns the larger of the two values
  */
 USTRUCT(meta=(DisplayName="Maximum", TemplateName="Maximum"))
-struct CONTROLRIG_API FRigUnit_MathFloatMax : public FRigUnit_MathFloatBinaryOp
+struct CONTROLRIG_API FRigUnit_MathFloatMax : public FRigUnit_MathFloatBinaryAggregateOp
 {
 	GENERATED_BODY()
 
diff --git a/Engine/Plugins/Animation/ControlRig/Source/ControlRig/Private/Units/Math/RigUnit_MathInt.h b/Engine/Plugins/Animation/ControlRig/Source/ControlRig/Private/Units/Math/RigUnit_MathInt.h
--- a/Engine/Plugins/Animation/ControlRig/Source/ControlRig/Private/Units/Math/RigUnit_MathInt.h
+++ b/Engine/Plugins/Animation/ControlRig/Source/ControlRig/Private/Units/Math/RigUnit_MathInt.h
@@ -48,11 +48,31 @@ struct CONTROLRIG_API FRigUnit_MathIntBinaryOp : public FRigUnit_MathIntBase
 	int32 Result;
 };
 
+USTRUCT(meta=(Abstract))
+struct CONTROLRIG_API FRigUnit_MathIntBinaryAggregateOp : public FRigUnit_MathIntBase
+{
+	GENERATED_BODY()
+
+	FRigUnit_MathIntBinaryAggregateOp()
+	{
+		A = B = Result = 0;
+	}
+
+	UPROPERTY(meta=(Input, Aggregate))
+	int32 A;
+
+	UPROPERTY(meta=(Input, Aggregate))
+	int32 B;
+
+	UPROPERTY(meta=(Output, Aggregate))
+	int32 Result;
+};
+
 /**
  * Returns the sum of the two values
  */
 USTRUCT(meta=(DisplayName="Add", TemplateName="Add", Keywords="Sum,+"))
-struct CONTROLRIG_API FRigUnit_MathIntAdd : public FRigUnit_MathIntBinaryOp
+struct CONTROLRIG_API FRigUnit_MathIntAdd : public FRigUnit_MathIntBinaryAggregateOp
 {
 	GENERATED_BODY()
 
@@ -76,7 +96,7 @@ struct CONTROLRIG_API FRigUnit_MathIntSub : public FRigUnit_MathIntBinaryOp
  * Returns the product of the two values
  */
 USTRUCT(meta=(DisplayName="Multiply", TemplateName="Multiply", Keywords="Product,*"))
-struct CONTROLRIG_API FRigUnit_MathIntMul : public FRigUnit_MathIntBinaryOp
+struct CONTROLRIG_API FRigUnit_MathIntMul : public FRigUnit_MathIntBinaryAggregateOp
 {
 	GENERATED_BODY()
 
@@ -129,7 +149,7 @@ struct CONTROLRIG_API FRigUnit_MathIntMod : public FRigUnit_MathIntBinaryOp
  * Returns the smaller of the two values
  */
 USTRUCT(meta=(DisplayName="Minimum", TemplateName="Minimum"))
-struct CONTROLRIG_API FRigUnit_MathIntMin : public FRigUnit_MathIntBinaryOp
+struct CONTROLRIG_API FRigUnit_MathIntMin : public FRigUnit_MathIntBinaryAggregateOp
 {
 	GENERATED_BODY()
 
@@ -141,7 +161,7 @@ struct CONTROLRIG_API FRigUnit_MathIntMin : public FRigUnit_MathIntBinaryOp
  * Returns the larger of the two values
  */
 USTRUCT(meta=(DisplayName="Maximum", TemplateName="Maximum"))
-struct CONTROLRIG_API FRigUnit_MathIntMax : public FRigUnit_MathIntBinaryOp
+struct CONTROLRIG_API FRigUnit_MathIntMax : public FRigUnit_MathIntBinaryAggregateOp
 {
 	GENERATED_BODY()
 
diff --git a/Engine/Plugins/Animation/ControlRig/Source/ControlRig/Private/Units/Math/RigUnit_MathMatrix.h b/Engine/Plugins/Animation/ControlRig/Source/ControlRig/Private/Units/Math/RigUnit_MathMatrix.h
--- a/Engine/Plugins/Animation/ControlRig/Source/ControlRig/Private/Units/Math/RigUnit_MathMatrix.h
+++ b/Engine/Plugins/Animation/ControlRig/Source/ControlRig/Private/Units/Math/RigUnit_MathMatrix.h
@@ -50,6 +50,26 @@ struct CONTROLRIG_API FRigUnit_MathMatrixBinaryOp : public FRigUnit_MathMatrixBa
 	FMatrix Result;
 };
 
+USTRUCT(meta=(Abstract))
+struct CONTROLRIG_API FRigUnit_MathMatrixBinaryAggregateOp : public FRigUnit_MathMatrixBase
+{
+	GENERATED_BODY()
+
+	FRigUnit_MathMatrixBinaryAggregateOp()
+	{
+		A = B = Result = FMatrix::Identity;
+	}
+
+	UPROPERTY(meta=(Input, Aggregate))
+	FMatrix A;
+
+	UPROPERTY(meta=(Input, Aggregate))
+	FMatrix B;
+
+	UPROPERTY(meta=(Output, Aggregate))
+	FMatrix Result;
+};
+
 /**
 * Makes a transform from a matrix
 */
@@ -201,7 +221,7 @@ struct CONTROLRIG_API FRigUnit_MathMatrixFromVectors : public FRigUnit_MathMatri
  * Returns the product of the two values
  */
 USTRUCT(meta=(DisplayName="Multiply", TemplateName="Multiply", Keywords="Product,*,Global"))
-struct CONTROLRIG_API FRigUnit_MathMatrixMul : public FRigUnit_MathMatrixBinaryOp
+struct CONTROLRIG_API FRigUnit_MathMatrixMul : public FRigUnit_MathMatrixBinaryAggregateOp
 {
 	GENERATED_BODY()
 	RIGVM_METHOD()
diff --git a/Engine/Plugins/Animation/ControlRig/Source/ControlRig/Private/Units/Math/RigUnit_MathQuaternion.h b/Engine/Plugins/Animation/ControlRig/Source/ControlRig/Private/Units/Math/RigUnit_MathQuaternion.h
--- a/Engine/Plugins/Animation/ControlRig/Source/ControlRig/Private/Units/Math/RigUnit_MathQuaternion.h
+++ b/Engine/Plugins/Animation/ControlRig/Source/ControlRig/Private/Units/Math/RigUnit_MathQuaternion.h
@@ -49,6 +49,26 @@ struct CONTROLRIG_API FRigUnit_MathQuaternionBinaryOp : public FRigUnit_MathQuat
 	FQuat Result;
 };
 
+USTRUCT(meta=(Abstract))
+struct CONTROLRIG_API FRigUnit_MathQuaternionBinaryAggregateOp : public FRigUnit_MathQuaternionBase
+{
+	GENERATED_BODY()
+
+	FRigUnit_MathQuaternionBinaryAggregateOp()
+	{
+		A = B = Result = FQuat::Identity;
+	}
+
+	UPROPERTY(meta=(Input, Aggregate))
+	FQuat A;
+
+	UPROPERTY(meta=(Input, Aggregate))
+	FQuat B;
+
+	UPROPERTY(meta=(Output, Aggregate))
+	FQuat Result;
+};
+
 /**
  * Makes a quaternion from an axis and an angle in radians
  */
@@ -321,7 +341,7 @@ struct CONTROLRIG_API FRigUnit_MathQuaternionToRotator : public FRigUnit_MathQua
  * Returns the product of the two values
  */
 USTRUCT(meta=(DisplayName="Multiply", TemplateName="Multiply", Keywords="Product,*"))
-struct CONTROLRIG_API FRigUnit_MathQuaternionMul : public FRigUnit_MathQuaternionBinaryOp
+struct CONTROLRIG_API FRigUnit_MathQuaternionMul : public FRigUnit_MathQuaternionBinaryAggregateOp
 {
 	GENERATED_BODY()
 	RIGVM_METHOD()
diff --git a/Engine/Plugins/Animation/ControlRig/Source/ControlRig/Private/Units/Math/RigUnit_MathTransform.h b/Engine/Plugins/Animation/ControlRig/Source/ControlRig/Private/Units/Math/RigUnit_MathTransform.h
--- a/Engine/Plugins/Animation/ControlRig/Source/ControlRig/Private/Units/Math/RigUnit_MathTransform.h
+++ b/Engine/Plugins/Animation/ControlRig/Source/ControlRig/Private/Units/Math/RigUnit_MathTransform.h
@@ -56,6 +56,26 @@ struct CONTROLRIG_API FRigUnit_MathTransformBinaryOp : public FRigUnit_MathTrans
 	FTransform Result;
 };
 
+USTRUCT(meta=(Abstract))
+struct CONTROLRIG_API FRigUnit_MathTransformBinaryAggregateOp : public FRigUnit_MathTransformBase
+{
+	GENERATED_BODY()
+
+	FRigUnit_MathTransformBinaryAggregateOp()
+	{
+		A = B = Result = FTransform::Identity;
+	}
+
+	UPROPERTY(meta=(Input, Aggregate))
+	FTransform A;
+
+	UPROPERTY(meta=(Input, Aggregate))
+	FTransform B;
+
+	UPROPERTY(meta=(Output, Aggregate))
+	FTransform Result;
+};
+
 /**
  * Makes a quaternion based transform from a euler based transform
  */
@@ -135,7 +155,7 @@ struct CONTROLRIG_API FRigUnit_MathTransformToEulerTransform : public FRigUnit_M
  * Returns the product of the two values
  */
 USTRUCT(meta=(DisplayName="Multiply", TemplateName="Multiply", Keywords="Product,*,Global"))
-struct CONTROLRIG_API FRigUnit_MathTransformMul : public FRigUnit_MathTransformBinaryOp
+struct CONTROLRIG_API FRigUnit_MathTransformMul : public FRigUnit_MathTransformBinaryAggregateOp
 {
 	GENERATED_BODY()
 	RIGVM_METHOD()
@@ -342,7 +362,7 @@ struct CONTROLRIG_API FRigUnit_MathTransformRotateVector : public FRigUnit_MathT
 /**
  * Multiplies a given vector (location) by the transform
  */
-USTRUCT(meta=(DisplayName="Transform Location", TemplateName="Multiply", Keywords="Multiply"))
+USTRUCT(meta=(DisplayName="Transform Location", Keywords="Multiply"))
 struct CONTROLRIG_API FRigUnit_MathTransformTransformVector : public FRigUnit_MathTransformBase
 {
 	GENERATED_BODY()
diff --git a/Engine/Plugins/Animation/ControlRig/Source/ControlRig/Private/Units/Math/RigUnit_MathVector.h b/Engine/Plugins/Animation/ControlRig/Source/ControlRig/Private/Units/Math/RigUnit_MathVector.h
--- a/Engine/Plugins/Animation/ControlRig/Source/ControlRig/Private/Units/Math/RigUnit_MathVector.h
+++ b/Engine/Plugins/Animation/ControlRig/Source/ControlRig/Private/Units/Math/RigUnit_MathVector.h
@@ -49,6 +49,26 @@ struct CONTROLRIG_API FRigUnit_MathVectorBinaryOp : public FRigUnit_MathVectorBa
 	FVector Result;
 };
 
+USTRUCT(meta=(Abstract))
+struct CONTROLRIG_API FRigUnit_MathVectorBinaryAggregateOp : public FRigUnit_MathVectorBase
+{
+	GENERATED_BODY()
+
+	FRigUnit_MathVectorBinaryAggregateOp()
+	{
+		A = B = Result = FVector::ZeroVector;
+	}
+
+	UPROPERTY(meta=(Input, Aggregate))
+	FVector A;
+
+	UPROPERTY(meta=(Input, Aggregate))
+	FVector B;
+
+	UPROPERTY(meta=(Output, Aggregate))
+	FVector Result;
+};
+
 /**
  * Makes a vector from a single float
  */
@@ -76,7 +96,7 @@ struct CONTROLRIG_API FRigUnit_MathVectorFromFloat : public FRigUnit_MathVectorB
  * Returns the sum of the two values
  */
 USTRUCT(meta=(DisplayName="Add", TemplateName="Add", Keywords="Sum,+"))
-struct CONTROLRIG_API FRigUnit_MathVectorAdd : public FRigUnit_MathVectorBinaryOp
+struct CONTROLRIG_API FRigUnit_MathVectorAdd : public FRigUnit_MathVectorBinaryAggregateOp
 {
 	GENERATED_BODY()
 
@@ -100,7 +120,7 @@ struct CONTROLRIG_API FRigUnit_MathVectorSub : public FRigUnit_MathVectorBinaryO
  * Returns the product of the two values
  */
 USTRUCT(meta=(DisplayName="Multiply", TemplateName="Multiply", Keywords="Product,*"))
-struct CONTROLRIG_API FRigUnit_MathVectorMul : public FRigUnit_MathVectorBinaryOp
+struct CONTROLRIG_API FRigUnit_MathVectorMul : public FRigUnit_MathVectorBinaryAggregateOp
 {
 	GENERATED_BODY()
 
@@ -179,7 +199,7 @@ struct CONTROLRIG_API FRigUnit_MathVectorMod : public FRigUnit_MathVectorBinaryO
  * Returns the smaller of the two values for each component
  */
 USTRUCT(meta=(DisplayName="Minimum", TemplateName="Minimum"))
-struct CONTROLRIG_API FRigUnit_MathVectorMin : public FRigUnit_MathVectorBinaryOp
+struct CONTROLRIG_API FRigUnit_MathVectorMin : public FRigUnit_MathVectorBinaryAggregateOp
 {
 	GENERATED_BODY()
 
@@ -191,7 +211,7 @@ struct CONTROLRIG_API FRigUnit_MathVectorMin : public FRigUnit_MathVectorBinaryO
  * Returns the larger of the two values each component
  */
 USTRUCT(meta=(DisplayName="Maximum", TemplateName="Maximum"))
-struct CONTROLRIG_API FRigUnit_MathVectorMax : public FRigUnit_MathVectorBinaryOp
+struct CONTROLRIG_API FRigUnit_MathVectorMax : public FRigUnit_MathVectorBinaryAggregateOp
 {
 	GENERATED_BODY()
 
diff --git a/Engine/Source/Developer/RigVMDeveloper/Private/RigVMModel/Nodes/RigVMTemplateNode.cpp b/Engine/Source/Developer/RigVMDeveloper/Private/RigVMModel/Nodes/RigVMTemplateNode.cpp
--- a/Engine/Source/Developer/RigVMDeveloper/Private/RigVMModel/Nodes/RigVMTemplateNode.cpp
+++ b/Engine/Source/Developer/RigVMDeveloper/Private/RigVMModel/Nodes/RigVMTemplateNode.cpp
@@ -45,6 +45,7 @@ URigVMTemplateNode::URigVMTemplateNode()
 	, ResolvedFunctionName()
 	, CachedTemplate(nullptr)
 	, CachedFunction(nullptr)
+	, ResolvedPermutation(INDEX_NONE)
 {
 }
 
@@ -131,7 +132,7 @@ TRigVMTypeIndex URigVMTemplateNode::GetPreferredType(const FName& ArgumentName)
 	{
 		if (PreferredType.Argument == ArgumentName)
 		{
-			return PreferredType.TypeIndex;
+			return PreferredType.GetTypeIndex();
 		}
 	}
 	return INDEX_NONE;
@@ -256,7 +257,17 @@ TArray<URigVMPin*> URigVMTemplateNode::GetAggregatePins(const ERigVMPinDirection
 #if UE_RIGVM_AGGREGATE_NODES_ENABLED
 	if (const FRigVMTemplate* Template = GetTemplate())
 	{
-		if (FilteredPermutations.IsEmpty())
+		TArray<int32> Permutations = FilteredPermutations;
+		if (Permutations.IsEmpty())
+		{
+			Permutations.Reserve(Template->NumPermutations());
+			for (int32 i=0; i<Template->NumPermutations(); ++i)
+			{
+				Permutations.Add(i);
+			}
+		}
+		
+		if (Permutations.IsEmpty())
 		{
 			return AggregatePins;
 		}
@@ -287,10 +298,10 @@ TArray<URigVMPin*> URigVMTemplateNode::GetAggregatePins(const ERigVMPinDirection
 			return Inputs;
 		};
 		
-		AggregatePins = FindAggregatePins(FilteredPermutations[0]);
-		for (int32 i=1; i<FilteredPermutations.Num(); ++i)
+		AggregatePins = FindAggregatePins(Permutations[0]);
+		for (int32 i=1; i<Permutations.Num(); ++i)
 		{
-			TArray<URigVMPin*> OtherAggregateInputs = FindAggregatePins(FilteredPermutations[i]);
+			TArray<URigVMPin*> OtherAggregateInputs = FindAggregatePins(Permutations[i]);
 			if (OtherAggregateInputs.Num() != AggregatePins.Num() ||
 				OtherAggregateInputs.FilterByPredicate([&](const URigVMPin* OtherPin)
 				{
@@ -407,15 +418,6 @@ bool URigVMTemplateNode::SupportsType(const URigVMPin* InPin, TRigVMTypeIndex In
 	if (const FRigVMTemplate* Template = GetTemplate())
 	{
 		const TPair<FName,TRigVMTypeIndex> CacheKey(RootPin->GetFName(), TypeIndex);
-		if (const TPair<bool, TRigVMTypeIndex>* CachedResult = SupportedTypesCache.Find(CacheKey))
-		{
-			if(OutTypeIndex)
-			{
-				*OutTypeIndex = CachedResult->Value;
-			}
-			return CachedResult->Key;
-		}
-
 		TRigVMTypeIndex SupportedTypeIndex = INDEX_NONE;
 		if (Template->ArgumentSupportsTypeIndex(RootPin->GetFName(), TypeIndex, &SupportedTypeIndex))
 		{
@@ -547,10 +549,9 @@ const FRigVMFunction* URigVMTemplateNode::GetResolvedFunction() const
 
 		if(CachedFunction == nullptr)
 		{
-			TArray<int32> PermutationIndices = GetFilteredPermutationsIndices();
-			if(PermutationIndices.Num() == 1)
+			if(ResolvedPermutation != INDEX_NONE)
 			{
-				CachedFunction = GetTemplate()->GetPermutation(PermutationIndices[0]);
+				CachedFunction = GetTemplate()->GetPermutation(ResolvedPermutation);
 			}
 		}
 	}
@@ -1082,10 +1083,8 @@ void URigVMTemplateNode::InvalidateCache()
 {
 	Super::InvalidateCache();
 	
-	SupportedTypesCache.Reset();
 	CachedFunction = nullptr;
 	CachedTemplate = nullptr;
-	ResolvedPermutations.Reset();
 
 	for(URigVMPin* Pin : GetPins())
 	{
diff --git a/Engine/Source/Developer/RigVMDeveloper/Private/RigVMModel/RigVMController.cpp b/Engine/Source/Developer/RigVMDeveloper/Private/RigVMModel/RigVMController.cpp
--- a/Engine/Source/Developer/RigVMDeveloper/Private/RigVMModel/RigVMController.cpp
+++ b/Engine/Source/Developer/RigVMDeveloper/Private/RigVMModel/RigVMController.cpp
@@ -3513,7 +3513,7 @@ TArray<FName> URigVMController::ImportNodesFromText(const FString& InText, bool
 
 			if (URigVMTemplateNode* TemplateNode = Cast<URigVMTemplateNode>(SubNode))
 			{
-				if (!TemplateNode->IsSingleton() && !TemplateNode->IsA<URigVMLibraryNode>())
+				if (!TemplateNode->IsSingleton())
 				{
 					FRigVMControllerGraphGuard GraphGuard(this, TemplateNode->GetGraph(), false);
 					TemplateNode->InitializeFilteredPermutations();
@@ -7926,6 +7926,7 @@ FString URigVMController::AddAggregatePin(URigVMNode* InNode, const FString& InP
 		return FString();
 	}
 
+	RecomputeAllTemplateFilteredPermutations(bSetupUndoRedo);
 	UpdateTemplateNodePinTypes(AggregateNode, bSetupUndoRedo);
 
 	if (bSetupUndoRedo)
@@ -15559,7 +15560,6 @@ bool URigVMController::FullyResolveTemplateNode(URigVMTemplateNode* InNode, int3
 
 	const FRigVMFunction* ResolvedFunction = Template->GetPermutation(InputPermutation);
 	const TArray<int32> PermutationIndices = {InputPermutation};
-	//InNode->FilteredPermutations = PermutationIndices;
 
 	// find all existing pins that we may need to change
 	TArray<FRigVMTemplateArgument> MissingPins;
@@ -15670,6 +15670,25 @@ bool URigVMController::FullyResolveTemplateNode(URigVMTemplateNode* InNode, int3
 	// exit out early if there's nothing to do
 	if(PinTypesToChange.IsEmpty() && MissingPins.IsEmpty() && PinsToRemove.IsEmpty())
 	{
+		if (URigVMCollapseNode* CollapseNode = Cast<URigVMCollapseNode>(InNode))
+		{
+			{
+				FRigVMControllerGraphGuard GraphGuard(this, CollapseNode->GetContainedGraph());
+				TGuardValue<bool> GuardEditable(CollapseNode->GetContainedGraph()->bEditable, true);
+				TGuardValue<bool> GuardRecomputeOuter(bSuspendRecomputingOuterTemplateFilters, true);
+				TArray<URigVMTemplateNode*> InterfaceNodes = {CollapseNode->GetEntryNode(), CollapseNode->GetReturnNode()};
+				for (URigVMTemplateNode* InterfaceNode : InterfaceNodes)
+				{
+					if (InterfaceNode)
+					{
+						UpdateFilteredPermutations(InterfaceNode, {CollapseNode->FilteredPermutations}, bSetupUndoRedo);
+						break;
+					}
+				}
+				RecomputeAllTemplateFilteredPermutations(bSetupUndoRedo);
+			}
+			UpdateFilteredPermutations(CollapseNode, {0}, bSetupUndoRedo);
+		}
 		ResolveTemplateNodeMetaData(InNode, bSetupUndoRedo);
 		return true;
 	}
@@ -16070,15 +16089,15 @@ bool URigVMController::PrepareTemplatePinForType(URigVMPin* InPin, const TArray<
 				{
 					if (Pin->IsWildCard())
 					{
-						if (!TemplateNode->FilteredSupportsType(Pin, OldPreferredTypes[i].TypeIndex))
+						if (!TemplateNode->FilteredSupportsType(Pin, OldPreferredTypes[i].GetTypeIndex()))
 						{
 							continue;
 						}
 					}
 				
-					if (Pin->IsWildCard() || Pin->GetTypeIndex() == OldPreferredTypes[i].TypeIndex)
+					if (Pin->IsWildCard() || Pin->GetTypeIndex() == OldPreferredTypes[i].GetTypeIndex())
 					{
-						ResolveWildCardPin(Pin, OldPreferredTypes[i].TypeIndex, bSetupUndoRedo);
+						ResolveWildCardPin(Pin, OldPreferredTypes[i].GetTypeIndex(), bSetupUndoRedo);
 					}
 				}
 			}
@@ -16120,15 +16139,15 @@ bool URigVMController::PrepareTemplatePinForType(URigVMPin* InPin, const TArray<
 				{
 					if (Pin->IsWildCard())
 					{
-						if (!TemplateNode->FilteredSupportsType(Pin, OldPreferredTypes[i].TypeIndex))
+						if (!TemplateNode->FilteredSupportsType(Pin, OldPreferredTypes[i].GetTypeIndex()))
 						{
 							continue;
 						}
 					}
 				
-					if (Pin->IsWildCard() || Pin->GetTypeIndex() == OldPreferredTypes[i].TypeIndex)
+					if (Pin->IsWildCard() || Pin->GetTypeIndex() == OldPreferredTypes[i].GetTypeIndex())
 					{
-						ResolveWildCardPin(Pin, OldPreferredTypes[i].TypeIndex, bSetupUndoRedo);
+						ResolveWildCardPin(Pin, OldPreferredTypes[i].GetTypeIndex(), bSetupUndoRedo);
 					}
 				}
 			}
@@ -16443,6 +16462,36 @@ bool URigVMController::UpdateFilteredPermutations(URigVMPin* InPin, const TArray
 	return false;
 }
 
+bool URigVMController::UpdateFilteredPermutations(URigVMTemplateNode* InNode, const TArray<int32>& InPermutations, bool bSetupUndoRedo)
+{
+	TArray<int32> OldPermutations;
+	if (bSetupUndoRedo)
+	{
+		OldPermutations = InNode->GetFilteredPermutationsIndices();
+	}
+
+	InNode->FilteredPermutations = InPermutations;
+
+	if (bSetupUndoRedo)
+	{
+		FRigVMSetTemplateFilteredPermutationsAction Action(InNode, OldPermutations);
+		ActionStack->BeginAction(Action);
+		ActionStack->AddAction(Action);
+		ActionStack->EndAction(Action);
+	}
+
+	// If updating the filtered permutations of entry or return nodes, update the library template
+	if (InNode->IsA<URigVMFunctionEntryNode>() || InNode->IsA<URigVMFunctionReturnNode>())
+	{
+		if (URigVMLibraryNode* LibraryNode = InNode->GetTypedOuter<URigVMLibraryNode>())
+		{
+			UpdateLibraryTemplate(LibraryNode, bSetupUndoRedo);
+		}
+	}
+
+	return true;
+}
+
 bool URigVMController::UpdateTemplateNodePinTypes(URigVMTemplateNode* InNode, bool bSetupUndoRedo, bool bInitializeDefaultValue)
 {
 	check(InNode->GetGraph() == GetGraph());
@@ -16461,8 +16510,8 @@ bool URigVMController::UpdateTemplateNodePinTypes(URigVMTemplateNode* InNode, bo
 		PreferredPermutations = InNode->FindPermutationsForTypes(InNode->PreferredPermutationPairs, true);
 		FilteredPermutations = FilteredPermutations.FilterByPredicate([&](const int32& OtherPermutation) { return PreferredPermutations.Contains(OtherPermutation); });
 	}
-	int32 ReducedToPermutation = INDEX_NONE;
-
+	InNode->ResolvedPermutation = FilteredPermutations.Num() == 1 ? FilteredPermutations[0] : INDEX_NONE;
+	
 	const FRigVMTemplate* Template = InNode->GetTemplate();
 
 	// First unresolve any pins that need unresolving, then resolve everything else
@@ -16501,7 +16550,7 @@ bool URigVMController::UpdateTemplateNodePinTypes(URigVMTemplateNode* InNode, bo
 					{
 						if (Pin->GetFName() == InNode->PreferredPermutationPairs[i].Argument)
 						{
-							Types = {InNode->PreferredPermutationPairs[i].TypeIndex};
+							Types = {InNode->PreferredPermutationPairs[i].GetTypeIndex()};
 							break;
 						}
 					}
@@ -16517,18 +16566,18 @@ bool URigVMController::UpdateTemplateNodePinTypes(URigVMTemplateNode* InNode, bo
 			if (Types.Num() > 1)
 			{
 				TRigVMTypeIndex PreferredType = Pin->GetTypeIndex();
-				if (ReducedToPermutation != INDEX_NONE)
+				if (InNode->ResolvedPermutation != INDEX_NONE)
 				{
 					if (const FRigVMTemplateArgument* Argument = Template->FindArgument(Pin->GetFName()))
 					{
-						PreferredType = Argument->TypeIndices[ReducedToPermutation];
+						PreferredType = Argument->TypeIndices[InNode->ResolvedPermutation];
 					}
 				}
 				InNode->TryReduceTypesToSingle(Types, PreferredType);
 				if (Types.Num() == 1)
 				{
 					// Make sure the other pins use the same permutation
-					ReducedToPermutation = *FilteredPermutations.FindByPredicate([&](const int32& Permutation)
+					InNode->ResolvedPermutation = *FilteredPermutations.FindByPredicate([&](const int32& Permutation)
 					{
 						if (const FRigVMTemplateArgument* Argument = Template->FindArgument(Pin->GetFName()))
 						{
@@ -16701,7 +16750,7 @@ bool URigVMController::PropagateTemplateFilteredTypes(URigVMTemplateNode* InNode
 	return true;
 }
 
-bool URigVMController::AddPreferredType(URigVMTemplateNode* InNode, const FName& InPinName, const int32 InPreferredTypeIndex, bool bSetupUndoRedo)
+bool URigVMController::AddPreferredType(URigVMTemplateNode* InNode, const FName& InPinName, const TRigVMTypeIndex& InPreferredTypeIndex, bool bSetupUndoRedo)
 {
 	check(InNode->FindPin(InPinName.ToString()));
 
@@ -16915,7 +16964,7 @@ bool URigVMController::UpdateLibraryTemplate(URigVMLibraryNode* LibraryNode, boo
 				const FRigVMTemplatePreferredType& PreferredType = InterfaceNode->PreferredPermutationPairs[i];
 				if (!ArgumentsToRemove.Contains(PreferredType.Argument) && !ArgPreferredType.Contains(PreferredType.Argument))
 				{
-					ArgPreferredType.Add(PreferredType.Argument, PreferredType.TypeIndex);
+					ArgPreferredType.Add(PreferredType.Argument, PreferredType.GetTypeIndex());
 				}			
 			}
 		}
@@ -17052,13 +17101,20 @@ bool URigVMController::UpdateLibraryTemplate(URigVMLibraryNode* LibraryNode, boo
 
 	NewTemplate.ComputeNotationFromArguments(LibraryNode->GetName());
 
-	if (bSetupUndoRedo)
+	// Set the new template
+	// The filtered permutations of the collapsed node are not valid anymore, lets reset them
 	{
-		FRigVMControllerGraphGuard GraphGuard(this, LibraryNode->GetTypedOuter<URigVMGraph>(), bSetupUndoRedo);
-		ActionStack->AddAction(FRigVMSetLibraryTemplateAction(LibraryNode, NewTemplate));		
+		TArray<int32> OldPermutations = LibraryNode->FilteredPermutations;
+		LibraryNode->FilteredPermutations.Reset();
+		if (bSetupUndoRedo)
+		{
+			FRigVMControllerGraphGuard GraphGuard(this, LibraryNode->GetTypedOuter<URigVMGraph>(), bSetupUndoRedo);
+			ActionStack->AddAction(FRigVMSetTemplateFilteredPermutationsAction(LibraryNode, OldPermutations));
+			ActionStack->AddAction(FRigVMSetLibraryTemplateAction(LibraryNode, NewTemplate));		
+		}
+		LibraryNode->Template = NewTemplate;
+		Notify(ERigVMGraphNotifType::NodeDescriptionChanged, LibraryNode);
 	}
-	LibraryNode->Template = NewTemplate;
-	Notify(ERigVMGraphNotifType::NodeDescriptionChanged, LibraryNode);
 
 	// Update Entry and Return filtered permutations to all permutations
 	{
@@ -17476,7 +17532,7 @@ void URigVMController::InitializeAllTemplateFiltersInGraph(bool bSetupUndoRedo,
 							FRigVMTemplatePreferredType& PreferredType = InterfaceNode->PreferredPermutationPairs[i];
 							if (PreferredType.Argument == Argument.Name)
 							{
-								PreferredTypeIndex = PreferredType.TypeIndex;
+								PreferredTypeIndex = PreferredType.GetTypeIndex();
 								bHasPreferredType = true;
 								break;
 							}
@@ -17553,7 +17609,7 @@ void URigVMController::InitializeAllTemplateFiltersInGraph(bool bSetupUndoRedo,
 				{
 					const FRigVMTemplatePreferredType& PreferredType = TemplateNode->PreferredPermutationPairs[i];
 					
-					if (TemplateNode->GetTemplate()->FindArgument(PreferredType.Argument) && TemplateNode->GetTemplate()->ArgumentSupportsTypeIndex(PreferredType.Argument, PreferredType.TypeIndex))
+					if (TemplateNode->GetTemplate()->FindArgument(PreferredType.Argument) && TemplateNode->GetTemplate()->ArgumentSupportsTypeIndex(PreferredType.Argument, PreferredType.GetTypeIndex()))
 					{
 						NewPreferredTypes.Add(PreferredType);
 					}
@@ -17961,7 +18017,7 @@ bool URigVMController::ChangePinType(URigVMPin* InPin, TRigVMTypeIndex InTypeInd
 
 			if (URigVMTemplateNode* TemplateNode = Cast<URigVMTemplateNode>(InPin->GetNode()))
 			{
-				if (TemplateNode->IsA<URigVMUnitNode>() || TemplateNode->IsA<URigVMDispatchNode>())
+				if (!TemplateNode->IsA<URigVMFunctionEntryNode>() && !TemplateNode->IsA<URigVMFunctionReturnNode>())
 				{
 					// Figure out the permutation from the pin types. During undo, the filtered permutations are not
 					// reliable as to which permutation we are resolving to.
diff --git a/Engine/Source/Developer/RigVMDeveloper/Public/RigVMModel/Nodes/RigVMTemplateNode.h b/Engine/Source/Developer/RigVMDeveloper/Public/RigVMModel/Nodes/RigVMTemplateNode.h
--- a/Engine/Source/Developer/RigVMDeveloper/Public/RigVMModel/Nodes/RigVMTemplateNode.h
+++ b/Engine/Source/Developer/RigVMDeveloper/Public/RigVMModel/Nodes/RigVMTemplateNode.h
@@ -190,11 +190,10 @@ protected:
 	TArray<FRigVMTemplatePreferredType> PreferredPermutationPairs;
 
 	TArray<int32> FilteredPermutations;
-	TMap<TPair<FName,TRigVMTypeIndex>, TPair<bool, TRigVMTypeIndex>> SupportedTypesCache;
 
 	mutable const FRigVMTemplate* CachedTemplate;
 	mutable const FRigVMFunction* CachedFunction;
-	mutable TArray<int32> ResolvedPermutations;
+	mutable int32 ResolvedPermutation;
 
 	friend class URigVMController;
 	friend class UControlRigBlueprint;
diff --git a/Engine/Source/Developer/RigVMDeveloper/Public/RigVMModel/RigVMController.h b/Engine/Source/Developer/RigVMDeveloper/Public/RigVMModel/RigVMController.h
--- a/Engine/Source/Developer/RigVMDeveloper/Public/RigVMModel/RigVMController.h
+++ b/Engine/Source/Developer/RigVMDeveloper/Public/RigVMModel/RigVMController.h
@@ -1070,6 +1070,7 @@ private:
 	// Updates the permutations allowed without having to break any links
 	bool UpdateFilteredPermutations(URigVMPin* InPin, URigVMPin* InLinkedPin, bool bSetupUndoRedo);
 	bool UpdateFilteredPermutations(URigVMPin* InPin, const TArray<TRigVMTypeIndex>& InTypeIndices, bool bSetupUndoRedo);
+	bool UpdateFilteredPermutations(URigVMTemplateNode* InNode, const TArray<int32>& InPermutations, bool bSetupUndoRedo);
 
 	// Changes Pin types if filtered types of a pin are unique
 	bool UpdateTemplateNodePinTypes(URigVMTemplateNode* InNode, bool bSetupUndoRedo, bool bInitializeDefaultValue = true);
@@ -1080,7 +1081,7 @@ private:
 
 	// Adds a preferred type for the pin
 	// Returns false if the pin already has a different type
-	bool AddPreferredType(URigVMTemplateNode* InNode, const FName& InPinName, const int32 InPreferredTypeIndex, bool bSetupUndoRedo);
+	bool AddPreferredType(URigVMTemplateNode* InNode, const FName& InPinName, const TRigVMTypeIndex& InPreferredTypeIndex, bool bSetupUndoRedo);
 
 	// Removes preferred type
 	// Returns true if the preferred type was found and removed
