Control Rig: Fix crash when undo/redoing paste of library node with bounded variables

diff --git a/Engine/Plugins/Animation/ControlRig/Source/ControlRigEditor/Private/Graph/ControlRigGraphPanelPinFactory.cpp b/Engine/Plugins/Animation/ControlRig/Source/ControlRigEditor/Private/Graph/ControlRigGraphPanelPinFactory.cpp
--- a/Engine/Plugins/Animation/ControlRig/Source/ControlRigEditor/Private/Graph/ControlRigGraphPanelPinFactory.cpp
+++ b/Engine/Plugins/Animation/ControlRig/Source/ControlRigEditor/Private/Graph/ControlRigGraphPanelPinFactory.cpp
@@ -39,7 +39,7 @@ TSharedPtr<SGraphPin> FControlRigGraphPanelPinFactory::CreatePin(UEdGraphPin* In
 			{
 				if (ModelPin->IsBoundToVariable())
 				{
-					if (UControlRigBlueprint* Blueprint = Cast<UControlRigBlueprint>(RigGraph->GetOuter()))
+					if (UControlRigBlueprint* Blueprint = RigGraph->GetTypedOuter<UControlRigBlueprint>())
 					{
 						return SNew(SControlRigGraphPinVariableBinding, InPin)
 							.ModelPins({ModelPin})
diff --git a/Engine/Source/Developer/RigVMDeveloper/Private/RigVMModel/RigVMController.cpp b/Engine/Source/Developer/RigVMDeveloper/Private/RigVMModel/RigVMController.cpp
--- a/Engine/Source/Developer/RigVMDeveloper/Private/RigVMModel/RigVMController.cpp
+++ b/Engine/Source/Developer/RigVMDeveloper/Private/RigVMModel/RigVMController.cpp
@@ -3450,7 +3450,7 @@ TArray<FName> URigVMController::ImportNodesFromText(const FString& InText, bool
 
 		for(int32 SubNodeIndex=0; SubNodeIndex < SubNodes.Num(); SubNodeIndex++)
 		{
-			if(URigVMTemplateNode* TemplateNode = Cast<URigVMTemplateNode>(CreatedNode))
+			if(URigVMTemplateNode* TemplateNode = Cast<URigVMTemplateNode>(SubNodes[SubNodeIndex]))
 			{
 				// make sure the update preferred type pairs
 				// from string to index - since indices may
@@ -3484,12 +3484,12 @@ TArray<FName> URigVMController::ImportNodesFromText(const FString& InText, bool
 				}
 				
 				{
-					FRigVMControllerGraphGuard GraphGuard(this, OuterNode->GetGraph(), bSetupUndoRedo);
-					UpdateLibraryTemplate(OuterNode, bSetupUndoRedo);
+					FRigVMControllerGraphGuard GraphGuard(this, OuterNode->GetGraph(), false);
+					UpdateLibraryTemplate(OuterNode, false);
 				}
 				{
-					FRigVMControllerGraphGuard GraphGuard(this, OuterNode->GetContainedGraph(), bSetupUndoRedo);
-					RecomputeAllTemplateFilteredPermutations(bSetupUndoRedo);
+					FRigVMControllerGraphGuard GraphGuard(this, OuterNode->GetContainedGraph(), false);
+					RecomputeAllTemplateFilteredPermutations(false);
 				}
 				
 			};
@@ -3648,8 +3648,8 @@ TArray<FName> URigVMController::ImportNodesFromText(const FString& InText, bool
 							VariableNodeName = TargetPin->GetBoundVariableNode()->GetName();
 							BindingPath = TargetPin->GetBoundVariablePath();
 
-							// The current situation is that the outter pin has an injection info, and the injected node exists
-							// but the injected node is not linked to the outter pin. BreakAllLinks will try to unbind the outter pin,
+							// The current situation is that the outer pin has an injection info, and the injected node exists
+							// but the injected node is not linked to the outer pin. BreakAllLinks will try to unbind the outer pin,
 							// for that to be successful, the binding needs to be complete
 							// Connect it so that the unbound is successful
 							if (!SourcePin->IsLinkedTo(TargetPin))
@@ -3658,6 +3658,11 @@ TArray<FName> URigVMController::ImportNodesFromText(const FString& InText, bool
 								SourcePin->Links.Add(CreatedLink);
 								TargetPin->Links.Add(CreatedLink);
 							}
+
+							if (TargetPin->IsBoundToVariable())
+							{
+								UnbindPinFromVariable(TargetPin, bSetupUndoRedo);
+							}
 						}
 						
 						BreakAllLinksRecursive(TargetPin, true, true, bSetupUndoRedo);
diff --git a/Engine/Source/Developer/RigVMDeveloper/Private/RigVMModel/RigVMControllerActions.cpp b/Engine/Source/Developer/RigVMDeveloper/Private/RigVMModel/RigVMControllerActions.cpp
--- a/Engine/Source/Developer/RigVMDeveloper/Private/RigVMModel/RigVMControllerActions.cpp
+++ b/Engine/Source/Developer/RigVMDeveloper/Private/RigVMModel/RigVMControllerActions.cpp
@@ -840,12 +840,20 @@ FRigVMRemoveNodeAction::FRigVMRemoveNodeAction(URigVMNode* InNode, URigVMControl
 		{
 			if (URigVMVariableNode* VariableNode = Cast<URigVMVariableNode>(Pin->GetInjectedNodes()[0]->Node))
 			{
-				FRigVMAddVariableNodeAction AddVariableNodeAction(VariableNode);
-				FRigVMAddLinkAction AddLinkAction(VariableNode->GetValuePin(), Pin);
-				FRigVMInjectNodeIntoPinAction InjectAction(Pin->GetInjectedNodes()[0]);
-				InverseAction.AddAction(AddVariableNodeAction, InController);
-				InverseAction.AddAction(AddLinkAction, InController);
-				InverseAction.AddAction(InjectAction, InController);
+				if (!InNode->IsA<URigVMLibraryNode>())
+				{
+					FRigVMAddVariableNodeAction AddVariableNodeAction(VariableNode);
+					FRigVMAddLinkAction AddLinkAction(VariableNode->GetValuePin(), Pin);
+					FRigVMInjectNodeIntoPinAction InjectAction(Pin->GetInjectedNodes()[0]);
+					InverseAction.AddAction(AddVariableNodeAction, InController);
+					InverseAction.AddAction(AddLinkAction, InController);
+					InverseAction.AddAction(InjectAction, InController);
+				}
+				else
+				{
+					FRigVMAddLinkAction AddLinkAction(VariableNode->GetValuePin(), Pin);
+					InverseAction.AddAction(AddLinkAction, InController);
+				}
 			}
 		}
 	}
@@ -1748,14 +1756,14 @@ bool FRigVMImportNodeFromTextAction::Undo(URigVMController* InController)
 	{
 		return false;
 	}
-	return InController->RemoveNodeByName(*NodePath, false);
+	return InController->RemoveNodeByName(*NodePath, false, true);
 }
 
 bool FRigVMImportNodeFromTextAction::Redo(URigVMController* InController)
 {
 #if WITH_EDITOR
 	TArray<FName> NodeNames = InController->ImportNodesFromText(ExportedText, false);
-	if (NodeNames.Num() == 1)
+	if (NodeNames.Num() > 0)
 	{
 		return FRigVMBaseAction::Redo(InController);
 	}
