Control Rig: Support watches within functions

diff --git a/Engine/Plugins/Animation/ControlRig/Source/ControlRigEditor/Private/Editor/ControlRigEditor.cpp b/Engine/Plugins/Animation/ControlRig/Source/ControlRigEditor/Private/Editor/ControlRigEditor.cpp
--- a/Engine/Plugins/Animation/ControlRig/Source/ControlRigEditor/Private/Editor/ControlRigEditor.cpp
+++ b/Engine/Plugins/Animation/ControlRig/Source/ControlRigEditor/Private/Editor/ControlRigEditor.cpp
@@ -2580,8 +2580,7 @@ bool FControlRigEditor::CanStartWatchingPin() const
 		if (URigVMPin* ModelPin = GetFocusedModel()->FindPin(Pin->GetName()))
 		{
 			return ModelPin->GetParentPin() == nullptr &&
-					!ModelPin->RequiresWatch() &&
-					!ModelPin->GetGraph()->GetRootGraph()->IsA<URigVMFunctionLibrary>();
+					!ModelPin->RequiresWatch();
 		}
 	}
 	return false;
diff --git a/Engine/Source/Developer/RigVMDeveloper/Private/RigVMCompiler/RigVMCompiler.cpp b/Engine/Source/Developer/RigVMDeveloper/Private/RigVMCompiler/RigVMCompiler.cpp
--- a/Engine/Source/Developer/RigVMDeveloper/Private/RigVMCompiler/RigVMCompiler.cpp
+++ b/Engine/Source/Developer/RigVMDeveloper/Private/RigVMCompiler/RigVMCompiler.cpp
@@ -2173,7 +2173,7 @@ FString URigVMCompiler::GetPinHashImpl(const URigVMPin* InPin, const FRigVMVarEx
 			{
 				// rely on the proxy callstack to differentiate registers
 				const FString CallStackPath = NodeProxy.GetCallstack().GetCallPath(false /* include last */);
-				if (!CallStackPath.IsEmpty())
+				if (!CallStackPath.IsEmpty() && !InPinProxy.IsValid())
 				{
 					Prefix += CallStackPath + TEXT("|");
 					bUseFullNodePath = false;
@@ -2213,7 +2213,10 @@ FString URigVMCompiler::GetPinHashImpl(const URigVMPin* InPin, const FRigVMVarEx
 FString URigVMCompiler::GetPinHash(const URigVMPin* InPin, const FRigVMVarExprAST* InVarExpr, bool bIsDebugValue, const FRigVMASTProxy& InPinProxy)
 {
 	const FString Hash = GetPinHashImpl(InPin, InVarExpr, bIsDebugValue, InPinProxy);
-	ensureMsgf(!Hash.Contains(TEXT("FunctionLibrary::")), TEXT("A library path should never be part of a pin hash %s."), *Hash);
+	if(!bIsDebugValue)
+	{
+		ensureMsgf(!Hash.Contains(TEXT("FunctionLibrary::")), TEXT("A library path should never be part of a pin hash %s."), *Hash);
+	}
 	return Hash;
 }
 
@@ -2287,8 +2290,10 @@ void URigVMCompiler::MarkDebugWatch(bool bRequired, URigVMPin* InPin, URigVM* Ou
 		
 	TArray<const FRigVMExprAST*> Expressions = InRuntimeAST->GetExpressionsForSubject(SourcePin);
 	TArray<FRigVMOperand> VisitedKeys;
-	for(const FRigVMExprAST* Expression : Expressions)
+	for(int32 ExpressionIndex=0;ExpressionIndex<Expressions.Num();ExpressionIndex++)
 	{
+		const FRigVMExprAST* Expression = Expressions[ExpressionIndex];
+		
 		check(Expression->IsA(FRigVMExprAST::EType::Var));
 		const FRigVMVarExprAST* VarExpression = Expression->To<FRigVMVarExprAST>();
 
@@ -2297,11 +2302,34 @@ void URigVMCompiler::MarkDebugWatch(bool bRequired, URigVMPin* InPin, URigVM* Ou
 			// literals don't need to be stored on the debug memory
 			if(VarExpression->IsA(FRigVMExprAST::Literal))
 			{
+				// check if there's also an IO expression for this pin
+				for(int32 ParentIndex=0;ParentIndex<VarExpression->NumParents();ParentIndex++)
+				{
+					const FRigVMExprAST* ParentExpression = VarExpression->ParentAt(ParentIndex);
+					if(ParentExpression->IsA(FRigVMExprAST::EType::Assign))
+					{
+						if(const FRigVMExprAST* GrandParentExpression = ParentExpression->GetParent())
+						{
+							if(GrandParentExpression->IsA(FRigVMExprAST::EType::Var))
+							{
+								if(GrandParentExpression->To<FRigVMVarExprAST>()->GetPin() == Pin)
+								{
+									Expressions.Add(GrandParentExpression);
+								}
+							}
+						}
+					}
+				}
 				continue;
 			}
 		}
-		
-		const FString PinHash = GetPinHash(SourcePin, VarExpression, false);
+
+		FString PinHash = GetPinHash(Pin, VarExpression, false);
+		if(!OutOperands->Contains(PinHash))
+		{
+			PinHash = GetPinHash(SourcePin, VarExpression, false);
+		}
+
 		if(const FRigVMOperand* Operand = OutOperands->Find(PinHash))
 		{
 			const FRigVMASTProxy PinProxy = FRigVMASTProxy::MakeFromUObject(Pin);
diff --git a/Engine/Source/Developer/RigVMDeveloper/Private/RigVMModel/RigVMController.cpp b/Engine/Source/Developer/RigVMDeveloper/Private/RigVMModel/RigVMController.cpp
--- a/Engine/Source/Developer/RigVMDeveloper/Private/RigVMModel/RigVMController.cpp
+++ b/Engine/Source/Developer/RigVMDeveloper/Private/RigVMModel/RigVMController.cpp
@@ -7161,12 +7161,6 @@ bool URigVMController::SetPinIsWatched(URigVMPin* InPin, bool bIsWatched, bool b
 	URigVMGraph* Graph = GetGraph();
 	check(Graph);
 
-	if (Graph->GetRootGraph()->IsA<URigVMFunctionLibrary>())
-	{
-		ReportError(TEXT("Cannot watch pins in function library graphs."));
-		return false;
-	}
-
 	FRigVMControllerCompileBracketScope CompileScope(this);
 	FRigVMSetPinWatchAction Action;
 	if (bSetupUndoRedo)
