Control Rig: Fix crash when copy pasting template node

diff --git a/Engine/Source/Developer/RigVMDeveloper/Private/RigVMModel/RigVMController.cpp b/Engine/Source/Developer/RigVMDeveloper/Private/RigVMModel/RigVMController.cpp
--- a/Engine/Source/Developer/RigVMDeveloper/Private/RigVMModel/RigVMController.cpp
+++ b/Engine/Source/Developer/RigVMDeveloper/Private/RigVMModel/RigVMController.cpp
@@ -16496,21 +16496,21 @@ bool URigVMController::UpdateTemplateNodePinTypes(URigVMTemplateNode* InNode, bo
 	check(InNode->GetGraph() == GetGraph());
 	bool bAnyTypeChanged = false;
 
-	TArray<int32> FilteredPermutations = InNode->GetFilteredPermutationsIndices();
+	TArray<int32> ResolvedPermutations = InNode->GetFilteredPermutationsIndices();
 
 	// Try to pick permutations that respects the preferred types
 	TArray<int32> PreferredPermutations = InNode->FindPermutationsForTypes(InNode->PreferredPermutationPairs, false);
-	FilteredPermutations = FilteredPermutations.FilterByPredicate([&](const int32& OtherPermutation) { return PreferredPermutations.Contains(OtherPermutation); });
-	if (FilteredPermutations.IsEmpty())
+	ResolvedPermutations = ResolvedPermutations.FilterByPredicate([&](const int32& OtherPermutation) { return PreferredPermutations.Contains(OtherPermutation); });
+	if (ResolvedPermutations.IsEmpty())
 	{
-		FilteredPermutations = InNode->GetFilteredPermutationsIndices();
+		ResolvedPermutations = InNode->GetFilteredPermutationsIndices();
 
 		// Try to pick one permutations that respects the preferred types (with casting)
 		PreferredPermutations = InNode->FindPermutationsForTypes(InNode->PreferredPermutationPairs, true);
-		FilteredPermutations = FilteredPermutations.FilterByPredicate([&](const int32& OtherPermutation) { return PreferredPermutations.Contains(OtherPermutation); });
+		ResolvedPermutations = ResolvedPermutations.FilterByPredicate([&](const int32& OtherPermutation) { return PreferredPermutations.Contains(OtherPermutation); });
 	}
-	InNode->ResolvedPermutation = FilteredPermutations.Num() == 1 ? FilteredPermutations[0] : INDEX_NONE;
-	
+	InNode->ResolvedPermutation = ResolvedPermutations.Num() == 1 ? ResolvedPermutations[0] : INDEX_NONE;
+
 	const FRigVMTemplate* Template = InNode->GetTemplate();
 
 	// First unresolve any pins that need unresolving, then resolve everything else
@@ -16531,8 +16531,8 @@ bool URigVMController::UpdateTemplateNodePinTypes(URigVMTemplateNode* InNode, bo
 			}
 		
 			TArray<TRigVMTypeIndex> Types;
-			Types.Reserve(FilteredPermutations.Num());
-			for (const int32& Permutation : FilteredPermutations)
+			Types.Reserve(ResolvedPermutations.Num());
+			for (const int32& Permutation : ResolvedPermutations)
 			{
 				if (const FRigVMTemplateArgument* Argument = Template->FindArgument(Pin->GetFName()))
 				{
@@ -16576,7 +16576,7 @@ bool URigVMController::UpdateTemplateNodePinTypes(URigVMTemplateNode* InNode, bo
 				if (Types.Num() == 1)
 				{
 					// Make sure the other pins use the same permutation
-					InNode->ResolvedPermutation = *FilteredPermutations.FindByPredicate([&](const int32& Permutation)
+					InNode->ResolvedPermutation = *ResolvedPermutations.FindByPredicate([&](const int32& Permutation)
 					{
 						if (const FRigVMTemplateArgument* Argument = Template->FindArgument(Pin->GetFName()))
 						{
@@ -16656,6 +16656,11 @@ bool URigVMController::UpdateTemplateNodePinTypes(URigVMTemplateNode* InNode, bo
 		}
 	}
 
+	if (InNode->HasWildCardPin())
+	{
+		InNode->ResolvedPermutation = INDEX_NONE;
+	}
+	
 	return bAnyTypeChanged;
 }
 
