Control Rig: Fix crash when undoing collapse function references

diff --git a/Engine/Source/Developer/RigVMDeveloper/Private/RigVMModel/RigVMController.cpp b/Engine/Source/Developer/RigVMDeveloper/Private/RigVMModel/RigVMController.cpp
--- a/Engine/Source/Developer/RigVMDeveloper/Private/RigVMModel/RigVMController.cpp
+++ b/Engine/Source/Developer/RigVMDeveloper/Private/RigVMModel/RigVMController.cpp
@@ -4395,8 +4395,8 @@ URigVMCollapseNode* URigVMController::CollapseNodes(const TArray<URigVMNode*>& I
 				{
 					if (!Pin->IsWildCard())
 					{
-						AddPreferredType(InterfaceNode, Pin->GetFName(), Pin->GetTypeIndex(), bSetupUndoRedo);
-						AddArgumentForPin(Pin, nullptr, bSetupUndoRedo);
+						AddPreferredType(InterfaceNode, Pin->GetFName(), Pin->GetTypeIndex(), false);
+						AddArgumentForPin(Pin, nullptr, false);
 					}
 				}
 			}
