Control Rig: Fix new name of the Function does not display on Function's node in the Rig Graph after renaming

diff --git a/Engine/Source/Developer/RigVMDeveloper/Private/RigVMModel/RigVMController.cpp b/Engine/Source/Developer/RigVMDeveloper/Private/RigVMModel/RigVMController.cpp
--- a/Engine/Source/Developer/RigVMDeveloper/Private/RigVMModel/RigVMController.cpp
+++ b/Engine/Source/Developer/RigVMDeveloper/Private/RigVMModel/RigVMController.cpp
@@ -6116,6 +6116,13 @@ bool URigVMController::RenameNode(URigVMNode* InNode, const FName& InNewName, bo
 		ActionStack->BeginAction(Action);
 	}
 
+	if (URigVMCollapseNode* CollapseNode = Cast<URigVMCollapseNode>(InNode))
+	{
+		const FString& OldNotation = CollapseNode->Template.Notation.ToString();
+		const FString NewNotation = ValidNewName.ToString() + OldNotation.RightChop(OldNotation.Find(TEXT("(")));
+		CollapseNode->Template.Notation = *NewNotation;
+	}
+
 	// loop over all links and remove them
 	TArray<URigVMLink*> Links = InNode->GetLinks();
 	for (URigVMLink* Link : Links)
