RigVM: Generalize all execute pins to use the RigVMExecuteContext

diff --git a/Engine/Plugins/Animation/ControlRig/Source/ControlRig/Public/Units/RigDispatchFactory.h b/Engine/Plugins/Animation/ControlRig/Source/ControlRig/Public/Units/RigDispatchFactory.h
--- a/Engine/Plugins/Animation/ControlRig/Source/ControlRig/Public/Units/RigDispatchFactory.h
+++ b/Engine/Plugins/Animation/ControlRig/Source/ControlRig/Public/Units/RigDispatchFactory.h
@@ -14,6 +14,11 @@ struct CONTROLRIG_API FRigDispatchFactory : public FRigVMDispatchFactory
 {
 	GENERATED_BODY()
 
+	FORCEINLINE virtual UScriptStruct* GetExecuteContextStruct() const override
+	{
+		return FControlRigExecuteContext::StaticStruct();
+	}
+
 	FORCEINLINE virtual void RegisterDependencyTypes() const override
 	{
 		FRigVMRegistry::Get().FindOrAddType(FControlRigExecuteContext::StaticStruct());
diff --git a/Engine/Source/Developer/RigVMDeveloper/Private/RigVMCompiler/RigVMCompiler.cpp b/Engine/Source/Developer/RigVMDeveloper/Private/RigVMCompiler/RigVMCompiler.cpp
--- a/Engine/Source/Developer/RigVMDeveloper/Private/RigVMCompiler/RigVMCompiler.cpp
+++ b/Engine/Source/Developer/RigVMDeveloper/Private/RigVMCompiler/RigVMCompiler.cpp
@@ -69,8 +69,21 @@ FRigVMOperand FRigVMCompilerWorkData::AddProperty(
 	const FString& InDefaultValue)
 {
 	check(bSetupMemory);
+
+	const FRigVMRegistry& Registry = FRigVMRegistry::Get();
+	FRigVMTemplateArgumentType ArgumentType(*InCPPType, InCPPTypeObject);
+	const TRigVMTypeIndex TypeIndex = Registry.GetTypeIndex(ArgumentType);
+	if(TypeIndex != INDEX_NONE)
+	{
+		// for execute pins we should use the graph's default execute context struct
+		if(Registry.IsExecuteType(TypeIndex))
+		{
+			ensure(!ArgumentType.IsArray());
+			ArgumentType = FRigVMTemplateArgumentType(ExecuteContextStruct);
+		}
+	}
 	
-	FRigVMPropertyDescription Description(InName, InCPPType, InCPPTypeObject, InDefaultValue);
+	FRigVMPropertyDescription Description(InName, ArgumentType.CPPType.ToString(), ArgumentType.CPPTypeObject, InDefaultValue);
 
 	TArray<FRigVMPropertyDescription>& PropertyArray = PropertyDescriptions.FindOrAdd(InMemoryType);
 	const int32 PropertyIndex = PropertyArray.Add(Description);
@@ -366,6 +379,7 @@ bool URigVMCompiler::Compile(TArray<URigVMGraph*> InGraphs, URigVMController* In
 	ensure(WorkData.AST.IsValid());
 
 	WorkData.VM = OutVM;
+	WorkData.ExecuteContextStruct = InGraphs[0]->GetExecuteContextStruct();
 	WorkData.PinPathToOperand = OutOperands;
 	WorkData.RigVMUserData = UserData[0];
 	WorkData.bSetupMemory = true;
@@ -775,12 +789,42 @@ int32 URigVMCompiler::TraverseCallExtern(const FRigVMCallExternExprAST* InExpr,
 		return INDEX_NONE;
 	}
 
+	auto CheckExecuteStruct = [this, &WorkData](URigVMNode* Subject, const UScriptStruct* ExecuteStruct) -> bool
+	{
+		if(ExecuteStruct->IsChildOf(FRigVMExecuteContext::StaticStruct()))
+		{
+			// top level expected execute struct is provided by the graph
+			const UScriptStruct* SpecializedExecuteStruct = WorkData.ExecuteContextStruct;
+			if(!SpecializedExecuteStruct->IsChildOf(ExecuteStruct))
+			{
+				static constexpr TCHAR UnknownExecuteContextMessage[] = TEXT("Node @@ uses an unexpected execute type '%s'. This graph uses '%s'.");
+				Settings.Report(EMessageSeverity::Error, Subject, FString::Printf(
+					UnknownExecuteContextMessage, *ExecuteStruct->GetStructCPPName(), *SpecializedExecuteStruct->GetStructCPPName()));
+				return false;
+			}
+		}
+		return true;
+	};
+	
 	if(UnitNode)
 	{
-		if(UnitNode->GetScriptStruct() == nullptr)
+		const UScriptStruct* ScriptStruct = UnitNode->GetScriptStruct(); 
+		if(ScriptStruct == nullptr)
 		{
 			return INDEX_NONE;
 		}
+
+		// check execute pins for compatibility
+		for (TFieldIterator<FProperty> It(ScriptStruct); It; ++It)
+		{
+			if(const FStructProperty* StructProperty = CastField<FStructProperty>(*It))
+			{
+				if(!CheckExecuteStruct(UnitNode, StructProperty->Struct))
+				{
+					return INDEX_NONE;
+				}
+			}
+		}
 	}
 
 	if(DispatchNode)
@@ -789,6 +833,12 @@ int32 URigVMCompiler::TraverseCallExtern(const FRigVMCallExternExprAST* InExpr,
 		{
 			return INDEX_NONE;
 		}
+
+		// check execute pins for compatibility
+		if(!CheckExecuteStruct(DispatchNode, DispatchNode->GetFactory()->GetExecuteContextStruct()))
+		{
+			return INDEX_NONE;
+		}
 	}
 
 	int32 InstructionIndex = INDEX_NONE;
diff --git a/Engine/Source/Developer/RigVMDeveloper/Private/RigVMModel/RigVMController.cpp b/Engine/Source/Developer/RigVMDeveloper/Private/RigVMModel/RigVMController.cpp
--- a/Engine/Source/Developer/RigVMDeveloper/Private/RigVMModel/RigVMController.cpp
+++ b/Engine/Source/Developer/RigVMDeveloper/Private/RigVMModel/RigVMController.cpp
@@ -1185,11 +1185,7 @@ URigVMVariableNode* URigVMController::AddVariableNode(const FName& InVariableNam
 
 	if (!bIsGetter)
 	{
-		UScriptStruct* ExecuteContextStruct = Graph->GetExecuteContextStruct();
-		URigVMPin* ExecutePin = NewObject<URigVMPin>(Node, FRigVMStruct::ExecuteContextName);
-		ExecutePin->CPPType = ExecuteContextStruct->GetStructCPPName();
-		ExecutePin->CPPTypeObject = ExecuteContextStruct;
-		ExecutePin->CPPTypeObjectPath = *ExecutePin->CPPTypeObject->GetPathName();
+		URigVMPin* ExecutePin = MakeExecutePin(Node, FRigVMStruct::ExecuteContextName);
 		ExecutePin->Direction = ERigVMPinDirection::IO;
 		AddNodePin(Node, ExecutePin);
 	}
@@ -11139,7 +11135,7 @@ URigVMLibraryNode* URigVMController::AddFunctionToLibrary(const FName& InFunctio
 	
 	if (bMutable)
 	{
-		UScriptStruct* ExecuteContextStruct = Graph->GetExecuteContextStruct();
+		const UScriptStruct* ExecuteContextStruct = FRigVMExecuteContext::StaticStruct();
 		
 		FRigVMControllerGraphGuard GraphGuard(this, CollapseNode->ContainedGraph, bSetupUndoRedo);
 		AddExposedPin(FRigVMStruct::ExecuteContextName,
@@ -12251,12 +12247,7 @@ URigVMBranchNode* URigVMController::AddBranchNode(const FVector2D& InPosition, c
 	URigVMBranchNode* Node = NewObject<URigVMBranchNode>(Graph, *Name);
 	Node->Position = InPosition;
 
-	UScriptStruct* ExecuteContextStruct = Graph->GetExecuteContextStruct();
-	URigVMPin* ExecutePin = NewObject<URigVMPin>(Node, FRigVMStruct::ExecuteContextName);
-	ExecutePin->DisplayName = FRigVMStruct::ExecuteName;
-	ExecutePin->CPPType = FString::Printf(TEXT("F%s"), *ExecuteContextStruct->GetName());
-	ExecutePin->CPPTypeObject = ExecuteContextStruct;
-	ExecutePin->CPPTypeObjectPath = *ExecutePin->CPPTypeObject->GetPathName();
+	URigVMPin* ExecutePin = MakeExecutePin(Node, FRigVMStruct::ExecuteContextName);
 	ExecutePin->Direction = ERigVMPinDirection::Input;
 	AddNodePin(Node, ExecutePin);
 
@@ -12624,6 +12615,8 @@ URigVMTemplateNode* URigVMController::AddTemplateNode(const FName& InNotation, c
 	FRigVMTemplate::FTypeMap Types;
 	Template->FullyResolve(Types, PermutationIndex);
 	Node->InitializeFilteredPermutations();
+
+	FRigVMRegistry& Registry = FRigVMRegistry::Get();
 	
 	for (int32 ArgIndex = 0; ArgIndex < Template->NumArguments(); ArgIndex++)
 	{
@@ -12631,7 +12624,8 @@ URigVMTemplateNode* URigVMController::AddTemplateNode(const FName& InNotation, c
 
 		URigVMPin* Pin = NewObject<URigVMPin>(Node, Arg->GetName());
 		const TRigVMTypeIndex& TypeIndex = Types.FindChecked(Arg->GetName());
-		const FRigVMTemplateArgumentType Type = FRigVMRegistry::Get().GetType(TypeIndex);
+		const FRigVMTemplateArgumentType Type = Registry.GetType(TypeIndex);
+		
 		Pin->CPPType = Type.CPPType.ToString();
 		Pin->CPPTypeObject = Type.CPPTypeObject;
 		if (Pin->CPPTypeObject)
@@ -12639,6 +12633,8 @@ URigVMTemplateNode* URigVMController::AddTemplateNode(const FName& InNotation, c
 			Pin->CPPTypeObjectPath = *Pin->CPPTypeObject->GetPathName();
 		}
 		Pin->Direction = Arg->GetDirection();
+		Pin->LastKnownTypeIndex = TypeIndex;
+		Pin->LastKnownCPPType = Pin->CPPType;
 
 		AddNodePin(Node, Pin);
 
@@ -12938,7 +12934,7 @@ URigVMArrayNode* URigVMController::AddArrayNode(ERigVMOpCode InOpCode, const FSt
 		static URigVMPin* AddExecutePin(URigVMController* InController, URigVMNode* InNode, ERigVMPinDirection InDirection = ERigVMPinDirection::IO, const FName& InName = NAME_None)
 		{
 			const FName PinName = InName.IsNone() ? FRigVMStruct::ExecuteContextName : InName;
-			UScriptStruct* ExecuteContextStruct = InController->GetGraph()->GetExecuteContextStruct();
+			UScriptStruct* ExecuteContextStruct = FRigVMExecuteContext::StaticStruct();
 			URigVMPin* Pin = AddPin(InController, InNode, PinName, InDirection, false, FString::Printf(TEXT("F%s"), *ExecuteContextStruct->GetName()), ExecuteContextStruct);
 			if(PinName == FRigVMStruct::ExecuteContextName)
 			{
@@ -13197,11 +13193,7 @@ URigVMInvokeEntryNode* URigVMController::AddInvokeEntryNode(const FName& InEntry
 	URigVMInvokeEntryNode* Node = NewObject<URigVMInvokeEntryNode>(Graph, *Name);
 	Node->Position = InPosition;
 
-	UScriptStruct* ExecuteContextStruct = Graph->GetExecuteContextStruct();
-	URigVMPin* ExecutePin = NewObject<URigVMPin>(Node, FRigVMStruct::ExecuteContextName);
-	ExecutePin->CPPType = ExecuteContextStruct->GetStructCPPName();
-	ExecutePin->CPPTypeObject = ExecuteContextStruct;
-	ExecutePin->CPPTypeObjectPath = *ExecutePin->CPPTypeObject->GetPathName();
+	URigVMPin* ExecutePin = MakeExecutePin(Node, FRigVMStruct::ExecuteContextName);
 	ExecutePin->Direction = ERigVMPinDirection::IO;
 	AddNodePin(Node, ExecutePin);
 
@@ -13954,6 +13946,11 @@ void URigVMController::ConfigurePinFromProperty(FProperty* InProperty, URigVMPin
 	}
 
 	InOutPin->CPPType = RigVMTypeUtils::PostProcessCPPType(InOutPin->CPPType, InOutPin->GetCPPTypeObject());
+
+	if(InOutPin->IsExecuteContext() && InOutPin->CPPTypeObject != FRigVMExecuteContext::StaticStruct())
+	{
+		MakeExecutePin(InOutPin);
+	}
 }
 
 void URigVMController::ConfigurePinFromPin(URigVMPin* InOutPin, URigVMPin* InPin, bool bCopyDisplayName)
@@ -13972,6 +13969,11 @@ void URigVMController::ConfigurePinFromPin(URigVMPin* InOutPin, URigVMPin* InPin
 	{
 		InOutPin->SetDisplayName(InPin->GetDisplayName());
 	}
+
+	if(InOutPin->IsExecuteContext() && InOutPin->CPPTypeObject != FRigVMExecuteContext::StaticStruct())
+	{
+		MakeExecutePin(InOutPin);
+	}
 }
 
 bool URigVMController::ShouldStructBeUnfolded(const UStruct* Struct)
@@ -14561,6 +14563,15 @@ void URigVMController::RepopulatePinsOnNode(URigVMNode* InNode, bool bFollowCore
 	URigVMGraph* Graph = GetGraph();
 	check(Graph);
 
+	// step 0/3: update execute pins
+	for(URigVMPin* Pin : InNode->Pins)
+	{
+		if(Pin->IsExecuteContext())
+		{
+			MakeExecutePin(Pin);
+		}
+	}
+
 	// step 1/3: keep a record of the current state of the node's pins
 	TMap<FString, FString> RedirectedPinPaths;
 	if (bFollowCoreRedirectors)
@@ -15586,7 +15597,7 @@ bool URigVMController::FullyResolveTemplateNode(URigVMTemplateNode* InNode, int3
 			return false;
 		}
 
-		if(Pin->GetTypeIndex() != ResolvedTypeIndex)
+		if(Pin->GetTypeIndex() != ResolvedTypeIndex && ResolvedTypeIndex != RigVMTypeUtils::TypeIndex::Execute)
 		{
 			PinTypesToChange.Add(Pin, ResolvedTypeIndex);
 		}
@@ -15603,7 +15614,7 @@ bool URigVMController::FullyResolveTemplateNode(URigVMTemplateNode* InNode, int3
 				const TRigVMTypeIndex ExpectedTypeIndex = Argument.TypeIndices[InputPermutation];
 				if(URigVMPin* Pin = InNode->FindPin(Argument.GetName().ToString()))
 				{
-					if(Pin->GetTypeIndex() != ExpectedTypeIndex)
+					if(Pin->GetTypeIndex() != ExpectedTypeIndex && ExpectedTypeIndex != RigVMTypeUtils::TypeIndex::Execute)
 					{
 						PinTypesToChange.Add(Pin, ExpectedTypeIndex);
 					}
@@ -15626,7 +15637,7 @@ bool URigVMController::FullyResolveTemplateNode(URigVMTemplateNode* InNode, int3
 
 					if(URigVMPin* Pin = InNode->FindPin(It->GetFName().ToString()))
 					{
-						if(Pin->GetTypeIndex() != ExpectedTypeIndex)
+						if(Pin->GetTypeIndex() != ExpectedTypeIndex && ExpectedTypeIndex != RigVMTypeUtils::TypeIndex::Execute)
 						{
 							PinTypesToChange.Add(Pin, ExpectedTypeIndex);
 						}
@@ -16627,7 +16638,7 @@ bool URigVMController::UpdateTemplateNodePinTypes(URigVMTemplateNode* InNode, bo
 			
 				if (Pin->GetCPPType() != CPPType || Pin->GetCPPTypeObject() != CPPObjectType)
 				{
-					bAnyTypeChanged = true;
+					bAnyTypeChanged = !Pin->IsExecuteContext();
 					ChangePinType(Pin, CPPType, CPPObjectType, bSetupUndoRedo, false, false, false, bInitializeDefaultValue);
 				}
 			}
@@ -16636,7 +16647,7 @@ bool URigVMController::UpdateTemplateNodePinTypes(URigVMTemplateNode* InNode, bo
 				// Resolve
 				if (Pin->GetTypeIndex() != Types[0])
 				{
-					bAnyTypeChanged = true;
+					bAnyTypeChanged = !Pin->IsExecuteContext();
 					ChangePinType(Pin, Types[0], bSetupUndoRedo, false, false, false, bInitializeDefaultValue);
 				}
 
@@ -17743,6 +17754,11 @@ bool URigVMController::ChangePinType(URigVMPin* InPin, TRigVMTypeIndex InTypeInd
 	}
 
 	check(InPin->GetGraph() == GetGraph());
+	if(InPin->IsExecuteContext() && FRigVMRegistry::Get().IsExecuteType(InTypeIndex))
+	{
+		return false;
+	}
+	
 
 	// only allow valid pin cpp types on template nodes
 	TRigVMTypeIndex TypeIndex = InTypeIndex;
@@ -18128,6 +18144,36 @@ void URigVMController::DestroyObject(UObject* InObjectToDestroy)
 	InObjectToDestroy->MarkAsGarbage();
 }
 
+URigVMPin* URigVMController::MakeExecutePin(URigVMNode* InNode, const FName& InName)
+{
+	URigVMPin* ExecutePin = NewObject<URigVMPin>(InNode, InName);
+	ExecutePin->DisplayName = FRigVMStruct::ExecuteName;
+	MakeExecutePin(ExecutePin);
+	return ExecutePin;
+}
+
+void URigVMController::MakeExecutePin(URigVMPin* InOutPin)
+{
+	if(InOutPin->CPPTypeObject != FRigVMExecuteContext::StaticStruct())
+	{
+		const bool bIsArray = InOutPin->IsArray();
+		InOutPin->CPPType = FRigVMExecuteContext::StaticStruct()->GetStructCPPName();
+		InOutPin->CPPTypeObject = FRigVMExecuteContext::StaticStruct();
+		InOutPin->CPPTypeObjectPath = *InOutPin->CPPTypeObject->GetPathName();
+
+		if(bIsArray)
+		{
+			InOutPin->CPPType = RigVMTypeUtils::ArrayTypeFromBaseType(InOutPin->CPPType);
+			InOutPin->LastKnownTypeIndex = FRigVMRegistry::Get().GetArrayTypeFromBaseTypeIndex(RigVMTypeUtils::TypeIndex::Execute);
+		}
+		else
+		{
+			InOutPin->LastKnownTypeIndex = RigVMTypeUtils::TypeIndex::Execute;
+		}
+		InOutPin->LastKnownCPPType = InOutPin->CPPType;
+	}
+}
+
 void URigVMController::AddNodePin(URigVMNode* InNode, URigVMPin* InPin)
 {
 	ValidatePin(InPin);
@@ -18248,6 +18294,11 @@ void URigVMController::ValidatePin(URigVMPin* InPin)
 	// create a property description from the pin here as a test,
 	// since the compiler needs this
 	FRigVMPropertyDescription(InPin->GetFName(), InPin->GetCPPType(), InPin->GetCPPTypeObject(), InPin->GetDefaultValue());
+
+	if(InPin->IsExecuteContext())
+	{
+		ensure(InPin->GetCPPTypeObject() == FRigVMExecuteContext::StaticStruct());
+	}
 }
 
 void URigVMController::EnsureLocalVariableValidity()
diff --git a/Engine/Source/Developer/RigVMDeveloper/Public/RigVMCompiler/RigVMCompiler.h b/Engine/Source/Developer/RigVMDeveloper/Public/RigVMCompiler/RigVMCompiler.h
--- a/Engine/Source/Developer/RigVMDeveloper/Public/RigVMCompiler/RigVMCompiler.h
+++ b/Engine/Source/Developer/RigVMDeveloper/Public/RigVMCompiler/RigVMCompiler.h
@@ -82,6 +82,7 @@ struct RIGVMDEVELOPER_API FRigVMCompilerWorkData
 public:
 	bool bSetupMemory;
 	URigVM* VM;
+	UScriptStruct* ExecuteContextStruct;
 	FRigVMUserDataArray RigVMUserData;
 	TMap<FString, FRigVMOperand>* PinPathToOperand;
 	TMap<const FRigVMVarExprAST*, FRigVMOperand> ExprToOperand;
diff --git a/Engine/Source/Developer/RigVMDeveloper/Public/RigVMModel/RigVMController.h b/Engine/Source/Developer/RigVMDeveloper/Public/RigVMModel/RigVMController.h
--- a/Engine/Source/Developer/RigVMDeveloper/Public/RigVMModel/RigVMController.h
+++ b/Engine/Source/Developer/RigVMDeveloper/Public/RigVMModel/RigVMController.h
@@ -1105,6 +1105,8 @@ private:
 
 	bool RenameObject(UObject* InObjectToRename, const TCHAR* InNewName, UObject* InNewOuter = nullptr);
 	void DestroyObject(UObject* InObjectToDestroy);
+	static URigVMPin* MakeExecutePin(URigVMNode* InNode, const FName& InName);
+	static void MakeExecutePin(URigVMPin* InOutPin);
 	static void AddNodePin(URigVMNode* InNode, URigVMPin* InPin);
 	static void AddSubPin(URigVMPin* InParentPin, URigVMPin* InPin);
 	static bool EnsurePinValidity(URigVMPin* InPin, bool bRecursive);
diff --git a/Engine/Source/Runtime/RigVM/Private/RigVMCore/RigVMRegistry.cpp b/Engine/Source/Runtime/RigVM/Private/RigVMCore/RigVMRegistry.cpp
--- a/Engine/Source/Runtime/RigVM/Private/RigVMCore/RigVMRegistry.cpp
+++ b/Engine/Source/Runtime/RigVM/Private/RigVMCore/RigVMRegistry.cpp
@@ -187,7 +187,7 @@ void FRigVMRegistry::Reset()
 	Factories.Reset();
 }
 
-TRigVMTypeIndex FRigVMRegistry::FindOrAddType(const FRigVMTemplateArgumentType& InType)
+TRigVMTypeIndex FRigVMRegistry::FindOrAddType_Internal(const FRigVMTemplateArgumentType& InType, bool bForce)
 {
 	TRigVMTypeIndex Index = GetTypeIndex(InType);
 	if(Index == INDEX_NONE)
@@ -199,7 +199,7 @@ TRigVMTypeIndex FRigVMRegistry::FindOrAddType(const FRigVMTemplateArgumentType&
 		}
 		
 		const UObject* CPPTypeObject = ElementType.CPPTypeObject;
-		if(CPPTypeObject != nullptr)
+		if(!bForce && (CPPTypeObject != nullptr))
 		{
 			if(const UClass* Class = Cast<UClass>(CPPTypeObject))
 			{
@@ -234,9 +234,9 @@ TRigVMTypeIndex FRigVMRegistry::FindOrAddType(const FRigVMTemplateArgumentType&
 		Indices.Reserve(3);
 		for (int32 ArrayDimension=0; ArrayDimension<3; ++ArrayDimension)
 		{
-			if (bIsExecute && ArrayDimension>0)
+			if (bIsExecute && ArrayDimension > 1)
 			{
-				continue;
+				break;
 			}
 			
 			FRigVMTemplateArgumentType CurType = ElementType;
@@ -259,23 +259,29 @@ TRigVMTypeIndex FRigVMRegistry::FindOrAddType(const FRigVMTemplateArgumentType&
 			Indices.Add(Index);
 		}
 
+		Types[Indices[1]].BaseTypeIndex = Indices[0];
+		Types[Indices[0]].ArrayTypeIndex = Indices[1];
+
 		if (!bIsExecute)
 		{
 			Types[Indices[1]].BaseTypeIndex = Indices[0];
 			Types[Indices[2]].BaseTypeIndex = Indices[1];
+		}
 
-			Types[Indices[0]].ArrayTypeIndex = Indices[1];
+		if (!bIsExecute)
+		{
+			Types[Indices[2]].BaseTypeIndex = Indices[1];
 			Types[Indices[1]].ArrayTypeIndex = Indices[2];
 		}
 
 		for (int32 ArrayDimension=0; ArrayDimension<3; ++ArrayDimension)
 		{
-			if(bIsExecute && ArrayDimension > 0)
+			if(bIsExecute && ArrayDimension > 1)
 			{
 				break;
 			}
 			Index = Indices[ArrayDimension];
-			
+
 			// Add to category
 			// simple types
 			if(CPPTypeObject == nullptr)
@@ -420,10 +426,10 @@ TRigVMTypeIndex FRigVMRegistry::FindOrAddType(const FRigVMTemplateArgumentType&
 			for (TFieldIterator<FProperty> It(Struct); It; ++It)
 			{
 				FProperty* Property = *It;
-				if(IsAllowedType(Property))
+				if(bForce || IsAllowedType(Property))
 				{
 					// by creating a template argument for the child property
-					// the type will be added by calling ::FindOrAddType recursively.
+					// the type will be added by calling ::FindOrAddType_Internal recursively.
 					FRigVMTemplateArgument DummyArgument(Property);
 				}
 			}			
@@ -544,6 +550,12 @@ bool FRigVMRegistry::CanMatchTypes(TRigVMTypeIndex InTypeIndexA, TRigVMTypeIndex
 		return true;
 	}
 
+	// execute types can always be connected
+	if(IsExecuteType(InTypeIndexA) && IsExecuteType(InTypeIndexB))
+	{
+		return GetArrayDimensionsForType(InTypeIndexA) == GetArrayDimensionsForType(InTypeIndexB);
+	}
+
 	if(bAllowFloatingPointCasts)
 	{
 		// swap order since float is known to registered before double
diff --git a/Engine/Source/Runtime/RigVM/Private/RigVMCore/RigVMTemplate.cpp b/Engine/Source/Runtime/RigVM/Private/RigVMCore/RigVMTemplate.cpp
--- a/Engine/Source/Runtime/RigVM/Private/RigVMCore/RigVMTemplate.cpp
+++ b/Engine/Source/Runtime/RigVM/Private/RigVMCore/RigVMTemplate.cpp
@@ -49,10 +49,11 @@ FRigVMTemplateArgument::FRigVMTemplateArgument(FProperty* InProperty)
 	}
 	Type.CPPType = *RigVMTypeUtils::PostProcessCPPType(Type.CPPType.ToString(), Type.CPPTypeObject);
 
-	const TRigVMTypeIndex TypeIndex = FRigVMRegistry::Get().FindOrAddType(Type); 
+	const TRigVMTypeIndex TypeIndex = FRigVMRegistry::Get().FindOrAddType_Internal(Type, true); 
 
 	TypeIndices.Add(TypeIndex);
-	TypeToPermutations.Add(TypeIndex, {0});
+	EnsureValidExecuteType();
+	UpdateTypeToPermutations();
 }
 
 FRigVMTemplateArgument::FRigVMTemplateArgument(const FName& InName, ERigVMPinDirection InDirection, TRigVMTypeIndex InTypeIndex)
@@ -61,6 +62,7 @@ FRigVMTemplateArgument::FRigVMTemplateArgument(const FName& InName, ERigVMPinDir
 , Direction(InDirection)
 , TypeIndices({InTypeIndex})
 {
+	EnsureValidExecuteType();
 	UpdateTypeToPermutations();
 }
 
@@ -72,6 +74,7 @@ FRigVMTemplateArgument::FRigVMTemplateArgument(const FName& InName, ERigVMPinDir
 {
 	check(TypeIndices.Num() > 0);
 
+	EnsureValidExecuteType();
 	UpdateTypeToPermutations();
 }
 
@@ -85,10 +88,44 @@ FRigVMTemplateArgument::FRigVMTemplateArgument(const FName& InName, ERigVMPinDir
 	{
 		TypeIndices.Append(FRigVMRegistry::Get().GetTypesForCategory(TypeCategory));
 	}
-	
+
+	EnsureValidExecuteType();
 	UpdateTypeToPermutations();
 }
 
+void FRigVMTemplateArgument::EnsureValidExecuteType()
+{
+	const FRigVMRegistry& Registry = FRigVMRegistry::Get();
+	for(TRigVMTypeIndex& TypeIndex : TypeIndices)
+	{
+		if(TypeIndex == INDEX_NONE)
+		{
+			continue;
+		}
+		
+		if(TypeIndex == RigVMTypeUtils::TypeIndex::Execute) 
+		{
+			continue;
+		}
+
+		if(!Registry.IsExecuteType(TypeIndex))
+		{
+			continue;
+		}
+
+		// execute arguments are converted to the base type to make
+		// matching types easier later
+		if(Registry.IsArrayType(TypeIndex))
+		{
+			TypeIndex = Registry.GetArrayTypeFromBaseTypeIndex(RigVMTypeUtils::TypeIndex::Execute);
+		}
+		else
+		{
+			TypeIndex = RigVMTypeUtils::TypeIndex::Execute;
+		}
+	}
+}
+
 void FRigVMTemplateArgument::Serialize(FArchive& Ar)
 {
 	Ar << Index;
@@ -180,6 +217,23 @@ void FRigVMTemplateArgument::UpdateTypeToPermutations()
 
 bool FRigVMTemplateArgument::SupportsTypeIndex(TRigVMTypeIndex InTypeIndex, TRigVMTypeIndex* OutTypeIndex) const
 {
+	if(InTypeIndex == INDEX_NONE)
+	{
+		return false;
+	}
+	
+	// convert any execute type into the base execute
+	const FRigVMRegistry& Registry = FRigVMRegistry::Get();
+	if(Registry.IsExecuteType(InTypeIndex))
+	{
+		const bool bIsArray = Registry.IsArrayType(InTypeIndex);
+		InTypeIndex = RigVMTypeUtils::TypeIndex::Execute;
+		if(bIsArray)
+		{
+			InTypeIndex = Registry.GetArrayTypeFromBaseTypeIndex(InTypeIndex);
+		}
+	}
+	
 	if (const TArray<int32>* Permutations = TypeToPermutations.Find(InTypeIndex))
 	{
 		if(OutTypeIndex)
@@ -190,7 +244,7 @@ bool FRigVMTemplateArgument::SupportsTypeIndex(TRigVMTypeIndex InTypeIndex, TRig
 	}
 
 	// Try to find compatible type
-	const TArray<TRigVMTypeIndex>& CompatibleTypes = FRigVMRegistry::Get().GetCompatibleTypes(InTypeIndex);
+	const TArray<TRigVMTypeIndex>& CompatibleTypes = Registry.GetCompatibleTypes(InTypeIndex);
 	for (const TRigVMTypeIndex& CompatibleTypeIndex : CompatibleTypes)
 	{
 		if (const TArray<int32>* Permutations = TypeToPermutations.Find(CompatibleTypeIndex))
@@ -1380,6 +1434,17 @@ bool FRigVMTemplate::AddTypeForArgument(const FName& InArgumentName, TRigVMTypeI
 		const FRigVMTemplateTypeMap Types = OnNewArgumentType().Execute(this, InArgumentName, InTypeIndex);
 		if(Types.Num() == Arguments.Num())
 		{
+			// make sure this permutation doesn't exist yet
+			FRigVMTemplateTypeMap TempTypes = Types;
+			TArray<int32> ExistingPermutations;
+			if(Resolve(TempTypes, ExistingPermutations, false))
+			{
+				if(ExistingPermutations.Num() == 1)
+				{
+					return false;
+				}
+			}
+			
 			for(FRigVMTemplateArgument& Argument : Arguments)
 			{
 				const TRigVMTypeIndex* TypeIndex = Types.Find(Argument.Name);
diff --git a/Engine/Source/Runtime/RigVM/Public/RigVMCore/RigVMDispatchFactory.h b/Engine/Source/Runtime/RigVM/Public/RigVMCore/RigVMDispatchFactory.h
--- a/Engine/Source/Runtime/RigVM/Public/RigVMCore/RigVMDispatchFactory.h
+++ b/Engine/Source/Runtime/RigVM/Public/RigVMCore/RigVMDispatchFactory.h
@@ -82,6 +82,9 @@ public:
 
 #endif
 
+	// Returns the execute context support for this dispatch factory
+	FORCEINLINE virtual UScriptStruct* GetExecuteContextStruct() const { return FRigVMExecuteContext::StaticStruct(); }
+
 	// registered needed types during registration of the factory
 	FORCEINLINE virtual void RegisterDependencyTypes() const {}
 
diff --git a/Engine/Source/Runtime/RigVM/Public/RigVMCore/RigVMExecuteContext.h b/Engine/Source/Runtime/RigVM/Public/RigVMCore/RigVMExecuteContext.h
--- a/Engine/Source/Runtime/RigVM/Public/RigVMCore/RigVMExecuteContext.h
+++ b/Engine/Source/Runtime/RigVM/Public/RigVMCore/RigVMExecuteContext.h
@@ -214,7 +214,7 @@ struct RIGVM_API FRigVMRuntimeSettings
  * The execute context is used for mutable nodes to
  * indicate execution order.
  */
-USTRUCT(BlueprintType)
+USTRUCT(BlueprintType, meta=(DisplayName="Execute Context"))
 struct RIGVM_API FRigVMExecuteContext
 {
 	GENERATED_BODY()
diff --git a/Engine/Source/Runtime/RigVM/Public/RigVMCore/RigVMRegistry.h b/Engine/Source/Runtime/RigVM/Public/RigVMCore/RigVMRegistry.h
--- a/Engine/Source/Runtime/RigVM/Public/RigVMCore/RigVMRegistry.h
+++ b/Engine/Source/Runtime/RigVM/Public/RigVMCore/RigVMRegistry.h
@@ -62,7 +62,10 @@ public:
 
 	// Adds a type if it doesn't exist yet and returns its index.
 	// This function is not thead-safe
-	TRigVMTypeIndex FindOrAddType(const FRigVMTemplateArgumentType& InType);
+	FORCEINLINE TRigVMTypeIndex FindOrAddType(const FRigVMTemplateArgumentType& InType)
+	{
+		return FindOrAddType_Internal(InType, false);
+	}
 
 	// Returns the type index given a type
 	TRigVMTypeIndex GetTypeIndex(const FRigVMTemplateArgumentType& InType) const;
@@ -232,6 +235,8 @@ private:
 		bool bIsExecute;
 	};
 
+	TRigVMTypeIndex FindOrAddType_Internal(const FRigVMTemplateArgumentType& InType, bool bForce);
+
 	FORCEINLINE static EObjectFlags DisallowedFlags()
 	{
 		return RF_BeginDestroyed | RF_FinishDestroyed;
@@ -278,4 +283,5 @@ private:
 	static FRigVMRegistry s_RigVMRegistry;
 	friend struct FRigVMStruct;
 	friend struct FRigVMTemplate;
+	friend struct FRigVMTemplateArgument;
 };
diff --git a/Engine/Source/Runtime/RigVM/Public/RigVMCore/RigVMTemplate.h b/Engine/Source/Runtime/RigVM/Public/RigVMCore/RigVMTemplate.h
--- a/Engine/Source/Runtime/RigVM/Public/RigVMCore/RigVMTemplate.h
+++ b/Engine/Source/Runtime/RigVM/Public/RigVMCore/RigVMTemplate.h
@@ -130,14 +130,16 @@ struct RIGVM_API FRigVMTemplateArgumentType
 		return CPPType.ToString();
 	}
 
-	void ConvertToArray() 
+	FRigVMTemplateArgumentType& ConvertToArray() 
 	{
 		CPPType = *RigVMTypeUtils::ArrayTypeFromBaseType(CPPType.ToString());
+		return *this;
 	}
 
-	void ConvertToBaseElement() 
+	FRigVMTemplateArgumentType& ConvertToBaseElement() 
 	{
 		CPPType = *RigVMTypeUtils::BaseTypeFromArrayType(CPPType.ToString());
+		return *this;
 	}
 };
 
@@ -229,9 +231,10 @@ protected:
 	TMap<TRigVMTypeIndex, TArray<int32>> TypeToPermutations;
 	TArray<ETypeCategory> TypeCategories;
 
-	// constructor from a property
+	// constructor from a property. this forces the type to be created
 	FRigVMTemplateArgument(FProperty* InProperty);
 
+	void EnsureValidExecuteType();
 	void UpdateTypeToPermutations();
 
 	friend struct FRigVMTemplate;
