RigVM: Nullptr fix for GetRegisteredUnitStructs

diff --git a/Engine/Source/Developer/RigVMDeveloper/Private/RigVMModel/RigVMController.cpp b/Engine/Source/Developer/RigVMDeveloper/Private/RigVMModel/RigVMController.cpp
--- a/Engine/Source/Developer/RigVMDeveloper/Private/RigVMModel/RigVMController.cpp
+++ b/Engine/Source/Developer/RigVMDeveloper/Private/RigVMModel/RigVMController.cpp
@@ -12696,12 +12696,14 @@ TArray<UScriptStruct*> URigVMController::GetRegisteredUnitStructs()
 		{
 			continue;
 		}
-		UScriptStruct* Struct = Function.Struct;
-		if (!Struct->IsChildOf(FRigVMStruct::StaticStruct()))
+		if(UScriptStruct* Struct = Function.Struct)
 		{
-			continue;
+			if (!Struct->IsChildOf(FRigVMStruct::StaticStruct()))
+			{
+				continue;
+			}
+			UnitStructs.Add(Struct);
 		}
-		UnitStructs.Add(Struct);
 	}
 	
 	return UnitStructs; 
