ControlRig: Rename exposed pin should also rename the display name in addition to the name of the pin object

diff --git a/Engine/Source/Developer/RigVMDeveloper/Private/RigVMModel/RigVMController.cpp b/Engine/Source/Developer/RigVMDeveloper/Private/RigVMModel/RigVMController.cpp
--- a/Engine/Source/Developer/RigVMDeveloper/Private/RigVMModel/RigVMController.cpp
+++ b/Engine/Source/Developer/RigVMDeveloper/Private/RigVMModel/RigVMController.cpp
@@ -10447,7 +10447,9 @@ bool URigVMController::RenameExposedPin(const FName& InOldPinName, const FName&
 			{
 				return false;
 			}
-
+			
+			InPin->SetDisplayName(InNewName);
+			
 			// update the eventually stored pin path to the new name
 			for (URigVMLink* Link : Links)
 			{
