Control Rig: Fix local variables nodes not dealing with constant values

diff --git a/Engine/Source/Developer/RigVMDeveloper/Private/RigVMCompiler/RigVMCompiler.cpp b/Engine/Source/Developer/RigVMDeveloper/Private/RigVMCompiler/RigVMCompiler.cpp
--- a/Engine/Source/Developer/RigVMDeveloper/Private/RigVMCompiler/RigVMCompiler.cpp
+++ b/Engine/Source/Developer/RigVMDeveloper/Private/RigVMCompiler/RigVMCompiler.cpp
@@ -2138,15 +2138,15 @@ FString URigVMCompiler::GetPinHashImpl(const URigVMPin* InPin, const FRigVMVarEx
 			{
 				if (bIsLiteral)
 				{
-					// Literal values will be reused for all instance of local variables
 					if (InVarExpr && InVarExpr->NumParents() == 0 && InVarExpr->NumChildren() == 0)
 					{
+						// Default literal values will be reused for all instance of local variables
 						return FString::Printf(TEXT("%sLocalVariableDefault::%s|%s%s"), *Prefix, *Node->GetGraph()->GetGraphName(), *VariableName.ToString(), *Suffix);
 					}
 					else if (InVarExpr)
 					{
-						const FString GraphPath = InVarExpr->GetProxy().GetCallstack().GetCallPath(false);
-						return FString::Printf(TEXT("%sLocalVariable::%s|%s%s"), *Prefix, *GraphPath, *VariableName.ToString(), *Suffix);					
+						const FString GraphPath = InVarExpr->GetProxy().GetCallstack().GetCallPath(true);
+						return FString::Printf(TEXT("%sLocalVariable::%s%s"), *Prefix, *GraphPath, *Suffix);					
 					}
 					else
 					{
