Control Rig: Fix CR_Minifig not working

diff --git a/Engine/Source/Developer/RigVMDeveloper/Private/RigVMCompiler/RigVMCompiler.cpp b/Engine/Source/Developer/RigVMDeveloper/Private/RigVMCompiler/RigVMCompiler.cpp
--- a/Engine/Source/Developer/RigVMDeveloper/Private/RigVMCompiler/RigVMCompiler.cpp
+++ b/Engine/Source/Developer/RigVMDeveloper/Private/RigVMCompiler/RigVMCompiler.cpp
@@ -336,6 +336,25 @@ bool URigVMCompiler::Compile(TArray<URigVMGraph*> InGraphs, URigVMController* In
 				}
 			}
 
+			if(URigVMUnitNode* UnitNode = Cast<URigVMUnitNode>(ModelNode))
+			{
+				if (!UnitNode->HasWildCardPin())
+				{
+					UScriptStruct* ScriptStruct = UnitNode->GetScriptStruct(); 
+					if(ScriptStruct == nullptr)
+					{
+						InController->FullyResolveTemplateNode(UnitNode, INDEX_NONE, false);
+					}
+
+					if (UnitNode->GetScriptStruct() == nullptr)
+					{
+						static const FString UnresolvedUnitNodeMessage = TEXT("Node @@ could not be resolved.");
+						Settings.ASTSettings.Report(EMessageSeverity::Error, ModelNode, UnresolvedUnitNodeMessage);
+						bEncounteredGraphError = true;
+					}
+				}
+			}
+
 			for(URigVMPin* Pin : ModelNode->Pins)
 			{
 				if(!URigVMController::EnsurePinValidity(Pin, true))
@@ -811,6 +830,8 @@ int32 URigVMCompiler::TraverseCallExtern(const FRigVMCallExternExprAST* InExpr,
 		const UScriptStruct* ScriptStruct = UnitNode->GetScriptStruct(); 
 		if(ScriptStruct == nullptr)
 		{
+			static const FString UnresolvedMessage = TEXT("Node @@ is unresolved.");
+			Settings.Report(EMessageSeverity::Error, UnitNode, UnresolvedMessage);
 			return INDEX_NONE;
 		}
 
@@ -831,6 +852,8 @@ int32 URigVMCompiler::TraverseCallExtern(const FRigVMCallExternExprAST* InExpr,
 	{
 		if(DispatchNode->GetFactory() == nullptr)
 		{
+			static const FString UnresolvedDispatchMessage = TEXT("Dispatch node @@ has no factory.");
+			Settings.Report(EMessageSeverity::Error, DispatchNode, UnresolvedDispatchMessage);
 			return INDEX_NONE;
 		}
 
