Control Rig: Fix dispatch nodes not being able to find their factory

diff --git a/Engine/Source/Developer/RigVMDeveloper/Private/RigVMModel/Nodes/RigVMDispatchNode.cpp b/Engine/Source/Developer/RigVMDeveloper/Private/RigVMModel/Nodes/RigVMDispatchNode.cpp
--- a/Engine/Source/Developer/RigVMDeveloper/Private/RigVMModel/Nodes/RigVMDispatchNode.cpp
+++ b/Engine/Source/Developer/RigVMDeveloper/Private/RigVMModel/Nodes/RigVMDispatchNode.cpp
@@ -131,6 +131,15 @@ const FRigVMDispatchFactory* URigVMDispatchNode::GetFactory() const
 		CachedFactory = Template->GetDispatchFactory(); 
 		return CachedFactory;
 	}
+
+	// Try to find factories that have not been registered yet
+	FRigVMRegistry::Get().Refresh();
+	if(const FRigVMTemplate* Template = GetTemplate())
+	{
+		CachedFactory = Template->GetDispatchFactory(); 
+		return CachedFactory;
+	}
+	
 	return nullptr;
 }
 
