RigVM: Auto-resolval for template nodes

diff --git a/Engine/Source/Developer/RigVMDeveloper/Private/RigVMModel/RigVMController.cpp b/Engine/Source/Developer/RigVMDeveloper/Private/RigVMModel/RigVMController.cpp
--- a/Engine/Source/Developer/RigVMDeveloper/Private/RigVMModel/RigVMController.cpp
+++ b/Engine/Source/Developer/RigVMDeveloper/Private/RigVMModel/RigVMController.cpp
@@ -72,6 +72,9 @@ URigVMController::URigVMController()
 	, bIsRunningUnitTest(false)
 	, bIsFullyResolvingTemplateNode(false)
 	, bSuspendRecomputingTemplateFilters(false)
+#if WITH_EDITOR
+	, bRegisterTemplateNodeUsage(true)
+#endif
 {
 }
 
@@ -87,6 +90,9 @@ URigVMController::URigVMController(const FObjectInitializer& ObjectInitializer)
 	, bIsRunningUnitTest(false)
 	, bIsFullyResolvingTemplateNode(false)
 	, bSuspendRecomputingTemplateFilters(false)
+#if WITH_EDITOR
+	, bRegisterTemplateNodeUsage(true)
+#endif
 {
 	ActionStack = CreateDefaultSubobject<URigVMActionStack>(TEXT("ActionStack"));
 
@@ -9352,6 +9358,52 @@ bool URigVMController::AddLink(URigVMPin* OutputPin, URigVMPin* InputPin, bool b
 
 	if (bSetupUndoRedo)
 	{
+#if WITH_EDITOR
+		auto ResolveTemplateNodeToCommonTypes = [this](URigVMPin* Pin)
+		{
+			if(!Pin->IsExecuteContext())
+			{
+				return;
+			}
+			
+			URigVMTemplateNode* TemplateNode = Cast<URigVMTemplateNode>(Pin->GetNode());
+			if(TemplateNode == nullptr)
+			{
+				return;
+			}
+
+			const FRigVMTemplate* Template = TemplateNode->GetTemplate();
+			if(Template == nullptr)
+			{
+				return;
+			}
+
+			if(!TemplateNode->HasWildCardPin())
+			{
+				return;
+			}
+
+			const FRigVMTemplate::FTypeMap PreferredTypes = GetCommonlyUsedTypesForTemplate(TemplateNode);
+			if(PreferredTypes.IsEmpty())
+			{
+				return;
+			}
+
+			const int32 PreferredPermutation = Template->FindPermutation(PreferredTypes);
+			if(PreferredPermutation != INDEX_NONE)
+			{
+				const TGuardValue<bool> DisableRegisterUseOfTemplate(bRegisterTemplateNodeUsage, false);
+				if(FullyResolveTemplateNode(TemplateNode, PreferredPermutation, true))
+				{
+					static const FString Message = TEXT("Template node was automatically resolved to commonly used types.");
+					SendUserFacingNotification(Message, 0.f, TemplateNode, TEXT("MessageLog.Note"));
+				}
+			}
+		};
+		ResolveTemplateNodeToCommonTypes(OutputPin);
+		ResolveTemplateNodeToCommonTypes(InputPin);
+#endif
+		
 		ActionStack->EndAction(Action);
 	}
 
@@ -15826,6 +15878,9 @@ bool URigVMController::FullyResolveTemplateNode(URigVMTemplateNode* InNode, int3
 	
 	if(bSetupUndoRedo)
 	{
+#if WITH_EDITOR
+		RegisterUseOfTemplate(InNode);
+#endif
 		CloseUndoBracket();
 	}
 
@@ -18553,3 +18608,93 @@ bool URigVMController::RestoreLinkedPaths(
 	
 	return bSuccess;
 }
+
+#if WITH_EDITOR
+
+void URigVMController::RegisterUseOfTemplate(const URigVMTemplateNode* InNode)
+{
+	check(InNode);
+
+	if(!bRegisterTemplateNodeUsage)
+	{
+		return;
+	}
+
+	const FRigVMTemplate* Template = InNode->GetTemplate();
+	if(Template == nullptr)
+	{
+		return;
+	}
+
+	if(!InNode->IsResolved())
+	{
+		return;
+	}
+
+	const TArray<int32>& Permutations = InNode->GetFilteredPermutationsIndices();
+	if(!ensure(Permutations.Num() == 1))
+	{
+		return;
+	}
+
+	URigVMControllerSettings* Settings = GetMutableDefault<URigVMControllerSettings>();
+	Settings->Modify();
+
+	const FName& Notation = Template->GetNotation();
+	FRigVMController_CommonTypePerTemplate& TypesForTemplate = Settings->TemplateDefaultTypes.FindOrAdd(Notation);
+
+	const FString TypesString = FRigVMTemplate::GetStringFromArgumentTypes(Template->GetTypesForPermutation(Permutations[0]));
+	int32& Count = TypesForTemplate.Counts.FindOrAdd(TypesString);
+	Count++;
+}
+
+FRigVMTemplate::FTypeMap URigVMController::GetCommonlyUsedTypesForTemplate(
+	const URigVMTemplateNode* InNode) const
+{
+	static FRigVMTemplate::FTypeMap EmptyTypes;
+
+	const URigVMControllerSettings* Settings = GetDefault<URigVMControllerSettings>();
+	if(!Settings->bAutoResolveTemplateNodesWhenLinkingExecute)
+	{
+		return EmptyTypes;
+	}
+	
+	const FRigVMTemplate* Template = InNode->GetTemplate();
+	if(Template == nullptr)
+	{
+		return EmptyTypes;
+	}
+
+	const FName& Notation = Template->GetNotation();
+
+	const FRigVMController_CommonTypePerTemplate* TypesForTemplate = Settings->TemplateDefaultTypes.Find(Notation);
+	if(TypesForTemplate == nullptr)
+	{
+		return EmptyTypes;
+	}
+
+	if(TypesForTemplate->Counts.IsEmpty())
+	{
+		return EmptyTypes;
+	}
+
+	TPair<FString,int32> MaxPair;
+	for(const TPair<FString,int32>& Pair : TypesForTemplate->Counts)
+	{
+		if(Pair.Value > MaxPair.Value)
+		{
+			MaxPair = Pair;
+		}
+	}
+
+	const FString& TypesString = MaxPair.Key;
+	return Template->GetArgumentTypesFromString(TypesString);
+}
+
+#endif
+
+URigVMControllerSettings::URigVMControllerSettings(const FObjectInitializer& Initializer)
+	: Super(Initializer)
+{
+	bAutoResolveTemplateNodesWhenLinkingExecute = true;
+}
diff --git a/Engine/Source/Developer/RigVMDeveloper/Public/RigVMModel/RigVMController.h b/Engine/Source/Developer/RigVMDeveloper/Public/RigVMModel/RigVMController.h
--- a/Engine/Source/Developer/RigVMDeveloper/Public/RigVMModel/RigVMController.h
+++ b/Engine/Source/Developer/RigVMDeveloper/Public/RigVMModel/RigVMController.h
@@ -1152,6 +1152,14 @@ public:
 		return RestoreLinkedPaths(InLinkedPaths, InNodeNameMap, InRemapDelegates, FRigVMController_CheckPinComatibilityDelegate(), bSetupUndoRedo, InUserDirection);
 	}
 
+#if WITH_EDITOR
+	// Registers this template node's use for later determining the commonly used types
+	void RegisterUseOfTemplate(const URigVMTemplateNode* InNode);
+
+	// Inquire on the commonly used types for a template node. This can be used to resolve a node without user input (as a default)
+	FRigVMTemplate::FTypeMap GetCommonlyUsedTypesForTemplate(const URigVMTemplateNode* InNode) const;
+#endif
+	
 private: 
 	UPROPERTY(transient)
 	TArray<TObjectPtr<URigVMGraph>> Graphs;
@@ -1209,6 +1217,9 @@ private:
 	bool bIsRunningUnitTest;
 	bool bIsFullyResolvingTemplateNode;
 	bool bSuspendRecomputingTemplateFilters;
+#if WITH_EDITOR
+	bool bRegisterTemplateNodeUsage;
+#endif
 
 	friend class URigVMGraph;
 	friend class URigVMPin;
@@ -1244,3 +1255,36 @@ private:
 	URigVMController* Controller;
 	bool bUndo;
 };
+
+USTRUCT()
+struct FRigVMController_CommonTypePerTemplate
+{
+	GENERATED_BODY()
+
+	UPROPERTY(EditAnywhere, Category = "RigVMController")
+	TMap<FString,int32> Counts;
+};
+
+/**
+ * Default settings for the RigVM Controller
+ */
+UCLASS(config = EditorSettings)
+class  URigVMControllerSettings : public UObject
+{
+public:
+	URigVMControllerSettings(const FObjectInitializer& Initializer);
+
+	GENERATED_BODY()
+
+	/**
+	 * When adding a link to an execute pin on a template node,
+	 * this functionality automatically resolves the template node to the
+	 * most commonly used type.
+	 */
+	UPROPERTY(EditAnywhere, Category = "RigVMController")
+	bool bAutoResolveTemplateNodesWhenLinkingExecute;
+
+	/** The commonly used types for a template node */
+	UPROPERTY()
+	TMap<FName,FRigVMController_CommonTypePerTemplate> TemplateDefaultTypes;
+};
\ No newline at end of file
diff --git a/Engine/Source/Runtime/RigVM/Private/RigVMCore/RigVMDispatchFactory.cpp b/Engine/Source/Runtime/RigVM/Private/RigVMCore/RigVMDispatchFactory.cpp
--- a/Engine/Source/Runtime/RigVM/Private/RigVMCore/RigVMDispatchFactory.cpp
+++ b/Engine/Source/Runtime/RigVM/Private/RigVMCore/RigVMDispatchFactory.cpp
@@ -90,17 +90,9 @@ FString FRigVMDispatchFactory::GetPermutationName(const FRigVMTemplateTypeMap& I
 
 FString FRigVMDispatchFactory::GetPermutationNameImpl(const FRigVMTemplateTypeMap& InTypes) const
 {
-	const FRigVMRegistry& Registry = FRigVMRegistry::Get();
-	TArray<FString> TypePairStrings;
-	for(const TPair<FName,TRigVMTypeIndex>& Pair : InTypes)
-	{
-		const FRigVMTemplateArgumentType& Type = Registry.GetType(Pair.Value);
-		static constexpr TCHAR Format[] = TEXT("%s:%s");
-		TypePairStrings.Add(FString::Printf(Format, *Pair.Key.ToString(), *Type.CPPType.ToString()));
-	}
-			
+	const FString TypePairStrings = FRigVMTemplate::GetStringFromArgumentTypes(InTypes);
 	static constexpr TCHAR Format[] = TEXT("%s::%s");
-	return FString::Printf(Format, *GetFactoryName().ToString(), *FString::Join(TypePairStrings, TEXT(",")));
+	return FString::Printf(Format, *GetFactoryName().ToString(), *TypePairStrings);
 }
 
 const FRigVMTemplate* FRigVMDispatchFactory::GetTemplate() const
diff --git a/Engine/Source/Runtime/RigVM/Private/RigVMCore/RigVMTemplate.cpp b/Engine/Source/Runtime/RigVM/Private/RigVMCore/RigVMTemplate.cpp
--- a/Engine/Source/Runtime/RigVM/Private/RigVMCore/RigVMTemplate.cpp
+++ b/Engine/Source/Runtime/RigVM/Private/RigVMCore/RigVMTemplate.cpp
@@ -618,6 +618,21 @@ FRigVMTemplate::FTypeMap FRigVMTemplate::GetArgumentTypesFromString(const FStrin
 	return Types;
 }
 
+FString FRigVMTemplate::GetStringFromArgumentTypes(const FTypeMap& InTypes)
+{
+	const FRigVMRegistry& Registry = FRigVMRegistry::Get();
+	
+	TArray<FString> TypePairStrings;
+	for(const TPair<FName,TRigVMTypeIndex>& Pair : InTypes)
+	{
+		const FRigVMTemplateArgumentType& Type = Registry.GetType(Pair.Value);
+		static constexpr TCHAR Format[] = TEXT("%s:%s");
+		TypePairStrings.Add(FString::Printf(Format, *Pair.Key.ToString(), *Type.CPPType.ToString()));
+	}
+			
+	return FString::Join(TypePairStrings, TEXT(","));
+}
+
 bool FRigVMTemplate::IsValid() const
 {
 	return !Notation.IsNone();
diff --git a/Engine/Source/Runtime/RigVM/Public/RigVMCore/RigVMTemplate.h b/Engine/Source/Runtime/RigVM/Public/RigVMCore/RigVMTemplate.h
--- a/Engine/Source/Runtime/RigVM/Public/RigVMCore/RigVMTemplate.h
+++ b/Engine/Source/Runtime/RigVM/Public/RigVMCore/RigVMTemplate.h
@@ -338,6 +338,9 @@ public:
 	// converts the types provided by a string (like "A:float,B:int32") into a type map
 	FTypeMap GetArgumentTypesFromString(const FString& InTypeString) const;
 
+	// converts the types provided to a string (like "A:float,B:int32")
+	static FString GetStringFromArgumentTypes(const FTypeMap& InTypes); 
+
 #if WITH_EDITOR
 
 	// Returns the color based on the permutation's metadata
