Control Rig: Fix GetFilteredTypes returning wildcard type

diff --git a/Engine/Source/Developer/RigVMDeveloper/Private/RigVMModel/RigVMController.cpp b/Engine/Source/Developer/RigVMDeveloper/Private/RigVMModel/RigVMController.cpp
--- a/Engine/Source/Developer/RigVMDeveloper/Private/RigVMModel/RigVMController.cpp
+++ b/Engine/Source/Developer/RigVMDeveloper/Private/RigVMModel/RigVMController.cpp
@@ -16254,6 +16254,10 @@ TArray<TRigVMTypeIndex> URigVMController::GetFilteredTypes(URigVMPin* InPin)
 			}
 		}
 	}
+	if (InPin->IsWildCard())
+	{
+		return TArray<TRigVMTypeIndex>();
+	}
 	return {InPin->GetTypeIndex()};
 }
 
