Control Rig: Fix crash on compile when node is not able to resolve

diff --git a/Engine/Source/Developer/RigVMDeveloper/Private/RigVMModel/RigVMController.cpp b/Engine/Source/Developer/RigVMDeveloper/Private/RigVMModel/RigVMController.cpp
--- a/Engine/Source/Developer/RigVMDeveloper/Private/RigVMModel/RigVMController.cpp
+++ b/Engine/Source/Developer/RigVMDeveloper/Private/RigVMModel/RigVMController.cpp
@@ -15599,6 +15599,11 @@ bool URigVMController::FullyResolveTemplateNode(URigVMTemplateNode* InNode, int3
 	
 	check(InNode);
 
+	if (InNode->IsSingleton())
+	{
+		return true;
+	}
+
 	const FRigVMTemplate* Template = InNode->GetTemplate();
 	const FRigVMDispatchFactory* Factory = Template->GetDispatchFactory();
 	
