Control Rig: Fix deprecated template permutations loading correctly

diff --git a/Engine/Source/Developer/RigVMDeveloper/Private/RigVMModel/Nodes/RigVMTemplateNode.cpp b/Engine/Source/Developer/RigVMDeveloper/Private/RigVMModel/Nodes/RigVMTemplateNode.cpp
--- a/Engine/Source/Developer/RigVMDeveloper/Private/RigVMModel/Nodes/RigVMTemplateNode.cpp
+++ b/Engine/Source/Developer/RigVMDeveloper/Private/RigVMModel/Nodes/RigVMTemplateNode.cpp
@@ -1086,13 +1086,9 @@ void URigVMTemplateNode::InvalidateCache()
 	CachedFunction = nullptr;
 	CachedTemplate = nullptr;
 
-	for(URigVMPin* Pin : GetPins())
+	if (HasWildCardPin())
 	{
-		if(Pin->IsWildCard())
-		{
-			ResolvedFunctionName.Reset();
-			break;
-		}
+		ResolvedFunctionName.Reset();
 	}
 }
 
diff --git a/Engine/Source/Developer/RigVMDeveloper/Private/RigVMModel/Nodes/RigVMUnitNode.cpp b/Engine/Source/Developer/RigVMDeveloper/Private/RigVMModel/Nodes/RigVMUnitNode.cpp
--- a/Engine/Source/Developer/RigVMDeveloper/Private/RigVMModel/Nodes/RigVMUnitNode.cpp
+++ b/Engine/Source/Developer/RigVMDeveloper/Private/RigVMModel/Nodes/RigVMUnitNode.cpp
@@ -12,16 +12,23 @@ void URigVMUnitNode::PostLoad()
 	Super::PostLoad();
 
 	// if we have a script struct but no notation let's figure out the template
-	if(GetScriptStruct() != nullptr && GetTemplate() == nullptr)
+	if(GetScriptStruct() != nullptr)
 	{
-		if(const FRigVMFunction* Function = FRigVMRegistry::Get().FindFunction(GetScriptStruct(), *GetMethodName().ToString()))
+		if (IsDeprecated())
 		{
-			if(Function->TemplateIndex != INDEX_NONE)
+			TemplateNotation = NAME_None;
+		}
+		else if(GetTemplate() == nullptr)
+		{
+			if(const FRigVMFunction* Function = FRigVMRegistry::Get().FindFunction(GetScriptStruct(), *GetMethodName().ToString()))
 			{
-				const FRigVMTemplate& Template = FRigVMRegistry::Get().GetTemplates()[Function->TemplateIndex];
-				TemplateNotation = Template.GetNotation();				
-			}
-			ResolvedFunctionName = Function->GetName();
+				if(Function->TemplateIndex != INDEX_NONE)
+				{
+					const FRigVMTemplate& Template = FRigVMRegistry::Get().GetTemplates()[Function->TemplateIndex];
+					TemplateNotation = Template.GetNotation();				
+				}
+				ResolvedFunctionName = Function->GetName();
+			}			
 		}
 	}
 }
