Control Rig: Fix broken rigs with deprecated sequence nodes

diff --git a/Engine/Source/Developer/RigVMDeveloper/Private/RigVMModel/Nodes/RigVMUnitNode.cpp b/Engine/Source/Developer/RigVMDeveloper/Private/RigVMModel/Nodes/RigVMUnitNode.cpp
--- a/Engine/Source/Developer/RigVMDeveloper/Private/RigVMModel/Nodes/RigVMUnitNode.cpp
+++ b/Engine/Source/Developer/RigVMDeveloper/Private/RigVMModel/Nodes/RigVMUnitNode.cpp
@@ -17,6 +17,10 @@ void URigVMUnitNode::PostLoad()
 		if (IsDeprecated())
 		{
 			TemplateNotation = NAME_None;
+			if(const FRigVMFunction* Function = FRigVMRegistry::Get().FindFunction(GetScriptStruct(), *GetMethodName().ToString()))
+			{
+				ResolvedFunctionName = Function->GetName();
+			}
 		}
 		else if(GetTemplate() == nullptr)
 		{
