Control Rig: Fix crash when collapsing to function

diff --git a/Engine/Source/Developer/RigVMDeveloper/Private/RigVMModel/RigVMController.cpp b/Engine/Source/Developer/RigVMDeveloper/Private/RigVMModel/RigVMController.cpp
--- a/Engine/Source/Developer/RigVMDeveloper/Private/RigVMModel/RigVMController.cpp
+++ b/Engine/Source/Developer/RigVMDeveloper/Private/RigVMModel/RigVMController.cpp
@@ -9283,7 +9283,10 @@ bool URigVMController::AddLink(URigVMPin* OutputPin, URigVMPin* InputPin, bool b
 		URigVMPin* SecondToResolve = (FirstToResolve == OutputPin) ? InputPin : OutputPin;
 		if (!PrepareToLink(FirstToResolve, SecondToResolve, bSetupUndoRedo))
 		{
-			ActionStack->CancelAction(Action, this);
+			if (bSetupUndoRedo)
+			{
+				ActionStack->CancelAction(Action, this);
+			}
 			return false;
 		}
 	}
