Control Rig: Fix crash when collapsing to function

diff --git a/Engine/Source/Developer/RigVMDeveloper/Private/RigVMModel/Nodes/RigVMTemplateNode.cpp b/Engine/Source/Developer/RigVMDeveloper/Private/RigVMModel/Nodes/RigVMTemplateNode.cpp
--- a/Engine/Source/Developer/RigVMDeveloper/Private/RigVMModel/Nodes/RigVMTemplateNode.cpp
+++ b/Engine/Source/Developer/RigVMDeveloper/Private/RigVMModel/Nodes/RigVMTemplateNode.cpp
@@ -728,6 +728,11 @@ TArray<TRigVMTypeIndex> URigVMTemplateNode::GetFilteredTypesForPin(URigVMPin* In
 	ensureMsgf(InPin->GetNode() == this, TEXT("GetFilteredTypesForPin of %s with pin from another node %s"), *GetNodePath(), *InPin->GetPinPath(true));
 	TArray<TRigVMTypeIndex> FilteredTypes;
 
+	if (IsSingleton())
+	{
+		return {InPin->GetTypeIndex()};
+	}
+
 	if (FilteredPermutations.IsEmpty())
 	{
 		return FilteredTypes;	
diff --git a/Engine/Source/Developer/RigVMDeveloper/Private/RigVMModel/RigVMPin.cpp b/Engine/Source/Developer/RigVMDeveloper/Private/RigVMModel/RigVMPin.cpp
--- a/Engine/Source/Developer/RigVMDeveloper/Private/RigVMModel/RigVMPin.cpp
+++ b/Engine/Source/Developer/RigVMDeveloper/Private/RigVMModel/RigVMPin.cpp
@@ -1614,7 +1614,7 @@ URigVMNode* URigVMPin::GetNode() const
 	}
 
 	URigVMNode* Node = Cast<URigVMNode>(GetOuter());
-	if(Node)
+	if(IsValid(Node))
 	{
 		return Node;
 	}
