ensure exec context uniqueness except for aggregate nodes

diff --git a/Engine/Source/Developer/RigVMDeveloper/Private/RigVMModel/RigVMController.cpp b/Engine/Source/Developer/RigVMDeveloper/Private/RigVMModel/RigVMController.cpp
--- a/Engine/Source/Developer/RigVMDeveloper/Private/RigVMModel/RigVMController.cpp
+++ b/Engine/Source/Developer/RigVMDeveloper/Private/RigVMModel/RigVMController.cpp
@@ -9789,13 +9789,20 @@ FName URigVMController::AddExposedPin(const FName& InPinName, ERigVMPinDirection
 	}
 	
 	// only allow one IO / input exposed pin of type execute context per direction
-	if(InDirection != ERigVMPinDirection::Output)
+	// except for aggregate nodes that can have multiple exec outputs
+	bool bCheckForExecUniqueness = true;
+	if (LibraryNode->IsA<URigVMAggregateNode>())
+	{
+		bCheckForExecUniqueness = InDirection != ERigVMPinDirection::Output;
+	}
+	
+	if (bCheckForExecUniqueness)
 	{
 		if(const UScriptStruct* CPPTypeStruct = Cast<UScriptStruct>(CPPTypeObject))
 		{
 			if(CPPTypeStruct->IsChildOf(FRigVMExecuteContext::StaticStruct()))
 			{
-				for(URigVMPin* ExistingPin : LibraryNode->Pins)
+				for(const URigVMPin* ExistingPin : LibraryNode->Pins)
 				{
 					if(ExistingPin->IsExecuteContext())
 					{
