ControlRig: use break link to removing injected nodes so that link deletion does not fail and leave a invalid link in the graph.

diff --git a/Engine/Source/Developer/RigVMDeveloper/Private/RigVMModel/RigVMController.cpp b/Engine/Source/Developer/RigVMDeveloper/Private/RigVMModel/RigVMController.cpp
--- a/Engine/Source/Developer/RigVMDeveloper/Private/RigVMModel/RigVMController.cpp
+++ b/Engine/Source/Developer/RigVMDeveloper/Private/RigVMModel/RigVMController.cpp
@@ -5980,12 +5980,7 @@ bool URigVMController::RemoveNode(URigVMNode* InNode, bool bSetupUndoRedo, bool
 
 		for (URigVMPin* Pin : InNode->GetPins())
 		{
-			TArray<URigVMInjectionInfo*> InjectedNodes = Pin->GetInjectedNodes();
-			for (int32 i=InjectedNodes.Num()-1; i >= 0; --i)
-			{
-				RemoveNode(InjectedNodes[i]->Node, bSetupUndoRedo, bRecursive);
-			}
-
+			// breaking links also removes injected nodes 
 			BreakAllLinks(Pin, true, bSetupUndoRedo);
 			BreakAllLinks(Pin, false, bSetupUndoRedo);
 			BreakAllLinksRecursive(Pin, true, false, bSetupUndoRedo);
