Control Rig: Fix aggregate nodes inside functions can not be pasted

diff --git a/Engine/Source/Developer/RigVMDeveloper/Private/RigVMModel/RigVMController.cpp b/Engine/Source/Developer/RigVMDeveloper/Private/RigVMModel/RigVMController.cpp
--- a/Engine/Source/Developer/RigVMDeveloper/Private/RigVMModel/RigVMController.cpp
+++ b/Engine/Source/Developer/RigVMDeveloper/Private/RigVMModel/RigVMController.cpp
@@ -3407,13 +3407,16 @@ TArray<FName> URigVMController::ImportNodesFromText(const FString& InText, bool
 	}
 
 	TArray<TGuardValue<bool>> EditGuards;
-	for (URigVMNode* CreatedNode : Factory.CreatedNodes)
+	TArray<URigVMNode*> CreatedNodes = Factory.CreatedNodes;
+	for (int32 i=0; i<CreatedNodes.Num(); ++i)
 	{
+		URigVMNode* CreatedNode = CreatedNodes[i];
 		if (URigVMLibraryNode* LibraryNode = Cast<URigVMLibraryNode>(CreatedNode))
 		{
 			if (URigVMGraph* ContainedGraph = LibraryNode->GetContainedGraph())
 			{
 				EditGuards.Emplace(ContainedGraph->bEditable, true);
+				CreatedNodes.Append(ContainedGraph->GetNodes());
 			}
 		}
 	}
