Control Rig: Fix local variables using wrong operand in RigVM

diff --git a/Engine/Source/Developer/RigVMDeveloper/Private/RigVMCompiler/RigVMCompiler.cpp b/Engine/Source/Developer/RigVMDeveloper/Private/RigVMCompiler/RigVMCompiler.cpp
--- a/Engine/Source/Developer/RigVMDeveloper/Private/RigVMCompiler/RigVMCompiler.cpp
+++ b/Engine/Source/Developer/RigVMDeveloper/Private/RigVMCompiler/RigVMCompiler.cpp
@@ -2607,11 +2607,47 @@ FRigVMOperand URigVMCompiler::FindOrAddRegister(const FRigVMVarExprAST* InVarExp
 		bIsLiteral ? ERigVMMemoryType::Literal:
 		(bIsDebugValue ? ERigVMMemoryType::Debug : ERigVMMemoryType::Work);
 
+	TArray<FString> HashesWithSharedOperand;
+	
 	FRigVMOperand const* ExistingOperand = WorkData.PinPathToOperand->Find(Hash);
+	if (!ExistingOperand)
+	{
+		if(Settings.ASTSettings.bFoldAssignments) 		
+		{
+			// Get all possible pins that lead to the same operand		
+			const FRigVMCompilerWorkData::FRigVMASTProxyArray PinProxies = FindProxiesWithSharedOperand(InVarExpr, WorkData);
+			ensure(!PinProxies.IsEmpty());
+
+			// Look for an existing operand from a different pin with shared operand
+			for (const FRigVMASTProxy& Proxy : PinProxies)
+			{
+				if (const URigVMPin* VirtualPin = Cast<URigVMPin>(Proxy.GetSubject()))
+				{
+					const FString VirtualPinHash = GetPinHash(VirtualPin, InVarExpr, bIsDebugValue, Proxy);
+					HashesWithSharedOperand.Add(VirtualPinHash);
+					if (Pin != VirtualPin)
+					{
+						ExistingOperand = WorkData.PinPathToOperand->Find(VirtualPinHash);
+						if (ExistingOperand)
+						{
+							break;
+						}
+					}
+				}	
+			}
+		}
+	}
+	
 	if (ExistingOperand)
 	{
 		if(ExistingOperand->GetMemoryType() == MemoryType)
 		{
+			// Add any missing hash that shares this existing operand
+			for (const FString& VirtualPinHash : HashesWithSharedOperand)
+			{
+				WorkData.PinPathToOperand->Add(VirtualPinHash, *ExistingOperand);
+			}
+			
 			if (!bIsDebugValue)
 			{
 				check(!WorkData.ExprToOperand.Contains(InVarExpr));
@@ -2753,16 +2789,9 @@ FRigVMOperand URigVMCompiler::FindOrAddRegister(const FRigVMVarExprAST* InVarExp
 		// tbd: this functionality is only needed when there is a watch anywhere?
 		//if(!WorkData.WatchedPins.IsEmpty())
 		{
-			const FRigVMCompilerWorkData::FRigVMASTProxyArray& PinProxies = FindProxiesWithSharedOperand(InVarExpr, WorkData);
-			ensure(!PinProxies.IsEmpty());
-
-			for (const FRigVMASTProxy& Proxy : PinProxies)
+			for (const FString& VirtualPinHash : HashesWithSharedOperand)
 			{
-				if (URigVMPin* VirtualPin = Cast<URigVMPin>(Proxy.GetSubject()))
-				{
-					FString VirtualPinHash = GetPinHash(VirtualPin, InVarExpr, bIsDebugValue, Proxy);
-					WorkData.PinPathToOperand->Add(VirtualPinHash, Operand);
-				}	
+				WorkData.PinPathToOperand->Add(VirtualPinHash, Operand);
 			}
 		}
 	}
diff --git a/Engine/Source/Developer/RigVMDeveloper/Private/RigVMModel/RigVMController.cpp b/Engine/Source/Developer/RigVMDeveloper/Private/RigVMModel/RigVMController.cpp
--- a/Engine/Source/Developer/RigVMDeveloper/Private/RigVMModel/RigVMController.cpp
+++ b/Engine/Source/Developer/RigVMDeveloper/Private/RigVMModel/RigVMController.cpp
@@ -11845,22 +11845,6 @@ bool URigVMController::SetLocalVariableDefaultValue(const FName& InVariableName,
 	FRigVMGraphVariableDescription& VariableDescription = LocalVariables[FoundIndex];
 	VariableDescription.DefaultValue = InDefaultValue;
 	
-	// Refresh variable nodes to reflect change in default value
-	TArray<URigVMNode*> Nodes = Graph->GetNodes();
-	for (URigVMNode* Node : Nodes)
-	{
-		if (URigVMVariableNode* VariableNode = Cast<URigVMVariableNode>(Node))
-		{
-			if (URigVMPin* VariablePin = VariableNode->FindPin(URigVMVariableNode::VariableName))
-			{
-				if (VariablePin->GetDefaultValue() == InVariableName.ToString())
-				{
-					SetPinDefaultValue(VariableNode->FindPin(URigVMVariableNode::ValueName), InDefaultValue, true, true, true, bNotify);
-				}
-			}
-		}
-	}
-
 	if (!bSuspendNotifications)
 	{
 		Graph->MarkPackageDirty();
