Control Rig: Fix RigVM.Graph.BindPinLibraryNode crashing in DebugMode

diff --git a/Engine/Source/Developer/RigVMDeveloper/Private/RigVMModel/RigVMController.cpp b/Engine/Source/Developer/RigVMDeveloper/Private/RigVMModel/RigVMController.cpp
--- a/Engine/Source/Developer/RigVMDeveloper/Private/RigVMModel/RigVMController.cpp
+++ b/Engine/Source/Developer/RigVMDeveloper/Private/RigVMModel/RigVMController.cpp
@@ -5983,6 +5983,12 @@ bool URigVMController::RemoveNode(URigVMNode* InNode, bool bSetupUndoRedo, bool
 
 		for (URigVMPin* Pin : InNode->GetPins())
 		{
+			// Remove injected nodes
+			while (Pin->HasInjectedNodes())
+			{
+				RemoveInjectedNode(Pin->GetPinPath(), Pin->GetDirection() != ERigVMPinDirection::Output, bSetupUndoRedo);
+			}
+		
 			// breaking links also removes injected nodes 
 			BreakAllLinks(Pin, true, bSetupUndoRedo);
 			BreakAllLinks(Pin, false, bSetupUndoRedo);
