Control Rig: Once a pin of a template node is resolved, it will add a preferred type

diff --git a/Engine/Source/Developer/RigVMDeveloper/Private/RigVMModel/Nodes/RigVMTemplateNode.cpp b/Engine/Source/Developer/RigVMDeveloper/Private/RigVMModel/Nodes/RigVMTemplateNode.cpp
--- a/Engine/Source/Developer/RigVMDeveloper/Private/RigVMModel/Nodes/RigVMTemplateNode.cpp
+++ b/Engine/Source/Developer/RigVMDeveloper/Private/RigVMModel/Nodes/RigVMTemplateNode.cpp
@@ -807,36 +807,32 @@ TArray<TRigVMTypeIndex> URigVMTemplateNode::GetFilteredTypesForPin(URigVMPin* In
 	return FilteredTypes;
 }
 
-bool URigVMTemplateNode::TryReduceTypesToSingle(TArray<TRigVMTypeIndex>& InOutTypes, const TRigVMTypeIndex PreferredType) const
+TRigVMTypeIndex URigVMTemplateNode::TryReduceTypesToSingle(const TArray<TRigVMTypeIndex>& InTypes, const TRigVMTypeIndex PreferredType) const
 {
-	if (InOutTypes.IsEmpty())
+	if (InTypes.IsEmpty())
 	{
-		return false;
+		return INDEX_NONE;
 	}
 
-	if (InOutTypes.Num() == 1)
+	if (InTypes.Num() == 1)
 	{
-		return true;
+		return InTypes[0];
 	}
 	
-	for (int32 i=1; i<InOutTypes.Num(); ++i)
+	for (int32 i=1; i<InTypes.Num(); ++i)
 	{
-		if (!FRigVMRegistry::Get().CanMatchTypes(InOutTypes[0], InOutTypes[i], true))
+		if (!FRigVMRegistry::Get().CanMatchTypes(InTypes[0], InTypes[i], true))
 		{
-			return false;
+			return INDEX_NONE;
 		}
 	}
 
-	if (InOutTypes.Contains(PreferredType))
+	if (InTypes.Contains(PreferredType))
 	{
-		InOutTypes = {PreferredType};
+		return PreferredType;
 	}
-	else
-	{
-		InOutTypes = {InOutTypes[0]};
-	}
-			
-	return true;	
+
+	return InTypes[0];
 }
 
 TArray<int32> URigVMTemplateNode::GetNewFilteredPermutations(URigVMPin* InPin, URigVMPin* LinkedPin)
diff --git a/Engine/Source/Developer/RigVMDeveloper/Private/RigVMModel/RigVMController.cpp b/Engine/Source/Developer/RigVMDeveloper/Private/RigVMModel/RigVMController.cpp
--- a/Engine/Source/Developer/RigVMDeveloper/Private/RigVMModel/RigVMController.cpp
+++ b/Engine/Source/Developer/RigVMDeveloper/Private/RigVMModel/RigVMController.cpp
@@ -1812,7 +1812,7 @@ bool URigVMController::UnresolveTemplateNodes(const TArray<FName>& InNodeNames,
 	return false;
 }
 
-bool URigVMController::UnresolveTemplateNodes(const TArray<URigVMNode*>& InNodes, bool bSetupUndoRedo)
+bool URigVMController::UnresolveTemplateNodes(const TArray<URigVMNode*>& InNodes, bool bSetupUndoRedo, bool bBreakLinks)
 {
 	if (!IsValidGraph() || InNodes.IsEmpty())
 	{
@@ -1863,19 +1863,22 @@ bool URigVMController::UnresolveTemplateNodes(const TArray<URigVMNode*>& InNodes
 					TemplateNode->PreferredPermutationPairs = {};
 				}	
 			}
-	
-			TArray<URigVMLink*> Links = Node->GetLinks();
-			for (int32 i=0; i<Links.Num(); ++i)
+
+			if (bBreakLinks)
 			{
-				URigVMLink* Link = Links[i];
-				URigVMPin* SourcePin = Link->GetSourcePin();
-				URigVMPin* TargetPin = Link->GetTargetPin();
-			
-				URigVMNode* OtherNode = SourcePin->GetNode() == Node ? TargetPin->GetNode() : SourcePin->GetNode();
-				if (!InNodes.Contains(OtherNode))
+				TArray<URigVMLink*> Links = Node->GetLinks();
+				for (int32 i=0; i<Links.Num(); ++i)
 				{
-					BreakLink(SourcePin, TargetPin, bSetupUndoRedo);
-				}			
+					URigVMLink* Link = Links[i];
+					URigVMPin* SourcePin = Link->GetSourcePin();
+					URigVMPin* TargetPin = Link->GetTargetPin();
+			
+					URigVMNode* OtherNode = SourcePin->GetNode() == Node ? TargetPin->GetNode() : SourcePin->GetNode();
+					if (!InNodes.Contains(OtherNode))
+					{
+						BreakLink(SourcePin, TargetPin, bSetupUndoRedo);
+					}			
+				}
 			}
 		}
 	}
@@ -9218,6 +9221,11 @@ bool URigVMController::AddLink(URigVMPin* OutputPin, URigVMPin* InputPin, bool b
 			// We temporarily add a preferred type to make sure that doesn't happen
 			if (URigVMTemplateNode* TemplateNode = Cast<URigVMTemplateNode>(EdgePin->GetNode()))
 			{
+				if (TemplateNode->IsSingleton())
+				{
+					continue;
+				}
+				
 				bAddedTemporaryPreferredType = true;
 				URigVMPin* RootPin = EdgePin->GetRootPin();
 				for (FRigVMTemplatePreferredType& Preferred : TemplateNode->PreferredPermutationPairs)
@@ -10257,9 +10265,11 @@ bool URigVMController::AddArgumentForPin(URigVMPin* InPin, URigVMPin* InToLinkPi
 	LibraryNode->Template = NewTemplate;	
 	
 	// Update Entry and Return filtered permutations to all permutations
+	URigVMTemplateNode* EntryNode = Graph->GetEntryNode();
+	URigVMTemplateNode* ReturnNode = Graph->GetReturnNode();	
 	for (int32 EntryReturn=0; EntryReturn<2; ++EntryReturn)
 	{
-		URigVMTemplateNode* Node = (EntryReturn == 0) ? (URigVMTemplateNode*)Graph->GetEntryNode() : (URigVMTemplateNode*)Graph->GetReturnNode();
+		URigVMTemplateNode* Node = (EntryReturn == 0) ? EntryNode : ReturnNode;
 		if (Node)
 		{
 			TArray<int32> OldPermutations = Node->FilteredPermutations;
@@ -10274,9 +10284,23 @@ bool URigVMController::AddArgumentForPin(URigVMPin* InPin, URigVMPin* InToLinkPi
 			{
 				ActionStack->AddAction(FRigVMSetTemplateFilteredPermutationsAction(Node, OldPermutations));
 			}
-
-			UpdateTemplateNodePinTypes(Node, bSetupUndoRedo);
-			Notify(ERigVMGraphNotifType::NodeDescriptionChanged, Node);
+		}
+	}
+	{
+		bool bUpdatedTypes = false;
+		for (int32 EntryReturn=0; EntryReturn<2; ++EntryReturn)
+		{
+			URigVMTemplateNode* Node = (EntryReturn == 0) ? EntryNode : ReturnNode;
+			if (Node)
+			{
+				// Update template node pin types of an interface node will update both entry and return types.
+				if (!bUpdatedTypes)
+				{
+					UpdateTemplateNodePinTypes(Node, bSetupUndoRedo);
+					bUpdatedTypes = true;
+				}
+				Notify(ERigVMGraphNotifType::NodeDescriptionChanged, Node);
+			}
 		}
 	}
 
@@ -10390,6 +10414,22 @@ bool URigVMController::RemoveExposedPin(const FName& InPinName, bool bSetupUndoR
 		}
 	}
 
+	// Remove preferred type of library node (and references to it)
+	{
+		FRigVMControllerGraphGuard GraphGuard(this, LibraryNode->GetGraph(), bSetupUndoRedo);
+		RemovePreferredType(LibraryNode, InPinName, bSetupUndoRedo);
+
+		// If the library node is a function, remove preferred type from all references
+		if (URigVMFunctionLibrary* FunctionLibrary = Cast<URigVMFunctionLibrary>(LibraryNode->GetGraph()))
+		{
+			FunctionLibrary->ForEachReference(LibraryNode->GetFName(), [this, InPinName, bSetupUndoRedo](URigVMFunctionReferenceNode* ReferenceNode)
+			{
+				FRigVMControllerGraphGuard GraphGuard(this, ReferenceNode->GetGraph(), bSetupUndoRedo);
+				RemovePreferredType(ReferenceNode, InPinName, bSetupUndoRedo);
+			});
+		}
+	}
+
 	// Recompute filtered types and library template
 	RecomputeAllTemplateFilteredPermutations(bSetupUndoRedo);
 
@@ -16205,7 +16245,7 @@ bool URigVMController::PrepareTemplatePinForType(URigVMPin* InPin, const TArray<
 				return false;
 			}
 			
-			UnresolveTemplateNodes({TemplateNode}, bSetupUndoRedo);
+			UnresolveTemplateNodes({TemplateNode}, bSetupUndoRedo, false);
 			PrepareTemplatePinForType(InPin, InTypeIndices, bSetupUndoRedo);
 
 			// Reapply the valid preferred types
@@ -16420,6 +16460,11 @@ bool URigVMController::ResolveWildCardPin(URigVMPin* InPin, TRigVMTypeIndex InTy
 	if(RootPin->IsWildCard())
 	{
 		ensure(RootPin->GetNode()->IsA<URigVMTemplateNode>());
+		URigVMTemplateNode* TemplateNode = CastChecked<URigVMTemplateNode>(RootPin->GetNode());
+		if (!TemplateNode->SupportsType(RootPin, Type, &Type))
+		{
+			return false;
+		}
 		
 		FRigVMBaseAction Action;
 		if (bSetupUndoRedo)
@@ -16428,8 +16473,6 @@ bool URigVMController::ResolveWildCardPin(URigVMPin* InPin, TRigVMTypeIndex InTy
 			ActionStack->BeginAction(Action);
 		}
 		
-		URigVMTemplateNode* TemplateNode = CastChecked<URigVMTemplateNode>(RootPin->GetNode());
-
 		// Add preferred type
 		AddPreferredType(TemplateNode, RootPin->GetFName(), Type, bSetupUndoRedo);
 
@@ -16502,9 +16545,36 @@ bool URigVMController::ResolveWildCardPin(URigVMPin* InPin, TRigVMTypeIndex InTy
 	else
 	{
 		OpenUndoBracket(TEXT("Resolving wildcard pin"));
-		UnresolveTemplateNodes({*InPin->GetNode()->GetName()}, bSetupUndoRedo);
+		URigVMTemplateNode* Node = Cast<URigVMTemplateNode>(InPin->GetNode());
+		TArray<URigVMLink*> Links = Node->GetLinks();
+		UnresolveTemplateNodes({*Node->GetName()}, bSetupUndoRedo);
 		if (ResolveWildCardPin(InPin, InTypeIndex, bSetupUndoRedo, bPrintPythonCommand))
 		{
+			// Try to reattach old links
+			for (URigVMLink* Link : Links)
+			{
+				URigVMPin* SourcePin = Link->GetSourcePin();
+				URigVMPin* TargetPin = Link->GetTargetPin();
+				URigVMPin* LocalPin = (SourcePin->GetNode() == Node) ? SourcePin : TargetPin;
+				URigVMPin* RemotePin = (LocalPin == TargetPin) ? SourcePin : TargetPin;
+				{
+					TArray<TRigVMTypeIndex> FilteredTypes = GetFilteredTypes(RemotePin);
+					bool bSupportsLink = false;
+					for (const TRigVMTypeIndex& FilteredType : FilteredTypes)
+					{
+						if (Node->FilteredSupportsType(LocalPin, FilteredType))
+						{
+							bSupportsLink = true;
+							break;
+						}
+					}
+					if (bSupportsLink)
+					{
+						AddLink(SourcePin, TargetPin, bSetupUndoRedo);
+					}
+				}
+			}
+			
 			CloseUndoBracket();
 			return true;
 		}
@@ -16624,105 +16694,244 @@ bool URigVMController::UpdateFilteredPermutations(URigVMTemplateNode* InNode, co
 
 bool URigVMController::UpdateTemplateNodePinTypes(URigVMTemplateNode* InNode, bool bSetupUndoRedo, bool bInitializeDefaultValue)
 {
-	check(InNode->GetGraph() == GetGraph());
+	URigVMGraph* Graph = GetGraph();
+	check(InNode->GetGraph() == Graph);
 	bool bAnyTypeChanged = false;
 
-	TArray<int32> ResolvedPermutations = InNode->GetFilteredPermutationsIndices();
+	const FRigVMTemplate* Template = InNode->GetTemplate();
 
-	// Try to pick permutations that respects the preferred types
-	TArray<int32> PreferredPermutations = InNode->FindPermutationsForTypes(InNode->PreferredPermutationPairs, false);
-	ResolvedPermutations = ResolvedPermutations.FilterByPredicate([&](const int32& OtherPermutation) { return PreferredPermutations.Contains(OtherPermutation); });
-	if (ResolvedPermutations.IsEmpty())
+	// Because interface nodes reference the same template, they need to be resolved at the same time (otherwise we
+	// might end up in a situation where they resolve to different permutations).
+	const bool bIsInterfaceNode = InNode->IsA<URigVMFunctionEntryNode>() || InNode->IsA<URigVMFunctionReturnNode>();
+	URigVMTemplateNode* EntryNode = nullptr;
+	URigVMTemplateNode* ReturnNode = nullptr;
+	if (bIsInterfaceNode)
+	{
+		EntryNode = Graph->GetEntryNode();
+		ReturnNode = Graph->GetReturnNode();
+	}
+	
+	TArray<int32> ResolvedPermutations;
+	TArray<int32> PreferredPermutations;
+	TArray<FRigVMTemplatePreferredType> PreferredPermutationPairs = InNode->PreferredPermutationPairs;
+	if (!bIsInterfaceNode)
 	{
 		ResolvedPermutations = InNode->GetFilteredPermutationsIndices();
-
-		// Try to pick one permutations that respects the preferred types (with casting)
+		
+		// Try to pick permutations that respects the preferred types
 		PreferredPermutations = InNode->FindPermutationsForTypes(InNode->PreferredPermutationPairs, true);
 		ResolvedPermutations = ResolvedPermutations.FilterByPredicate([&](const int32& OtherPermutation) { return PreferredPermutations.Contains(OtherPermutation); });
+		InNode->ResolvedPermutation = ResolvedPermutations.Num() == 1 ? ResolvedPermutations[0] : INDEX_NONE;		
 	}
-	InNode->ResolvedPermutation = ResolvedPermutations.Num() == 1 ? ResolvedPermutations[0] : INDEX_NONE;
+	else
+	{
+		if (EntryNode)
+		{
+			ResolvedPermutations = EntryNode->FilteredPermutations;
+			PreferredPermutationPairs = EntryNode->PreferredPermutationPairs;
+		}
+		if (ReturnNode)
+		{
+			ResolvedPermutations = ResolvedPermutations.FilterByPredicate([ReturnNode](const int32& Index)
+			{
+				return ReturnNode->FilteredPermutations.Contains(Index);
+			});
 
-	const FRigVMTemplate* Template = InNode->GetTemplate();
+			for (const FRigVMTemplatePreferredType& Preferred : ReturnNode->PreferredPermutationPairs)
+			{
+				PreferredPermutationPairs.AddUnique(Preferred);				
+			}
+		}
+	
+		// Try to pick permutations that respects the preferred types
+		PreferredPermutations = InNode->FindPermutationsForTypes(PreferredPermutationPairs, true);
+		ResolvedPermutations = ResolvedPermutations.FilterByPredicate([&](const int32& OtherPermutation) { return PreferredPermutations.Contains(OtherPermutation); });
 
-	// First unresolve any pins that need unresolving, then resolve everything else
-	for (int32 UnresolveResolve=0; UnresolveResolve<2; ++UnresolveResolve)
+		if (EntryNode)
+		{
+			EntryNode->ResolvedPermutation = ResolvedPermutations.Num() == 1 ? ResolvedPermutations[0] : INDEX_NONE;
+		}
+		if (ReturnNode)
+		{
+			ReturnNode->ResolvedPermutation = ResolvedPermutations.Num() == 1 ? ResolvedPermutations[0] : INDEX_NONE;
+		}
+	}
+
+	TArray<URigVMPin*> Pins = InNode->GetPins();
+	if (bIsInterfaceNode)
 	{
-		for(int32 PinIndex=0; PinIndex < InNode->GetPins().Num(); ++PinIndex)
+		if (EntryNode)
 		{
-			URigVMPin* Pin = InNode->GetPins()[PinIndex];
-			if (UnresolveResolve == 0 && Pin->IsWildCard())
+			Pins = EntryNode->GetPins();
+		}
+		if (ReturnNode)
+		{
+			Pins.Append(ReturnNode->GetPins());
+		}
+	}
+	
+	TArray<const FRigVMTemplateArgument*> Arguments;
+	Arguments.SetNumUninitialized(Pins.Num());
+	for(int32 PinIndex=0; PinIndex < Pins.Num(); ++PinIndex)
+	{
+		URigVMPin* Pin = Pins[PinIndex];
+		Arguments[PinIndex] = Template->FindArgument(Pin->GetFName());
+	}
+
+	// Remove invalid permutations
+	ResolvedPermutations.RemoveAll([&](const int32& Permutation)
+	{
+		for (const FRigVMTemplateArgument* Argument : Arguments)
+		{
+			if (Argument && Argument->TypeIndices[Permutation] == INDEX_NONE)
 			{
-				continue;
+				return true;
 			}
-			
-			bool bNeedsToUnresolve = false;
+		}
+		return false;
+	});	
+
+	// Find the types for each possible permutation
+	TMap<int32, TArray<TRigVMTypeIndex>> PinTypes; // permutation to pin types
+	for (int32 ResolvedPermutation : ResolvedPermutations)
+	{
+		for(int32 PinIndex=0; PinIndex < Pins.Num(); ++PinIndex)
+		{
+			URigVMPin* Pin = Pins[PinIndex];
 			if (Pin->GetDirection() == ERigVMPinDirection::Hidden)
 			{
+				PinTypes.FindOrAdd(ResolvedPermutation).Add(INDEX_NONE);
 				continue;
 			}
-		
-			TArray<TRigVMTypeIndex> Types;
-			Types.Reserve(ResolvedPermutations.Num());
-			for (const int32& Permutation : ResolvedPermutations)
-			{
-				if (const FRigVMTemplateArgument* Argument = Template->FindArgument(Pin->GetFName()))
-				{	
-					// INDEX_NONE indicates invalid permutation
-					if (Argument->TypeIndices[Permutation] != INDEX_NONE)
-					{
-						Types.AddUnique(Argument->TypeIndices[Permutation]);
-					}
-				}
+
+			bool bAddedType = false;
+			TArray<TRigVMTypeIndex>& Types = PinTypes.FindOrAdd(ResolvedPermutation);
+			if (const FRigVMTemplateArgument* Argument = Arguments[PinIndex])
+			{
+				Types.Add(Argument->TypeIndices[ResolvedPermutation]);
+				bAddedType = true;
 			}
-			
-			if (Types.IsEmpty())
+			else
 			{
 				// We might be in the case where the pin doesn't own an argument yet, but the node has a preferred type
-				if (!InNode->GetTemplate()->FindArgument(Pin->GetFName()))
+				for (int32 i=0; i<PreferredPermutationPairs.Num(); ++i)
 				{
-					for (int32 i=0; i<InNode->PreferredPermutationPairs.Num(); ++i)
+					if (Pin->GetFName() == PreferredPermutationPairs[i].Argument)
 					{
-						if (Pin->GetFName() == InNode->PreferredPermutationPairs[i].Argument)
-						{
-							Types = {InNode->PreferredPermutationPairs[i].GetTypeIndex()};
-							break;
-						}
+						Types.Add(PreferredPermutationPairs[i].GetTypeIndex());
+						bAddedType = true;
+						break;
 					}
 				}
+			}
+			
+			if (!bAddedType)
+			{
+				// This is a pin with no argument or preferred type
+				// Its marked as invalid, and will show wildcard type
+				Types.Add(INDEX_NONE);
+			}
+		}
+	}
 
-				if (Types.IsEmpty())
+	// Some pins can benefit from reducing types to a single option
+	// Even if some pins reduce to a single type (maybe represented by a single permutation), the rest should still display a wildcard pin
+	// If reduction happens for multiple pins, we need to make sure that they reduce to the same permutation
+	// If possible, we want to respect the preferred type if it exists, or the current pin type if its not wildcard
+	TArray<bool> WasReduced;
+	WasReduced.SetNumZeroed(Pins.Num());
+
+	TMap<int32, TArray<TRigVMTypeIndex>> ReducedTypes = PinTypes;
+	TArray<TRigVMTypeIndex> FinalPinTypes;
+	FinalPinTypes.SetNumUninitialized(Pins.Num());
+	
+	// Reduce to preferred types (no casting)
+	if (ReducedTypes.Num() > 1)
+	{
+		for (FRigVMTemplatePreferredType& Preferred : PreferredPermutationPairs)
+		{
+			// Finding all pins related to this argument (in the case of interface nodes, there might be more than one)
+			TArray<URigVMPin*> ArgumentPins = Pins.FilterByPredicate([Preferred](URigVMPin* Pin)
+			{
+				return Pin->GetFName() == Preferred.Argument;
+			});
+
+			for (URigVMPin* ArgumentPin : ArgumentPins)
+			{
+				int32 PinIndex = Pins.Find(ArgumentPin);
+				WasReduced[PinIndex] = true;
+				FinalPinTypes[PinIndex] = Preferred.TypeIndex;
+				ReducedTypes = ReducedTypes.FilterByPredicate([Preferred, PinIndex](const TPair<int32, TArray<TRigVMTypeIndex>>& Permutation)
 				{
-					bNeedsToUnresolve = true;
-				}
+					return Permutation.Value[PinIndex] == Preferred.TypeIndex;
+				});
 			}
+		}
+	}
 
-			// Try to reduce the number of types to one
-			if (Types.Num() > 1)
+	if (PinTypes.Num() > ReducedTypes.Num())
+	{
+		ensureMsgf(!ReducedTypes.IsEmpty(), TEXT("Found incompatible preferred types"));
+	}
+	
+	// Reduce compatible types, try to find the permutation in the reduced types
+	for(int32 PinIndex=0; PinIndex < Pins.Num(); ++PinIndex)
+	{
+		if (WasReduced[PinIndex])
+		{
+			continue;
+		}
+		
+		URigVMPin* Pin = Pins[PinIndex];
+		if (Pin->GetDirection() == ERigVMPinDirection::Hidden)
+		{
+			continue;
+		}
+		TArray<TRigVMTypeIndex> Types;
+		Types.Reserve(ResolvedPermutations.Num());
+		TRigVMTypeIndex PreferredType = INDEX_NONE;
+		int32 TypesFoundInReduced = 0;
+		for (int32 ResolvedPermutation : ResolvedPermutations)
+		{
+			Types.Add(PinTypes.FindChecked(ResolvedPermutation)[PinIndex]);
+			if (ReducedTypes.Contains(ResolvedPermutation))
 			{
-				TRigVMTypeIndex PreferredType = Pin->GetTypeIndex();
-				if (InNode->ResolvedPermutation != INDEX_NONE)
+				TypesFoundInReduced++;
+				if (PreferredType == INDEX_NONE)
 				{
-					if (const FRigVMTemplateArgument* Argument = Template->FindArgument(Pin->GetFName()))
-					{
-						PreferredType = Argument->TypeIndices[InNode->ResolvedPermutation];
-					}
-				}
-				InNode->TryReduceTypesToSingle(Types, PreferredType);
-				if (Types.Num() == 1)
-				{
-					// Make sure the other pins use the same permutation
-					InNode->ResolvedPermutation = *ResolvedPermutations.FindByPredicate([&](const int32& Permutation)
-					{
-						if (const FRigVMTemplateArgument* Argument = Template->FindArgument(Pin->GetFName()))
-						{
-							return Types[0] == Argument->TypeIndices[Permutation];
-						}
-						return false;
-					});
-				}
+					PreferredType = PinTypes.FindChecked(ResolvedPermutation)[PinIndex];
+				}				
+			}
+		}
+
+		if (TypesFoundInReduced > 1)
+		{
+			PreferredType = Pin->GetTypeIndex();
+		}
+		FinalPinTypes[PinIndex] = InNode->TryReduceTypesToSingle(Types, PreferredType);
+		if (FinalPinTypes[PinIndex] != INDEX_NONE)
+		{
+			WasReduced[PinIndex] =  true;
+		}
+	}
+	
+	// First unresolve any pins that need unresolving, then resolve everything else
+	for (int32 UnresolveResolve=0; UnresolveResolve<2; ++UnresolveResolve)
+	{
+		for(int32 PinIndex=0; PinIndex < Pins.Num(); ++PinIndex)
+		{
+			URigVMPin* Pin = Pins[PinIndex];
+			if (UnresolveResolve == 0 && Pin->IsWildCard())
+			{
+				continue;
 			}
 			
-			if (bNeedsToUnresolve || Types.Num() > 1)
+			if (Pin->GetDirection() == ERigVMPinDirection::Hidden)
+			{
+				continue;
+			}
+		
+			if (FinalPinTypes[PinIndex] == INDEX_NONE)
 			{
 				// Unresolve
 				if (Pin->HasInjectedNodes())
@@ -16730,7 +16939,7 @@ bool URigVMController::UpdateTemplateNodePinTypes(URigVMTemplateNode* InNode, bo
 					EjectNodeFromPin(Pin, bSetupUndoRedo);
 				}
 				
-				const FRigVMTemplateArgument* Argument = InNode->GetTemplate()->FindArgument(*Pin->GetName());
+				const FRigVMTemplateArgument* Argument = Template->FindArgument(*Pin->GetName());
 				FRigVMTemplateArgument::EArrayType ArrayType;
 				if (Argument)
 				{
@@ -16759,43 +16968,102 @@ bool URigVMController::UpdateTemplateNodePinTypes(URigVMTemplateNode* InNode, bo
 					ChangePinType(Pin, CPPType, CPPObjectType, bSetupUndoRedo, false, false, false, bInitializeDefaultValue);
 				}
 			}
-			else if (Types.Num() == 1)
+			else 
 			{
 				// Resolve
-				if (Pin->GetTypeIndex() != Types[0])
+				if (Pin->GetTypeIndex() != FinalPinTypes[PinIndex])
 				{
 					bAnyTypeChanged = !Pin->IsExecuteContext();
-					ChangePinType(Pin, Types[0], bSetupUndoRedo, false, false, false, bInitializeDefaultValue);
+					ChangePinType(Pin, FinalPinTypes[PinIndex], bSetupUndoRedo, false, false, false, bInitializeDefaultValue);
 				}
 
-				// Once an interface pin is resolved, add it as a preferred type
-				if (InNode->IsA<URigVMFunctionEntryNode>() || InNode->IsA<URigVMFunctionReturnNode>())
+				// Once an pin is resolved, add it as a preferred type
+				URigVMTemplateNode* TemplateNode = CastChecked<URigVMTemplateNode>(Pin->GetNode());
+				TRigVMTypeIndex TypeIndex = TemplateNode->GetPreferredType(Pin->GetFName());
+				if (TypeIndex != FinalPinTypes[PinIndex] && TypeIndex != INDEX_NONE)
 				{
-					TRigVMTypeIndex TypeIndex = InNode->GetPreferredType(Pin->GetFName());
-					if (TypeIndex != Types[0] && TypeIndex != INDEX_NONE)
-					{
-						RemovePreferredType(InNode, Pin->GetFName(), bSetupUndoRedo);
-						TypeIndex = INDEX_NONE;
-					}
-					if (TypeIndex == INDEX_NONE)
-					{
-						AddPreferredType(InNode, Pin->GetFName(), Types[0], bSetupUndoRedo);
-					}
+					RemovePreferredType(TemplateNode, Pin->GetFName(), bSetupUndoRedo);
+					TypeIndex = INDEX_NONE;
+				}
+				if (TypeIndex == INDEX_NONE)
+				{
+					AddPreferredType(TemplateNode, Pin->GetFName(), FinalPinTypes[PinIndex], bSetupUndoRedo);
 				}
 			}
-			else
+		}
+	}
+
+	bool bHasWildcard=false;
+	if (bIsInterfaceNode)
+	{
+		if (EntryNode)
+		{
+			bHasWildcard = EntryNode->HasWildCardPin();
+		}
+		if (ReturnNode)
+		{
+			bHasWildcard |= ReturnNode->HasWildCardPin();
+		}
+	}
+	else
+	{
+		bHasWildcard = InNode->HasWildCardPin();
+	}
+	if (bHasWildcard)
+	{
+		if (bIsInterfaceNode)
+		{
+			if (EntryNode)
 			{
-				ensure(false);
-				return false;
+				EntryNode->ResolvedPermutation = INDEX_NONE;
+			}
+			if (ReturnNode)
+			{
+				ReturnNode->ResolvedPermutation = INDEX_NONE;
 			}
 		}
+		else
+		{
+			InNode->ResolvedPermutation = INDEX_NONE;
+		}
 	}
-
-	if (InNode->HasWildCardPin())
+	else if (!InNode->PreferredPermutationPairs.IsEmpty())
 	{
-		InNode->ResolvedPermutation = INDEX_NONE;
+		if (bIsInterfaceNode)
+		{
+			PreferredPermutationPairs.Reset();
+			if (EntryNode)
+			{
+				PreferredPermutationPairs = EntryNode->PreferredPermutationPairs;
+			}
+			if (ReturnNode)
+			{
+				for (const FRigVMTemplatePreferredType& Preferred : ReturnNode->PreferredPermutationPairs)
+				{
+					PreferredPermutationPairs.AddUnique(Preferred);				
+				}
+			}
+			
+			PreferredPermutations = InNode->FindPermutationsForTypes(PreferredPermutationPairs, false);
+			ensure(PreferredPermutations.Num() == 1);
+			if (EntryNode)
+			{
+				EntryNode->ResolvedPermutation = PreferredPermutations[0];
+			}
+			if (ReturnNode)
+			{
+				EntryNode->ResolvedPermutation = PreferredPermutations[0];
+			}
+		}
+		else
+		{
+			PreferredPermutations = InNode->FindPermutationsForTypes(InNode->PreferredPermutationPairs, false);
+			ensure(PreferredPermutations.Num() == 1);
+			InNode->ResolvedPermutation = PreferredPermutations[0];
+		}			
 	}
-	else if (URigVMUnitNode* UnitNode = Cast<URigVMUnitNode>(InNode))
+	
+	if (URigVMUnitNode* UnitNode = Cast<URigVMUnitNode>(InNode))
 	{
 		if (const FRigVMFunction* Function = UnitNode->GetResolvedFunction())
 		{
@@ -16899,6 +17167,7 @@ bool URigVMController::PropagateTemplateFilteredTypes(URigVMTemplateNode* InNode
 bool URigVMController::AddPreferredType(URigVMTemplateNode* InNode, const FName& InPinName, const TRigVMTypeIndex& InPreferredTypeIndex, bool bSetupUndoRedo)
 {
 	check(InNode->FindPin(InPinName.ToString()));
+	check(!InNode->IsSingleton());
 
 	FRigVMTemplatePreferredType* PreferredType = InNode->PreferredPermutationPairs.FindByPredicate([&](const FRigVMTemplatePreferredType& Other)
 	{
@@ -16910,7 +17179,62 @@ bool URigVMController::AddPreferredType(URigVMTemplateNode* InNode, const FName&
 		return PreferredType->TypeIndex == InPreferredTypeIndex;		
 	}
 
+	// We need to remove or modify the existing preferred types that are not compatible with the new preferred type
+	// For example, a template which has 2 pins (A, B), and both can be float, double or FVector (always A type == B type)
+	// If a preferred type exists in which B is double, and we are adding A as float
+	// we need to change the B preferred type to float
+	// If the old preferred type is not compatible, we need to remove the B preferred type
+	const FRigVMTemplate* Template = InNode->GetTemplate();
+	TArray<int32> Permutations;
+	FRigVMTemplate::FTypeMap TypeMap;
+	TypeMap.Add(InPinName, InPreferredTypeIndex);
+	Template->Resolve(TypeMap, Permutations, false);
+
+	const FRigVMRegistry& Registry = FRigVMRegistry::Get();
+	
 	TArray<FRigVMTemplatePreferredType> NewPreferredTypes = InNode->PreferredPermutationPairs;
+	TArray<int32> ToRemove;
+	for (int32 PreferredIndex=0; PreferredIndex<InNode->PreferredPermutationPairs.Num(); ++PreferredIndex)
+	{
+		FRigVMTemplatePreferredType& Preferred = InNode->PreferredPermutationPairs[PreferredIndex];
+		TArray<TRigVMTypeIndex> PinTypes;
+		if (const FRigVMTemplateArgument* Argument = Template->FindArgument(Preferred.Argument))
+		{
+			bool bFoundMatch = Permutations.ContainsByPredicate([Argument, Preferred](const int32& Permutation)
+			{
+				if (Argument->TypeIndices[Permutation] == Preferred.TypeIndex)
+				{
+					return true;
+				}
+				return false;
+			});
+			if (!bFoundMatch)
+			{
+				for (const int32& Permutation : Permutations)
+				{
+					if (Registry.CanMatchTypes(Argument->TypeIndices[Permutation], Preferred.TypeIndex, true))
+					{
+						NewPreferredTypes[PreferredIndex].TypeIndex = Argument->TypeIndices[Permutation];
+#if UE_RIGVM_DEBUG_TYPEINDEX
+						NewPreferredTypes[PreferredIndex].UpdateStringFromIndex();
+#endif
+						bFoundMatch = true;
+						break;
+					}
+				}
+			}
+			if (!bFoundMatch)
+			{
+				ToRemove.Add(PreferredIndex);
+			}
+		}
+	}
+	for (int32 i=ToRemove.Num()-1; i>=0; ++i)
+	{
+		NewPreferredTypes.RemoveAt(ToRemove[i]);
+	}
+
+	
 	NewPreferredTypes.Add(FRigVMTemplatePreferredType(InPinName, InPreferredTypeIndex));
 	if (bSetupUndoRedo)
 	{
@@ -17081,8 +17405,8 @@ bool URigVMController::UpdateLibraryTemplate(URigVMLibraryNode* LibraryNode, boo
 		return false;
 	}
 
-	URigVMFunctionEntryNode* EntryNode = ContainedGraph->GetEntryNode();
-	URigVMFunctionReturnNode* ReturnNode = ContainedGraph->GetReturnNode();
+	URigVMTemplateNode* EntryNode = ContainedGraph->GetEntryNode();
+	URigVMTemplateNode* ReturnNode = ContainedGraph->GetReturnNode();
 	const TArray<URigVMTemplateNode*> InterfaceNodes = {EntryNode, ReturnNode};
 	const FRigVMTemplate* CurTemplate = LibraryNode->GetTemplate();
 
@@ -17267,7 +17591,7 @@ bool URigVMController::UpdateLibraryTemplate(URigVMLibraryNode* LibraryNode, boo
 		FRigVMControllerGraphGuard GraphGuard(this, ContainedGraph, bSetupUndoRedo);
 		for (int32 EntryReturn=0; EntryReturn<2; ++EntryReturn)
 		{
-			URigVMTemplateNode* Node = (EntryReturn == 0) ? (URigVMTemplateNode*)EntryNode : (URigVMTemplateNode*)ReturnNode;
+			URigVMTemplateNode* Node = (EntryReturn == 0) ? EntryNode : ReturnNode;
 			if (Node)
 			{
 				TArray<int32> OldPermutations = Node->FilteredPermutations;
@@ -17275,15 +17599,29 @@ bool URigVMController::UpdateLibraryTemplate(URigVMLibraryNode* LibraryNode, boo
 				for (int32 i=0; i<Node->FilteredPermutations.Num(); ++i)
 				{
 					Node->FilteredPermutations[i] = i;
-				}		
+				}
 
 				if (bSetupUndoRedo)
 				{
 					ActionStack->AddAction(FRigVMSetTemplateFilteredPermutationsAction(Node, OldPermutations));
 				}
-
-				UpdateTemplateNodePinTypes(Node, bSetupUndoRedo);
-				Notify(ERigVMGraphNotifType::NodeDescriptionChanged, Node);
+			}
+		}
+		{
+			bool bUpdatedTypes = false;
+			for (int32 EntryReturn=0; EntryReturn<2; ++EntryReturn)
+			{
+				URigVMTemplateNode* Node = (EntryReturn == 0) ? EntryNode : ReturnNode;
+				if (Node)
+				{
+					// Update template node pin types of an interface node will update both entry and return types.
+					if (!bUpdatedTypes)
+					{
+						UpdateTemplateNodePinTypes(Node, bSetupUndoRedo);
+						bUpdatedTypes = true;
+					}
+					Notify(ERigVMGraphNotifType::NodeDescriptionChanged, Node);
+				}
 			}
 		}
 	}
@@ -18696,8 +19034,8 @@ void URigVMController::RegisterUseOfTemplate(const URigVMTemplateNode* InNode)
 		return;
 	}
 
-	const TArray<int32>& Permutations = InNode->GetFilteredPermutationsIndices();
-	if(!ensure(Permutations.Num() == 1))
+	const int32& ResolvedPermutation = InNode->ResolvedPermutation;
+	if(!ensure(ResolvedPermutation != INDEX_NONE))
 	{
 		return;
 	}
@@ -18708,7 +19046,7 @@ void URigVMController::RegisterUseOfTemplate(const URigVMTemplateNode* InNode)
 	const FName& Notation = Template->GetNotation();
 	FRigVMController_CommonTypePerTemplate& TypesForTemplate = Settings->TemplateDefaultTypes.FindOrAdd(Notation);
 
-	const FString TypesString = FRigVMTemplate::GetStringFromArgumentTypes(Template->GetTypesForPermutation(Permutations[0]));
+	const FString TypesString = FRigVMTemplate::GetStringFromArgumentTypes(Template->GetTypesForPermutation(ResolvedPermutation));
 	int32& Count = TypesForTemplate.Counts.FindOrAdd(TypesString);
 	Count++;
 }
diff --git a/Engine/Source/Developer/RigVMDeveloper/Public/RigVMModel/Nodes/RigVMTemplateNode.h b/Engine/Source/Developer/RigVMDeveloper/Public/RigVMModel/Nodes/RigVMTemplateNode.h
--- a/Engine/Source/Developer/RigVMDeveloper/Public/RigVMModel/Nodes/RigVMTemplateNode.h
+++ b/Engine/Source/Developer/RigVMDeveloper/Public/RigVMModel/Nodes/RigVMTemplateNode.h
@@ -140,7 +140,7 @@ public:
 
 	// Tries to reduce the input types to a single type, if all are compatible
 	// Will prioritize the InPreferredType if available
-	bool TryReduceTypesToSingle(TArray<TRigVMTypeIndex>& InOutTypes, const TRigVMTypeIndex InPreferredType = INDEX_NONE) const;	
+	TRigVMTypeIndex TryReduceTypesToSingle(const TArray<TRigVMTypeIndex>& InTypes, const TRigVMTypeIndex PreferredType = TRigVMTypeIndex()) const;
 
 	// returns true if updating pin filters with InTypes would result in different filters 
 	bool PinNeedsFilteredTypesUpdate(URigVMPin* InPin, const TArray<TRigVMTypeIndex>& InTypeIndices);
diff --git a/Engine/Source/Developer/RigVMDeveloper/Public/RigVMModel/RigVMController.h b/Engine/Source/Developer/RigVMDeveloper/Public/RigVMModel/RigVMController.h
--- a/Engine/Source/Developer/RigVMDeveloper/Public/RigVMModel/RigVMController.h
+++ b/Engine/Source/Developer/RigVMDeveloper/Public/RigVMModel/RigVMController.h
@@ -284,7 +284,7 @@ public:
 	// Turns a resolved templated node(s) back into its template.
 	UFUNCTION(BlueprintCallable, Category = RigVMController)
 	bool UnresolveTemplateNodes(const TArray<FName>& InNodeNames, bool bSetupUndoRedo = true, bool bPrintPythonCommand = false);
-	bool UnresolveTemplateNodes(const TArray<URigVMNode*>& InNodes, bool bSetupUndoRedo);
+	bool UnresolveTemplateNodes(const TArray<URigVMNode*>& InNodes, bool bSetupUndoRedo, bool bBreakLinks = true);
 
 	// Upgrades a set of nodes with each corresponding next known version
 	UFUNCTION(BlueprintCallable, Category = RigVMController)
