Control Rig: Temporarily only check if permutations are empty. If there are multiples, pick the first one.

#rnx

diff --git a/Engine/Source/Developer/RigVMDeveloper/Private/RigVMModel/RigVMController.cpp b/Engine/Source/Developer/RigVMDeveloper/Private/RigVMModel/RigVMController.cpp
--- a/Engine/Source/Developer/RigVMDeveloper/Private/RigVMModel/RigVMController.cpp
+++ b/Engine/Source/Developer/RigVMDeveloper/Private/RigVMModel/RigVMController.cpp
@@ -17045,7 +17045,8 @@ bool URigVMController::UpdateTemplateNodePinTypes(URigVMTemplateNode* InNode, bo
 			}
 			
 			PreferredPermutations = InNode->FindPermutationsForTypes(PreferredPermutationPairs, false);
-			ensure(PreferredPermutations.Num() == 1);
+			// ensure(PreferredPermutations.Num() == 1);
+			check(!PreferredPermutations.IsEmpty());
 			if (EntryNode)
 			{
 				EntryNode->ResolvedPermutation = PreferredPermutations[0];
@@ -17058,7 +17059,8 @@ bool URigVMController::UpdateTemplateNodePinTypes(URigVMTemplateNode* InNode, bo
 		else
 		{
 			PreferredPermutations = InNode->FindPermutationsForTypes(InNode->PreferredPermutationPairs, false);
-			ensure(PreferredPermutations.Num() == 1);
+			// ensure(PreferredPermutations.Num() == 1);
+			check(!PreferredPermutations.IsEmpty());
 			InNode->ResolvedPermutation = PreferredPermutations[0];
 		}			
 	}
