Control Rig: Fix bad for loop

diff --git a/Engine/Source/Developer/RigVMDeveloper/Private/RigVMModel/RigVMController.cpp b/Engine/Source/Developer/RigVMDeveloper/Private/RigVMModel/RigVMController.cpp
--- a/Engine/Source/Developer/RigVMDeveloper/Private/RigVMModel/RigVMController.cpp
+++ b/Engine/Source/Developer/RigVMDeveloper/Private/RigVMModel/RigVMController.cpp
@@ -17231,7 +17231,7 @@ bool URigVMController::AddPreferredType(URigVMTemplateNode* InNode, const FName&
 			}
 		}
 	}
-	for (int32 i=ToRemove.Num()-1; i>=0; ++i)
+	for (int32 i=ToRemove.Num()-1; i>=0; --i)
 	{
 		NewPreferredTypes.RemoveAt(ToRemove[i]);
 	}
