Control Rig: Improve BubbleUpExpressions performance

diff --git a/Engine/Source/Developer/RigVMDeveloper/Private/RigVMCompiler/RigVMAST.cpp b/Engine/Source/Developer/RigVMDeveloper/Private/RigVMCompiler/RigVMAST.cpp
--- a/Engine/Source/Developer/RigVMDeveloper/Private/RigVMCompiler/RigVMAST.cpp
+++ b/Engine/Source/Developer/RigVMDeveloper/Private/RigVMCompiler/RigVMAST.cpp
@@ -398,22 +398,38 @@ bool FRigVMBlockExprAST::ContainsEntry() const
 	return false;
 }
 
-bool FRigVMBlockExprAST::Contains(const FRigVMExprAST* InExpression) const
+bool FRigVMBlockExprAST::Contains(const FRigVMExprAST* InExpression, TMap<const FRigVMExprAST*, bool>* ContainedExpressionsCache) const
 {
 	if (InExpression == this)
 	{
 		return true;
 	}
 
+	if (ContainedExpressionsCache)
+	{
+		if (bool* Result = ContainedExpressionsCache->Find(InExpression))
+		{
+			return *Result;
+		}
+	}
+
 	for (int32 ParentIndex = 0; ParentIndex < InExpression->NumParents(); ParentIndex++)
 	{
 		const FRigVMExprAST* ParentExpr = InExpression->ParentAt(ParentIndex);
-		if (Contains(ParentExpr))
+		if (Contains(ParentExpr, ContainedExpressionsCache))
 		{
+			if (ContainedExpressionsCache)
+			{
+				ContainedExpressionsCache->Add(InExpression, true);
+			}
 			return true;
 		}
 	}
 
+	if (ContainedExpressionsCache)
+	{
+		ContainedExpressionsCache->Add(InExpression, false);
+	}
 	return false;
 }
 
@@ -1460,10 +1476,11 @@ void FRigVMParserAST::BubbleUpExpressions()
 				FRigVMBlockExprAST* BlockCandidate = BlockCandidates[BlockCandidateIndex];
 
 				bool bFoundCandidate = true;
+				TMap<const FRigVMExprAST*, bool> ContainedExpressionsCache;
 				for (int32 BlockIndex = 0; BlockIndex < Blocks.Num(); BlockIndex++)
 				{
 					FRigVMBlockExprAST* Block = Blocks[BlockIndex];
-					if (!BlockCandidate->Contains(Block))
+					if (!BlockCandidate->Contains(Block, &ContainedExpressionsCache))
 					{
 						bFoundCandidate = false;
 						break;
diff --git a/Engine/Source/Developer/RigVMDeveloper/Public/RigVMCompiler/RigVMAST.h b/Engine/Source/Developer/RigVMDeveloper/Public/RigVMCompiler/RigVMAST.h
--- a/Engine/Source/Developer/RigVMDeveloper/Public/RigVMCompiler/RigVMAST.h
+++ b/Engine/Source/Developer/RigVMDeveloper/Public/RigVMCompiler/RigVMAST.h
@@ -497,7 +497,7 @@ public:
 	// returns true if this block contains a given expression
 	// @param InExpression the expression to check
 	// @return true if this block contains a given expression
-	bool Contains(const FRigVMExprAST* InExpression) const;
+	bool Contains(const FRigVMExprAST* InExpression, TMap<const FRigVMExprAST*, bool>* ContainedExpressionsCache = nullptr) const;
 
 	// overload of the type checking mechanism
 	virtual bool IsA(EType InType) const override
