Control Rig: Fix incorrect array default values when connected to reroute node

diff --git a/Engine/Source/Developer/RigVMDeveloper/Private/RigVMCompiler/RigVMAST.cpp b/Engine/Source/Developer/RigVMDeveloper/Private/RigVMCompiler/RigVMAST.cpp
--- a/Engine/Source/Developer/RigVMDeveloper/Private/RigVMCompiler/RigVMAST.cpp
+++ b/Engine/Source/Developer/RigVMDeveloper/Private/RigVMCompiler/RigVMAST.cpp
@@ -3244,11 +3244,15 @@ void FRigVMParserAST::Inline(TArray<URigVMGraph*> InGraphs, const TArray<FRigVMA
 				}
 			}
 
+			// If the pin has subpins, and it has a source pin, the subpins can be safely ignored
 			URigVMPin* Pin = InPinProxy.GetSubjectChecked<URigVMPin>();
-			for (URigVMPin* SubPin : Pin->GetSubPins())
+			if (!SourcePinProxy.IsValid())
 			{
-				FRigVMASTProxy SubPinProxy = InPinProxy.GetSibling(SubPin);
-				VisitPin(SubPinProxy, OutTraversalInfo);
+				for (URigVMPin* SubPin : Pin->GetSubPins())
+				{
+					FRigVMASTProxy SubPinProxy = InPinProxy.GetSibling(SubPin);
+					VisitPin(SubPinProxy, OutTraversalInfo);
+				}
 			}
 		}
 
