Control Rig: Revert fix for incorrect array default values when connected to reroute node

diff --git a/Engine/Source/Developer/RigVMDeveloper/Private/RigVMCompiler/RigVMAST.cpp b/Engine/Source/Developer/RigVMDeveloper/Private/RigVMCompiler/RigVMAST.cpp
--- a/Engine/Source/Developer/RigVMDeveloper/Private/RigVMCompiler/RigVMAST.cpp
+++ b/Engine/Source/Developer/RigVMDeveloper/Private/RigVMCompiler/RigVMAST.cpp
@@ -3246,7 +3246,7 @@ void FRigVMParserAST::Inline(TArray<URigVMGraph*> InGraphs, const TArray<FRigVMA
 
 			// If the pin has subpins, and it has a source pin, the subpins can be safely ignored
 			URigVMPin* Pin = InPinProxy.GetSubjectChecked<URigVMPin>();
-			if (!SourcePinProxy.IsValid())
+			//if (!SourcePinProxy.IsValid())
 			{
 				for (URigVMPin* SubPin : Pin->GetSubPins())
 				{
