Control Rig: Second attempt to fix incorrect array default values when connected to reroute node#jira UE-163540

diff --git a/Engine/Source/Developer/RigVMDeveloper/Private/RigVMCompiler/RigVMAST.cpp b/Engine/Source/Developer/RigVMDeveloper/Private/RigVMCompiler/RigVMAST.cpp
--- a/Engine/Source/Developer/RigVMDeveloper/Private/RigVMCompiler/RigVMAST.cpp
+++ b/Engine/Source/Developer/RigVMDeveloper/Private/RigVMCompiler/RigVMAST.cpp
@@ -3244,9 +3244,9 @@ void FRigVMParserAST::Inline(TArray<URigVMGraph*> InGraphs, const TArray<FRigVMA
 				}
 			}
 
-			// If the pin has subpins, and it has a source pin, the subpins can be safely ignored
+			// If the pin is an array, and it has a source pin, the subpins can be safely ignored
 			URigVMPin* Pin = InPinProxy.GetSubjectChecked<URigVMPin>();
-			//if (!SourcePinProxy.IsValid())
+			if (!Pin->IsArray() || !SourcePinProxy.IsValid())
 			{
 				for (URigVMPin* SubPin : Pin->GetSubPins())
 				{
diff --git a/Engine/Source/Developer/RigVMDeveloper/Private/RigVMModel/RigVMController.cpp b/Engine/Source/Developer/RigVMDeveloper/Private/RigVMModel/RigVMController.cpp
--- a/Engine/Source/Developer/RigVMDeveloper/Private/RigVMModel/RigVMController.cpp
+++ b/Engine/Source/Developer/RigVMDeveloper/Private/RigVMModel/RigVMController.cpp
@@ -8564,6 +8564,11 @@ bool URigVMController::SetArrayPinSize(const FString& InArrayPinPath, int32 InSi
 		AddedPins++;
 	}
 
+	if (!InDefaultValue.IsEmpty() && InDefaultValue != TEXT("()"))
+	{
+		SetPinDefaultValue(Pin, InDefaultValue, false, bSetupUndoRedo, true, !bSuspendNotifications);
+	}
+
 	if (bSetupUndoRedo)
 	{
 		if (RemovedPins > 0 || AddedPins > 0)
