Control Rig: Fix crash when copy pasting function

diff --git a/Engine/Source/Developer/RigVMDeveloper/Private/RigVMModel/RigVMController.cpp b/Engine/Source/Developer/RigVMDeveloper/Private/RigVMModel/RigVMController.cpp
--- a/Engine/Source/Developer/RigVMDeveloper/Private/RigVMModel/RigVMController.cpp
+++ b/Engine/Source/Developer/RigVMDeveloper/Private/RigVMModel/RigVMController.cpp
@@ -3506,6 +3506,7 @@ TArray<FName> URigVMController::ImportNodesFromText(const FString& InText, bool
 
 		for(URigVMNode* SubNode : SubNodes)
 		{
+			FRigVMControllerGraphGuard GraphGuard(this, SubNode->GetGraph(), false);
 			if (URigVMUnitNode* UnitNode = Cast<URigVMUnitNode>(SubNode))
 			{
 				if (UnitNodeCreatedContext.IsValid())
