Control Rig: Fix compilation issue

diff --git a/Engine/Source/Developer/RigVMDeveloper/Private/RigVMModel/RigVMController.cpp b/Engine/Source/Developer/RigVMDeveloper/Private/RigVMModel/RigVMController.cpp
--- a/Engine/Source/Developer/RigVMDeveloper/Private/RigVMModel/RigVMController.cpp
+++ b/Engine/Source/Developer/RigVMDeveloper/Private/RigVMModel/RigVMController.cpp
@@ -3524,7 +3524,6 @@ TArray<FName> URigVMController::ImportNodesFromText(const FString& InText, bool
 			{
 				if (!TemplateNode->IsSingleton())
 				{
-					FRigVMControllerGraphGuard GraphGuard(this, TemplateNode->GetGraph(), false);
 					TemplateNode->InitializeFilteredPermutations();
 					UpdateTemplateNodePinTypes(TemplateNode, false, false);					
 				}
