ControlRig: fix graphguard not considering CancelAction. If an action was cancelled, guard would no longer need to pop graph in dtor

+ also we should register the pin type

diff --git a/Engine/Source/Developer/RigVMDeveloper/Private/RigVMModel/RigVMController.cpp b/Engine/Source/Developer/RigVMDeveloper/Private/RigVMModel/RigVMController.cpp
--- a/Engine/Source/Developer/RigVMDeveloper/Private/RigVMModel/RigVMController.cpp
+++ b/Engine/Source/Developer/RigVMDeveloper/Private/RigVMModel/RigVMController.cpp
@@ -18204,7 +18204,9 @@ bool URigVMController::ChangePinType(URigVMPin* InPin, const FString& InCPPType,
 	}
 
 	const FRigVMTemplateArgumentType Type(*InCPPType, InCPPTypeObject);
-	const TRigVMTypeIndex TypeIndex = FRigVMRegistry::Get().GetTypeIndex(Type);
+	// pin types are chosen from the graph pin type menu so it is not guaranteed that the chosen type
+	// is registered, hence the use of FindOrAddType
+	const TRigVMTypeIndex TypeIndex = FRigVMRegistry::Get().FindOrAddType(Type);
 	return ChangePinType(InPin, TypeIndex, bSetupUndoRedo, bSetupOrphanPins, bBreakLinks, bRemoveSubPins, bInitializeDefaultValue);
 }
 
diff --git a/Engine/Source/Developer/RigVMDeveloper/Public/RigVMModel/RigVMController.h b/Engine/Source/Developer/RigVMDeveloper/Public/RigVMModel/RigVMController.h
--- a/Engine/Source/Developer/RigVMDeveloper/Public/RigVMModel/RigVMController.h
+++ b/Engine/Source/Developer/RigVMDeveloper/Public/RigVMModel/RigVMController.h
@@ -1233,6 +1233,7 @@ private:
 	friend struct FRigVMInjectNodeIntoPinAction;
 	friend class FRigVMParserAST;
 	friend class FRigVMControllerCompileBracketScope;
+	friend class FRigVMControllerGraphGuard;
 };
 
 class FRigVMControllerGraphGuard
@@ -1244,10 +1245,17 @@ public:
 		, bUndo(bSetupUndoRedo)
 	{
 		Controller->PushGraph(InGraph, bUndo);
+		NumGraphs = Controller->Graphs.Num();
 	}
 
 	~FRigVMControllerGraphGuard()
 	{
+		// an action can be cancelled in the middle of a graph guard,
+		// in that case CancelAction should have already popped the graph
+		if (Controller->Graphs.Num() < NumGraphs)
+		{
+			return;
+		}
 		Controller->PopGraph(bUndo);
 	}
 
@@ -1255,6 +1263,8 @@ private:
 
 	URigVMController* Controller;
 	bool bUndo;
+
+	int32 NumGraphs;
 };
 
 USTRUCT()
