controlrig: fix a place that was using GetStructCPPName, for control rig it should use GetUniqueStructTypeName

diff --git a/Engine/Source/Developer/RigVMDeveloper/Private/RigVMDeveloperTypeUtils.cpp b/Engine/Source/Developer/RigVMDeveloper/Private/RigVMDeveloperTypeUtils.cpp
--- a/Engine/Source/Developer/RigVMDeveloper/Private/RigVMDeveloperTypeUtils.cpp
+++ b/Engine/Source/Developer/RigVMDeveloper/Private/RigVMDeveloperTypeUtils.cpp
@@ -590,7 +590,7 @@ bool RigVMTypeUtils::CPPTypeFromPinType(const FEdGraphPinType& InPinType, FStrin
 	{
 		if (UScriptStruct* Struct = Cast<UScriptStruct>(InPinType.PinSubCategoryObject))
 		{
-			OutCPPType = Prefix + *Struct->GetStructCPPName() + Suffix;
+			OutCPPType = Prefix + GetUniqueStructTypeName(Struct) + Suffix;
 			*OutCPPTypeObject = Struct;
 		}
 	}
