Control Rig: Fix crash when copy-pasting function with nested functions from one engine to another

diff --git a/Engine/Source/Developer/RigVMDeveloper/Private/RigVMModel/RigVMController.cpp b/Engine/Source/Developer/RigVMDeveloper/Private/RigVMModel/RigVMController.cpp
--- a/Engine/Source/Developer/RigVMDeveloper/Private/RigVMModel/RigVMController.cpp
+++ b/Engine/Source/Developer/RigVMDeveloper/Private/RigVMModel/RigVMController.cpp
@@ -3490,6 +3490,12 @@ TArray<FName> URigVMController::ImportNodesFromText(const FString& InText, bool
 						Recompute(ContainedNode);
 					}
 				}
+
+				// If the OuterNode is referencing a non-existing function, the contained graph might be empty
+				if (OuterNode->GetContainedGraph() == nullptr)
+				{
+					return;
+				}
 				
 				{
 					FRigVMControllerGraphGuard GraphGuard(this, OuterNode->GetGraph(), false);
