Control Rig: Fix crash when connecting subpin from entry node to cast, and resolving type of cast output

diff --git a/Engine/Source/Developer/RigVMDeveloper/Private/RigVMModel/RigVMController.cpp b/Engine/Source/Developer/RigVMDeveloper/Private/RigVMModel/RigVMController.cpp
--- a/Engine/Source/Developer/RigVMDeveloper/Private/RigVMModel/RigVMController.cpp
+++ b/Engine/Source/Developer/RigVMDeveloper/Private/RigVMModel/RigVMController.cpp
@@ -17099,7 +17099,8 @@ bool URigVMController::PropagateTemplateFilteredTypes(URigVMTemplateNode* InNode
 			{
 				if (OtherTemplate->IsA<URigVMFunctionEntryNode>() || OtherTemplate->IsA<URigVMFunctionReturnNode>())
 				{
-					if (!OtherTemplate->GetTemplate()->FindArgument(OtherPin->GetFName()))
+					const URigVMPin* OtherPinRoot = OtherPin->GetRootPin();
+					if (!OtherTemplate->GetTemplate()->FindArgument(OtherPinRoot->GetFName()))
 					{
 						if (!AddArgumentForPin(OtherPin, Pin, bSetupUndoRedo))
 						{
