Control Rig: Do not allow IO exposed pin type change to something other than execute context

diff --git a/Engine/Source/Developer/RigVMDeveloper/Private/RigVMModel/RigVMController.cpp b/Engine/Source/Developer/RigVMDeveloper/Private/RigVMModel/RigVMController.cpp
--- a/Engine/Source/Developer/RigVMDeveloper/Private/RigVMModel/RigVMController.cpp
+++ b/Engine/Source/Developer/RigVMDeveloper/Private/RigVMModel/RigVMController.cpp
@@ -10750,6 +10750,26 @@ bool URigVMController::ChangeExposedPinType(const FName& InPinName, const FStrin
 		}
 	}
 
+	if (Pin->GetDirection() == ERigVMPinDirection::IO)
+	{
+		bool bIsExecute = false;
+		if (CPPTypeObject)
+		{
+			if(const UScriptStruct* CPPTypeStruct = Cast<UScriptStruct>(CPPTypeObject))
+			{
+				if(CPPTypeStruct->IsChildOf(FRigVMExecuteContext::StaticStruct()))
+				{
+					bIsExecute = true;
+				}
+			}
+		}
+		if (!bIsExecute)
+		{
+			ReportAndNotifyError(TEXT("Input/Output pins only allow Execute Context types."));
+			return false;
+		}
+	}
+
 	if(bSetupUndoRedo)
 	{
 		if(RequestBulkEditDialogDelegate.IsBound())
