Control Rig: Fix add link between array and element of Select values does not resolve select pin correctly

diff --git a/Engine/Source/Developer/RigVMDeveloper/Private/RigVMModel/RigVMController.cpp b/Engine/Source/Developer/RigVMDeveloper/Private/RigVMModel/RigVMController.cpp
--- a/Engine/Source/Developer/RigVMDeveloper/Private/RigVMModel/RigVMController.cpp
+++ b/Engine/Source/Developer/RigVMDeveloper/Private/RigVMModel/RigVMController.cpp
@@ -9433,6 +9433,17 @@ bool URigVMController::AddLink(URigVMPin* OutputPin, URigVMPin* InputPin, bool b
 		ActionStack->EndAction(Action);
 	}
 
+	if (!bIsTransacting)
+	{
+		ensureMsgf(RigVMTypeUtils::AreCompatible(*OutputPin->GetCPPType(), OutputPin->GetCPPTypeObject(), *InputPin->GetCPPType(), InputPin->GetCPPTypeObject()),
+		   TEXT("Incompatible types after successful link %s (%s) -> %s (%s) created in %s")
+		   , *OutputPin->GetPinPath(true)
+		   , *OutputPin->GetCPPType()
+		   , *InputPin->GetPinPath(true)
+		   , *InputPin->GetCPPType()
+		   , *GetPackage()->GetPathName());
+	}
+
 	return true;
 }
 
@@ -18272,16 +18283,20 @@ bool URigVMController::ChangePinType(URigVMPin* InPin, TRigVMTypeIndex InTypeInd
 			}
 		}
 
-		// change the pin's type to be of an array as well
-		const bool bIsArrayType = FRigVMRegistry::Get().IsArrayType(TypeIndex);
-		if(InPin->IsRootPin() && bIsArrayType != InPin->IsArray())
-		{
-			// nothing to do here - leave the type as is 
-		}
-		else
+		// If changing to wildcard, try to maintain the container type
+		const FRigVMRegistry& Registry = FRigVMRegistry::Get();
+		if (Registry.IsWildCardType(TypeIndex)) 
 		{
-			const TRigVMTypeIndex BaseTypeIndex = bIsArrayType ? FRigVMRegistry::Get().GetBaseTypeFromArrayTypeIndex(TypeIndex) : TypeIndex;
-			TypeIndex = InPin->IsArray() ? FRigVMRegistry::Get().GetArrayTypeFromBaseTypeIndex(BaseTypeIndex) : BaseTypeIndex;
+			const bool bIsArrayType = FRigVMRegistry::Get().IsArrayType(TypeIndex);
+			if(InPin->IsRootPin() && bIsArrayType != InPin->IsArray())
+			{
+				// nothing to do here - leave the type as is 
+			}
+			else
+			{
+				const TRigVMTypeIndex BaseTypeIndex = bIsArrayType ? FRigVMRegistry::Get().GetBaseTypeFromArrayTypeIndex(TypeIndex) : TypeIndex;
+				TypeIndex = InPin->IsArray() ? Registry.GetArrayTypeFromBaseTypeIndex(BaseTypeIndex) : BaseTypeIndex;
+			}
 		}
 	}
 
