Control Rig: Fix crash when resolving template to different type by adding a link

diff --git a/Engine/Source/Developer/RigVMDeveloper/Private/RigVMModel/RigVMController.cpp b/Engine/Source/Developer/RigVMDeveloper/Private/RigVMModel/RigVMController.cpp
--- a/Engine/Source/Developer/RigVMDeveloper/Private/RigVMModel/RigVMController.cpp
+++ b/Engine/Source/Developer/RigVMDeveloper/Private/RigVMModel/RigVMController.cpp
@@ -16974,8 +16974,15 @@ bool URigVMController::UpdateTemplateNodePinTypes(URigVMTemplateNode* InNode, bo
 			{
 				continue;
 			}
-		
-			if (FinalPinTypes[PinIndex] == INDEX_NONE)
+
+			bool bShouldUnresolve = FinalPinTypes[PinIndex] == INDEX_NONE;
+			// If we are about to resolve the pin to a different type, first unresolve it
+			if (!bShouldUnresolve && UnresolveResolve == 0 && Pin->GetTypeIndex() != FinalPinTypes[PinIndex])
+			{
+				bShouldUnresolve = true;
+			}
+				
+			if (bShouldUnresolve)
 			{
 				// Unresolve
 				if (Pin->HasInjectedNodes())
