Control Rig: Fix undo rename local variable not working

diff --git a/Engine/Source/Developer/RigVMDeveloper/Private/RigVMModel/RigVMController.cpp b/Engine/Source/Developer/RigVMDeveloper/Private/RigVMModel/RigVMController.cpp
--- a/Engine/Source/Developer/RigVMDeveloper/Private/RigVMModel/RigVMController.cpp
+++ b/Engine/Source/Developer/RigVMDeveloper/Private/RigVMModel/RigVMController.cpp
@@ -11719,13 +11719,13 @@ bool URigVMController::RenameLocalVariable(const FName& InVariableName, const FN
 	FRigVMControllerCompileBracketScope CompileScope(this);
 	if (bSetupUndoRedo)
 	{
-		FRigVMInverseAction InverseAction;
-		InverseAction.Title = FString::Printf(TEXT("Rename Local Variable %s to %s"), *InVariableName.ToString(), *InNewVariableName.ToString());
+		FRigVMBaseAction BaseAction;
+		BaseAction.Title = FString::Printf(TEXT("Rename Local Variable %s to %s"), *InVariableName.ToString(), *InNewVariableName.ToString());
 
-		ActionStack->BeginAction(InverseAction);
+		ActionStack->BeginAction(BaseAction);
 		ActionStack->AddAction(FRigVMRenameLocalVariableAction(LocalVariables[FoundIndex].Name, InNewVariableName));
-		ActionStack->EndAction(InverseAction);
-	}	
+		ActionStack->EndAction(BaseAction);
+	}
 	
 	LocalVariables[FoundIndex].Name = InNewVariableName;
 
