ControlRig: Fix Enum default getting overriden on load

diff --git a/Engine/Source/Developer/RigVMDeveloper/Private/RigVMModel/RigVMController.cpp b/Engine/Source/Developer/RigVMDeveloper/Private/RigVMModel/RigVMController.cpp
--- a/Engine/Source/Developer/RigVMDeveloper/Private/RigVMModel/RigVMController.cpp
+++ b/Engine/Source/Developer/RigVMDeveloper/Private/RigVMModel/RigVMController.cpp
@@ -15649,6 +15649,29 @@ void URigVMController::PostProcessDefaultValue(URigVMPin* Pin, FString& OutDefau
 	{
 		OutDefaultValue = TEXT("()");
 	}
+	else if (Pin->IsEnum() && OutDefaultValue.IsEmpty())
+	{
+		int32 EnumIndex = Pin->GetEnum()->GetIndexByName(*NoneString);
+		// make sure that none is a valid enum value
+		if (EnumIndex != INDEX_NONE)
+		{
+			OutDefaultValue = NoneString;
+		}
+		else
+		{
+			// in the rare case that none is not a valid enum value
+			// there isn't a good way to resolve it, so here we just use the
+			// enum value at index 0 and provide a warning
+			OutDefaultValue = Pin->GetEnum()->GetNameStringByIndex(0);
+			UE_LOG(LogRigVMDeveloper, Warning, TEXT("Enum Pin %s had None as default value, but None is not a valid value for %s, new value: %s"),
+					*Pin->GetFullName(), *Pin->GetEnum()->GetFullName(), *OutDefaultValue);
+			
+			// an alternative is to leave the default as empty so that it uses whatever
+			// default value that the rig unit gave this pin when repopulate pin
+			// however it might not be a good idea since it might cause old assets to change
+			// whenever the defaults of a rig unit changed.
+		}
+	}
 	else if (Pin->IsStruct() && (OutDefaultValue.IsEmpty() || OutDefaultValue == TEXT("()")))
 	{
 		CreateDefaultValueForStructIfRequired(Pin->GetScriptStruct(), OutDefaultValue);
diff --git a/Engine/Source/Developer/RigVMDeveloper/Private/RigVMModel/RigVMPin.cpp b/Engine/Source/Developer/RigVMDeveloper/Private/RigVMModel/RigVMPin.cpp
--- a/Engine/Source/Developer/RigVMDeveloper/Private/RigVMModel/RigVMPin.cpp
+++ b/Engine/Source/Developer/RigVMDeveloper/Private/RigVMModel/RigVMPin.cpp
@@ -501,6 +501,15 @@ bool URigVMPin::RequiresWatch(const bool bCheckExposedPinChain) const
 	return bRequiresWatch;
 }
 
+bool URigVMPin::IsEnum() const
+{
+	if (IsArray())
+	{
+		return false;
+	}
+	return GetEnum() != nullptr;
+}
+
 bool URigVMPin::IsStruct() const
 {
 	if (IsArray())
diff --git a/Engine/Source/Developer/RigVMDeveloper/Public/RigVMModel/RigVMPin.h b/Engine/Source/Developer/RigVMDeveloper/Public/RigVMModel/RigVMPin.h
--- a/Engine/Source/Developer/RigVMDeveloper/Public/RigVMModel/RigVMPin.h
+++ b/Engine/Source/Developer/RigVMDeveloper/Public/RigVMModel/RigVMPin.h
@@ -186,6 +186,10 @@ public:
 	// Returns true if the pin should be watched
 	UFUNCTION(BlueprintCallable, Category = RigVMPin)
 	bool RequiresWatch(const bool bCheckExposedPinChain = false) const;
+	
+	// Returns true if the data type of the Pin is a enum
+	UFUNCTION(BlueprintCallable, Category = RigVMPin)
+	bool IsEnum() const;
 
 	// Returns true if the data type of the Pin is a struct
 	UFUNCTION(BlueprintCallable, Category = RigVMPin)
