Control Rig: Fix undo delete select node fails

diff --git a/Engine/Source/Developer/RigVMDeveloper/Private/RigVMModel/RigVMControllerActions.cpp b/Engine/Source/Developer/RigVMDeveloper/Private/RigVMModel/RigVMControllerActions.cpp
--- a/Engine/Source/Developer/RigVMDeveloper/Private/RigVMModel/RigVMControllerActions.cpp
+++ b/Engine/Source/Developer/RigVMDeveloper/Private/RigVMModel/RigVMControllerActions.cpp
@@ -743,10 +743,16 @@ FRigVMRemoveNodeAction::FRigVMRemoveNodeAction(URigVMNode* InNode, URigVMControl
 	else if (URigVMIfNode* IfNode = Cast<URigVMIfNode>(InNode))
 	{
 		InverseAction.AddAction(FRigVMAddIfNodeAction(IfNode), InController);
+		URigVMPin* TrueNamePin = IfNode->FindPin(URigVMIfNode::TrueName);
+		URigVMPin* FalseNamePin = IfNode->FindPin(URigVMIfNode::FalseName);
+		InverseAction.AddAction(FRigVMSetPinDefaultValueAction(TrueNamePin, TrueNamePin->GetDefaultValue()), InController);
+		InverseAction.AddAction(FRigVMSetPinDefaultValueAction(FalseNamePin, FalseNamePin->GetDefaultValue()), InController);
 	}
 	else if (URigVMSelectNode* SelectNode = Cast<URigVMSelectNode>(InNode))
 	{
 		InverseAction.AddAction(FRigVMAddSelectNodeAction(SelectNode), InController);
+		URigVMPin* ValuesNamePin = SelectNode->FindPin(URigVMSelectNode::ValueName);
+		InverseAction.AddAction(FRigVMSetPinDefaultValueAction(ValuesNamePin, ValuesNamePin->GetDefaultValue()), InController);		
 	}
 	else if (URigVMEnumNode* EnumNode = Cast<URigVMEnumNode>(InNode))
 	{
diff --git a/Engine/Source/Developer/RigVMDeveloper/Public/RigVMModel/Nodes/RigVMIfNode.h b/Engine/Source/Developer/RigVMDeveloper/Public/RigVMModel/Nodes/RigVMIfNode.h
--- a/Engine/Source/Developer/RigVMDeveloper/Public/RigVMModel/Nodes/RigVMIfNode.h
+++ b/Engine/Source/Developer/RigVMDeveloper/Public/RigVMModel/Nodes/RigVMIfNode.h
@@ -45,5 +45,6 @@ private:
 	friend struct FRigVMAddIfNodeAction;
 	friend class UControlRigIfNodeSpawner;
 	friend class FRigVMParserAST;
+	friend struct FRigVMRemoveNodeAction;
 };
 
diff --git a/Engine/Source/Developer/RigVMDeveloper/Public/RigVMModel/Nodes/RigVMSelectNode.h b/Engine/Source/Developer/RigVMDeveloper/Public/RigVMModel/Nodes/RigVMSelectNode.h
--- a/Engine/Source/Developer/RigVMDeveloper/Public/RigVMModel/Nodes/RigVMSelectNode.h
+++ b/Engine/Source/Developer/RigVMDeveloper/Public/RigVMModel/Nodes/RigVMSelectNode.h
@@ -51,5 +51,6 @@ private:
 	friend class UControlRigSelectNodeSpawner;
 	friend class FRigVMParserAST;
 	friend class FRigVMSelectExprAST;
+	friend struct FRigVMRemoveNodeAction;
 };
 
