Control Rig:  Avoid crashing when injection info is null

diff --git a/Engine/Source/Developer/RigVMDeveloper/Private/RigVMModel/RigVMController.cpp b/Engine/Source/Developer/RigVMDeveloper/Private/RigVMModel/RigVMController.cpp
--- a/Engine/Source/Developer/RigVMDeveloper/Private/RigVMModel/RigVMController.cpp
+++ b/Engine/Source/Developer/RigVMDeveloper/Private/RigVMModel/RigVMController.cpp
@@ -15386,10 +15386,13 @@ void URigVMController::ApplyPinState(URigVMPin* InPin, const FPinState& InPinSta
 {
 	for (URigVMInjectionInfo* InjectionInfo : InPinState.InjectionInfos)
 	{
-		RenameObject(InjectionInfo, nullptr, InPin);
-		InjectionInfo->InputPin = InjectionInfo->InputPin ? InjectionInfo->Node->FindPin(InjectionInfo->InputPin->GetName()) : nullptr;
-		InjectionInfo->OutputPin = InjectionInfo->OutputPin ? InjectionInfo->Node->FindPin(InjectionInfo->OutputPin->GetName()) : nullptr;
-		InPin->InjectionInfos.Add(InjectionInfo);
+		if (InjectionInfo)
+		{
+			RenameObject(InjectionInfo, nullptr, InPin);
+			InjectionInfo->InputPin = InjectionInfo->InputPin ? InjectionInfo->Node->FindPin(InjectionInfo->InputPin->GetName()) : nullptr;
+			InjectionInfo->OutputPin = InjectionInfo->OutputPin ? InjectionInfo->Node->FindPin(InjectionInfo->OutputPin->GetName()) : nullptr;
+			InPin->InjectionInfos.Add(InjectionInfo);
+		}
 	}
 
 	// alternatively if the injection infos are not provided as strong pointers
