Control Rig: Add even more assurance that we raise compilation issues when ResolvedFunctionName is not set

diff --git a/Engine/Source/Developer/RigVMDeveloper/Private/RigVMCompiler/RigVMCompiler.cpp b/Engine/Source/Developer/RigVMDeveloper/Private/RigVMCompiler/RigVMCompiler.cpp
--- a/Engine/Source/Developer/RigVMDeveloper/Private/RigVMCompiler/RigVMCompiler.cpp
+++ b/Engine/Source/Developer/RigVMDeveloper/Private/RigVMCompiler/RigVMCompiler.cpp
@@ -348,7 +348,7 @@ bool URigVMCompiler::Compile(TArray<URigVMGraph*> InGraphs, URigVMController* In
 						InController->FullyResolveTemplateNode(UnitNode, INDEX_NONE, false);
 					}
 
-					if (UnitNode->GetScriptStruct() == nullptr)
+					if (UnitNode->GetScriptStruct() == nullptr || UnitNode->ResolvedFunctionName.IsEmpty())
 					{
 						static const FString UnresolvedUnitNodeMessage = TEXT("Node @@ could not be resolved.");
 						Settings.ASTSettings.Report(EMessageSeverity::Error, ModelNode, UnresolvedUnitNodeMessage);
diff --git a/Engine/Source/Developer/RigVMDeveloper/Public/RigVMModel/Nodes/RigVMTemplateNode.h b/Engine/Source/Developer/RigVMDeveloper/Public/RigVMModel/Nodes/RigVMTemplateNode.h
--- a/Engine/Source/Developer/RigVMDeveloper/Public/RigVMModel/Nodes/RigVMTemplateNode.h
+++ b/Engine/Source/Developer/RigVMDeveloper/Public/RigVMModel/Nodes/RigVMTemplateNode.h
@@ -202,5 +202,6 @@ protected:
 	friend struct FRigVMSetPreferredTemplatePermutationsAction;
 	friend struct FRigVMRemoveNodeAction;
 	friend class FRigVMTemplatesLibraryNode3Test;
+	friend class URigVMCompiler;
 };
 
