ControlRig: Remove an unnecessary warning

diff --git a/Engine/Source/Developer/RigVMDeveloper/Private/RigVMModel/RigVMController.cpp b/Engine/Source/Developer/RigVMDeveloper/Private/RigVMModel/RigVMController.cpp
--- a/Engine/Source/Developer/RigVMDeveloper/Private/RigVMModel/RigVMController.cpp
+++ b/Engine/Source/Developer/RigVMDeveloper/Private/RigVMModel/RigVMController.cpp
@@ -15662,17 +15662,11 @@ void URigVMController::PostProcessDefaultValue(URigVMPin* Pin, FString& OutDefau
 		}
 		else
 		{
-			// in the rare case that none is not a valid enum value
-			// there isn't a good way to resolve it, so here we just use the
-			// enum value at index 0 and provide a warning
+			// when none string is given but none is not a valid enum value,
+			// it implies that the pin should use the default value of the enum
+			// the value at index 0 is the best guess that can be made.
+			// usually a user provided pin default is given later to override this value
 			OutDefaultValue = Pin->GetEnum()->GetNameStringByIndex(0);
-			UE_LOG(LogRigVMDeveloper, Warning, TEXT("Enum Pin %s had None as default value, but None is not a valid value for %s, new value: %s"),
-					*Pin->GetFullName(), *Pin->GetEnum()->GetFullName(), *OutDefaultValue);
-			
-			// an alternative is to leave the default as empty so that it uses whatever
-			// default value that the rig unit gave this pin when repopulate pin
-			// however it might not be a good idea since it might cause old assets to change
-			// whenever the defaults of a rig unit changed.
 		}
 	}
 	else if (Pin->IsStruct() && (OutDefaultValue.IsEmpty() || OutDefaultValue == TEXT("()")))
