ControlRig: user defined struct used as function argument can trigger an assertion when one of the member variables change its type

diff --git a/Engine/Plugins/Animation/ControlRig/Source/ControlRigDeveloper/Private/ControlRigBlueprint.cpp b/Engine/Plugins/Animation/ControlRig/Source/ControlRigDeveloper/Private/ControlRigBlueprint.cpp
--- a/Engine/Plugins/Animation/ControlRig/Source/ControlRigDeveloper/Private/ControlRigBlueprint.cpp
+++ b/Engine/Plugins/Animation/ControlRig/Source/ControlRigDeveloper/Private/ControlRigBlueprint.cpp
@@ -1231,8 +1231,29 @@ void UControlRigBlueprint::RecompileVM()
 		UserData.Add(FRigVMUserDataArray(&InitContextPtr, 1));
 		UserData.Add(FRigVMUserDataArray(&UpdateContextPtr, 1));
 
+		// Clear all Errors
 		CompileLog.Messages.Reset();
 		CompileLog.NumErrors = CompileLog.NumWarnings = 0;
+		
+		TArray<UEdGraph*> EdGraphs;
+		GetAllGraphs(EdGraphs);
+
+		for (UEdGraph* Graph : EdGraphs)
+		{
+			UControlRigGraph* RigGraph = Cast<UControlRigGraph>(Graph);
+			if (RigGraph == nullptr)
+			{
+				continue;
+			}
+
+			for (UEdGraphNode* GraphNode : Graph->Nodes)
+			{
+				if (UControlRigGraphNode* ControlRigGraphNode = Cast<UControlRigGraphNode>(GraphNode))
+				{
+					ControlRigGraphNode->ClearErrorInfo();
+				}
+			}
+		}
 
 		URigVMCompiler* Compiler = URigVMCompiler::StaticClass()->GetDefaultObject<URigVMCompiler>();
 		Compiler->Settings = (bCompileInDebugMode) ? FRigVMCompileSettings::Fast() : VMCompileSettings;
@@ -1252,28 +1273,6 @@ void UControlRigBlueprint::RecompileVM()
 			}
 			return;
 		}
-		else
-		{
-			TArray<UEdGraph*> EdGraphs;
-			GetAllGraphs(EdGraphs);
-
-			for (UEdGraph* Graph : EdGraphs)
-			{
-				UControlRigGraph* RigGraph = Cast<UControlRigGraph>(Graph);
-				if (RigGraph == nullptr)
-				{
-					continue;
-				}
-
-				for (UEdGraphNode* GraphNode : Graph->Nodes)
-				{
-					if (UControlRigGraphNode* ControlRigGraphNode = Cast<UControlRigGraphNode>(GraphNode))
-					{
-						ControlRigGraphNode->bHasCompilerMessage = false;
-					}
-				}
-			}
-		}
 
 		TArray<UObject*> ArchetypeInstances;
 		CDO->GetArchetypeInstances(ArchetypeInstances);
diff --git a/Engine/Plugins/Animation/ControlRig/Source/ControlRigDeveloper/Private/Graph/ControlRigGraphNode.cpp b/Engine/Plugins/Animation/ControlRig/Source/ControlRigDeveloper/Private/Graph/ControlRigGraphNode.cpp
--- a/Engine/Plugins/Animation/ControlRig/Source/ControlRigDeveloper/Private/Graph/ControlRigGraphNode.cpp
+++ b/Engine/Plugins/Animation/ControlRig/Source/ControlRigDeveloper/Private/Graph/ControlRigGraphNode.cpp
@@ -458,6 +458,15 @@ const FRigVMTemplate* UControlRigGraphNode::GetTemplate() const
 	return CachedTemplate;
 }
 
+void UControlRigGraphNode::ClearErrorInfo()
+{
+	bHasCompilerMessage = false;
+	// SControlRigGraphNode only updates if the error types do not match so we have
+	// clear the error type as well, see SControlRigGraphNode::RefreshErrorInfo()
+	ErrorType = (int32)EMessageSeverity::Info + 1;
+	ErrorMsg = FString();	
+}
+
 FLinearColor UControlRigGraphNode::GetNodeProfilingColor() const
 {
 #if WITH_EDITOR
diff --git a/Engine/Plugins/Animation/ControlRig/Source/ControlRigDeveloper/Public/Graph/ControlRigGraphNode.h b/Engine/Plugins/Animation/ControlRig/Source/ControlRigDeveloper/Public/Graph/ControlRigGraphNode.h
--- a/Engine/Plugins/Animation/ControlRig/Source/ControlRigDeveloper/Public/Graph/ControlRigGraphNode.h
+++ b/Engine/Plugins/Animation/ControlRig/Source/ControlRigDeveloper/Public/Graph/ControlRigGraphNode.h
@@ -157,6 +157,8 @@ public:
 
 	const FRigVMTemplate* GetTemplate() const;
 
+	void ClearErrorInfo();
+
 protected:
 
 	FLinearColor GetNodeProfilingColor() const;
diff --git a/Engine/Source/Developer/RigVMDeveloper/Private/RigVMModel/RigVMController.cpp b/Engine/Source/Developer/RigVMDeveloper/Private/RigVMModel/RigVMController.cpp
--- a/Engine/Source/Developer/RigVMDeveloper/Private/RigVMModel/RigVMController.cpp
+++ b/Engine/Source/Developer/RigVMDeveloper/Private/RigVMModel/RigVMController.cpp
@@ -9596,6 +9596,11 @@ bool URigVMController::BreakLink(URigVMPin* OutputPin, URigVMPin* InputPin, bool
 			InputPin->Links.Remove(Link);
 			Graph->Links.Remove(Link);
 
+			// each time a link is removed, existing orphaned pins
+			// may become unused and thus can be removed
+			RemoveUnusedOrphanedPins(OutputPin->GetNode(), true);
+			RemoveUnusedOrphanedPins(InputPin->GetNode(), true);
+			
 			if (!bIsTransacting && !bSuspendRecomputingTemplateFilters)
 			{
 				RecomputeAllTemplateFilteredPermutations(bSetupUndoRedo);
@@ -17937,10 +17942,13 @@ bool URigVMController::PrepareToLink(URigVMPin* FirstToResolve, URigVMPin* Secon
 		if (FirstTemplateNode->IsA<URigVMFunctionEntryNode>() || FirstTemplateNode->IsA<URigVMFunctionReturnNode>())
 		{
 			URigVMPin* RootPin = FirstToResolve->GetRootPin();
-			if (!FirstTemplateNode->GetTemplate()->FindArgument(RootPin->GetFName()))
+			if (!RootPin->IsOrphanPin())
 			{
-				ensure(RootPin == FirstToResolve);
-				bShouldPrepare = AddArgumentForPin(FirstToResolve, SecondToResolve, bSetupUndoRedo);
+				if (!FirstTemplateNode->GetTemplate()->FindArgument(RootPin->GetFName()))
+				{
+					ensure(RootPin == FirstToResolve);
+					bShouldPrepare = AddArgumentForPin(FirstToResolve, SecondToResolve, bSetupUndoRedo);
+				}
 			}
 		}
 		if (!FirstTemplateNode->IsSingleton() && bShouldPrepare)
@@ -17961,10 +17969,13 @@ bool URigVMController::PrepareToLink(URigVMPin* FirstToResolve, URigVMPin* Secon
 		if (SecondTemplateNode->IsA<URigVMFunctionEntryNode>() || SecondTemplateNode->IsA<URigVMFunctionReturnNode>())
 		{
 			URigVMPin* RootPin = SecondToResolve->GetRootPin();
-			if (!SecondTemplateNode->GetTemplate()->FindArgument(RootPin->GetFName()))
+			if (!RootPin->IsOrphanPin())
 			{
-				ensure(RootPin == SecondToResolve);
-				bShouldPrepare = AddArgumentForPin(SecondToResolve, FirstToResolve, bSetupUndoRedo);
+				if (!SecondTemplateNode->GetTemplate()->FindArgument(RootPin->GetFName()))
+				{
+					ensure(RootPin == SecondToResolve);
+					bShouldPrepare = AddArgumentForPin(SecondToResolve, FirstToResolve, bSetupUndoRedo);
+				}
 			}
 		}
 		if (!SecondTemplateNode->IsSingleton() && bShouldPrepare)
