Control Rig: Avoid double initialization of struct instances

diff --git a/Engine/Source/Developer/RigVMDeveloper/Private/RigVMModel/Nodes/RigVMUnitNode.cpp b/Engine/Source/Developer/RigVMDeveloper/Private/RigVMModel/Nodes/RigVMUnitNode.cpp
--- a/Engine/Source/Developer/RigVMDeveloper/Private/RigVMModel/Nodes/RigVMUnitNode.cpp
+++ b/Engine/Source/Developer/RigVMDeveloper/Private/RigVMModel/Nodes/RigVMUnitNode.cpp
@@ -317,11 +317,7 @@ TSharedPtr<FStructOnScope> URigVMUnitNode::ConstructStructInstance(bool bUseDefa
 	{
 		TSharedPtr<FStructOnScope> StructOnScope = MakeShareable(new FStructOnScope(Struct));
 		FRigVMStruct* StructMemory = (FRigVMStruct*)StructOnScope->GetStructMemory();
-		if (bUseDefault)
-		{
-			Struct->InitializeDefaultValue((uint8*)StructMemory);
-		}
-		else
+		if (!bUseDefault)
 		{
 			FString StructDefaultValue = GetStructDefaultValue();
 			Struct->ImportText(*StructDefaultValue, StructMemory, nullptr, PPF_IncludeTransient, GLog, Struct->GetName());
diff --git a/Engine/Source/Developer/RigVMDeveloper/Private/RigVMModel/RigVMController.cpp b/Engine/Source/Developer/RigVMDeveloper/Private/RigVMModel/RigVMController.cpp
--- a/Engine/Source/Developer/RigVMDeveloper/Private/RigVMModel/RigVMController.cpp
+++ b/Engine/Source/Developer/RigVMDeveloper/Private/RigVMModel/RigVMController.cpp
@@ -865,7 +865,6 @@ URigVMUnitNode* URigVMController::AddUnitNode(UScriptStruct* InScriptStruct, con
 
 	FStructOnScope StructOnScope(InScriptStruct);
 	FRigVMStruct* StructMemory = (FRigVMStruct*)StructOnScope.GetStructMemory();
-	InScriptStruct->InitializeDefaultValue((uint8*)StructMemory);
 	const bool bIsEventNode = (!StructMemory->GetEventName().IsNone());
 	if (bIsEventNode)
 	{
@@ -7555,7 +7554,6 @@ FString URigVMController::GetPinInitialDefaultValueFromStruct(UScriptStruct* Scr
 	{
 		TSharedPtr<FStructOnScope> StructOnScope = MakeShareable(new FStructOnScope(ScriptStruct));
 		uint8* Memory = (uint8*)StructOnScope->GetStructMemory();
-		ScriptStruct->InitializeDefaultValue(Memory);
 
 		if (InPin->GetScriptStruct() == ScriptStruct)
 		{
@@ -7598,7 +7596,6 @@ FString URigVMController::GetPinInitialDefaultValueFromStruct(UScriptStruct* Scr
 					UScriptStruct* InnerStruct = StructProperty->Struct;
 					StructOnScope = MakeShareable(new FStructOnScope(InnerStruct));
 					Memory = (uint8 *)StructOnScope->GetStructMemory();
-					InnerStruct->InitializeDefaultValue(Memory);
 				}
 				continue;
 			}
@@ -13686,7 +13683,6 @@ TObjectPtr<URigVMNode> URigVMController::FindEventNode(const UScriptStruct* InSc
 
 	// construct equivalent default struct
 	FStructOnScope InDefaultStructScope(InScriptStruct);
-	InScriptStruct->InitializeDefaultValue((uint8*)InDefaultStructScope.GetStructMemory());
 	
 	if (URigVMGraph* Graph = GetGraph())
 	{
@@ -18566,7 +18562,6 @@ bool URigVMController::ChangePinType(URigVMPin* InPin, TRigVMTypeIndex InTypeInd
 				{
 					const URigVMPin* ParentPin = Pin->GetParentPin();
 					TSharedPtr<FStructOnScope> StructOnScope = MakeShareable(new FStructOnScope(ParentPin->GetScriptStruct()));
-					ParentPin->GetScriptStruct()->InitializeDefaultValue((uint8*)StructOnScope->GetStructMemory());
 					const FString StructDefaultValue = FRigVMStruct::ExportToFullyQualifiedText(ParentPin->GetScriptStruct(), StructOnScope->GetStructMemory());
 					Local::ApplyResolvedDefaultValue(this, Pin, Pin->GetName(), StructDefaultValue, bSetupUndoRedo);
 				}
