Control Rig: Fix the pythonize name not working for Enum Values

Copied PyGenUtil::PythonizeName (which is a private function in an Experimental plugin) into RigVMPythonUtils.

diff --git a/Engine/Plugins/Animation/ControlRig/Source/ControlRig/Private/Rigs/RigHierarchy.cpp b/Engine/Plugins/Animation/ControlRig/Source/ControlRig/Private/Rigs/RigHierarchy.cpp
--- a/Engine/Plugins/Animation/ControlRig/Source/ControlRig/Private/Rigs/RigHierarchy.cpp
+++ b/Engine/Plugins/Animation/ControlRig/Source/ControlRig/Private/Rigs/RigHierarchy.cpp
@@ -3311,7 +3311,7 @@ void URigHierarchy::SetControlSettings(FRigControlElement* InControlElement, FRi
 		}
 		if (!BlueprintName.IsEmpty())
 		{
-			FString ControlNamePythonized = RigVMPythonUtils::NameToPep8(InControlElement->GetName().ToString());
+			FString ControlNamePythonized = RigVMPythonUtils::PythonizeName(InControlElement->GetName().ToString());
 			FString SettingsName = FString::Printf(TEXT("control_settings_%s"),
 				*ControlNamePythonized);
 			TArray<FString> Commands = ControlSettingsToPythonCommands(InControlElement->Settings, SettingsName);
diff --git a/Engine/Plugins/Animation/ControlRig/Source/ControlRig/Private/Rigs/RigHierarchyController.cpp b/Engine/Plugins/Animation/ControlRig/Source/ControlRig/Private/Rigs/RigHierarchyController.cpp
--- a/Engine/Plugins/Animation/ControlRig/Source/ControlRig/Private/Rigs/RigHierarchyController.cpp
+++ b/Engine/Plugins/Animation/ControlRig/Source/ControlRig/Private/Rigs/RigHierarchyController.cpp
@@ -1493,7 +1493,7 @@ TArray<FString> URigHierarchyController::GetAddControlPythonCommands(FRigControl
 	FRigControlSettings& Settings = Control->Settings;
 	FString SettingsStr;
 	{
-		FString ControlNamePythonized = RigVMPythonUtils::NameToPep8(Control->GetName().ToString());
+		FString ControlNamePythonized = RigVMPythonUtils::PythonizeName(Control->GetName().ToString());
 		SettingsStr = FString::Printf(TEXT("control_settings_%s"),
 			*ControlNamePythonized);
 			
@@ -1538,7 +1538,7 @@ TArray<FString> URigHierarchyController::GetAddRigidBodyPythonCommands(FRigRigid
 		ParentKeyStr = RigidBody->ParentElement->GetKey().ToPythonString();
 	}
 
-	FString RigidBodyNamePythonized = RigVMPythonUtils::NameToPep8(RigidBody->GetName().ToString());
+	FString RigidBodyNamePythonized = RigVMPythonUtils::PythonizeName(RigidBody->GetName().ToString());
 	FRigRigidBodySettings& Settings = RigidBody->Settings;
 	FString SettingsStr;
 	{
diff --git a/Engine/Plugins/Animation/ControlRig/Source/ControlRigDeveloper/Private/ControlRigBlueprint.cpp b/Engine/Plugins/Animation/ControlRig/Source/ControlRigDeveloper/Private/ControlRigBlueprint.cpp
--- a/Engine/Plugins/Animation/ControlRig/Source/ControlRigDeveloper/Private/ControlRigBlueprint.cpp
+++ b/Engine/Plugins/Animation/ControlRig/Source/ControlRigDeveloper/Private/ControlRigBlueprint.cpp
@@ -2192,10 +2192,10 @@ TArray<FString> UControlRigBlueprint::GeneratePythonCommands(const FString InNew
 				{						
 					// Add Function
 					InternalCommands.Add(FString::Printf(TEXT("function_%s = library_controller.add_function_to_library('%s', mutable=%s)\ngraph = function_%s.get_contained_graph()"),
-							*RigVMPythonUtils::NameToPep8(Graph->GetGraphName()),
+							*RigVMPythonUtils::PythonizeName(Graph->GetGraphName()),
 							*Graph->GetGraphName(),
 							Graph->GetEntryNode()->IsMutable() ? TEXT("True") : TEXT("False"),
-							*RigVMPythonUtils::NameToPep8(Graph->GetGraphName())));
+							*RigVMPythonUtils::PythonizeName(Graph->GetGraphName())));
 
 					URigVMLibraryNode* LibraryNode = Cast<URigVMLibraryNode>(Graph->GetOuter());
 
diff --git a/Engine/Source/Developer/RigVMDeveloper/Private/RigVMModel/RigVMController.cpp b/Engine/Source/Developer/RigVMDeveloper/Private/RigVMModel/RigVMController.cpp
--- a/Engine/Source/Developer/RigVMDeveloper/Private/RigVMModel/RigVMController.cpp
+++ b/Engine/Source/Developer/RigVMDeveloper/Private/RigVMModel/RigVMController.cpp
@@ -671,7 +671,7 @@ TArray<FString> URigVMController::GetAddNodePythonCommands(URigVMNode* Node) con
 		{
 			Commands.Add(FString::Printf(TEXT("blueprint.get_controller_by_name('%s').add_function_reference_node(function_%s, %s, '%s')"),
 						*GraphName,
-						*RigVMPythonUtils::NameToPep8(ContainedGraphName),
+						*RigVMPythonUtils::PythonizeName(ContainedGraphName),
 						*RigVMPythonUtils::Vector2DToPythonString(LibraryNode->GetPosition()), 
 						*NodeName));
 		}
diff --git a/Engine/Source/Developer/RigVMDeveloper/Private/RigVMPythonUtils.cpp b/Engine/Source/Developer/RigVMDeveloper/Private/RigVMPythonUtils.cpp
--- a/Engine/Source/Developer/RigVMDeveloper/Private/RigVMPythonUtils.cpp
+++ b/Engine/Source/Developer/RigVMDeveloper/Private/RigVMPythonUtils.cpp
@@ -1,6 +1,7 @@
 // Copyright Epic Games, Inc. All Rights Reserved.
 
 #include "RigVMPythonUtils.h"
+#include "Internationalization/BreakIterator.h"
 
 #if WITH_EDITOR
 #include "Modules/ModuleManager.h"
@@ -11,29 +12,99 @@
 #define LOCTEXT_NAMESPACE "RigVMDeveloperModule"
 
 
-FString RigVMPythonUtils::NameToPep8(const FString& Name)
+FString RigVMPythonUtils::PythonizeName(FStringView InName, const RigVMPythonUtils::EPythonizeNameCase InNameCase)
 {
 	// Wish we could use PyGenUtil::PythonizeName, but unfortunately it's private
+	
+	static const TSet<FString> ReservedKeywords = {
+		TEXT("and"),
+		TEXT("as"),
+		TEXT("assert"),
+		TEXT("async"),
+		TEXT("break"),
+		TEXT("class"),
+		TEXT("continue"),
+		TEXT("def"),
+		TEXT("del"),
+		TEXT("elif"),
+		TEXT("else"),
+		TEXT("except"),
+		TEXT("finally"),
+		TEXT("for"),
+		TEXT("from"),
+		TEXT("global"),
+		TEXT("if"),
+		TEXT("import"),
+		TEXT("in"),
+		TEXT("is"),
+		TEXT("lambda"),
+		TEXT("nonlocal"),
+		TEXT("not"),
+		TEXT("or"),
+		TEXT("pass"),
+		TEXT("raise"),
+		TEXT("return"),
+		TEXT("try"),
+		TEXT("while"),
+		TEXT("with"),
+		TEXT("yield"),
+		TEXT("property"),
+	};
+
+	FString PythonizedName;
+	PythonizedName.Reserve(InName.Len() + 10);
 
-	const FString NameNoSpaces = Name.Replace(TEXT(" "), TEXT("_"));
-	FString Result;
+	static TSharedPtr<IBreakIterator> NameBreakIterator;
+	if (!NameBreakIterator.IsValid())
+	{
+		NameBreakIterator = FBreakIterator::CreateCamelCaseBreakIterator();
+	}
 
-	for (const TCHAR& Char : NameNoSpaces)
+	NameBreakIterator->SetStringRef(InName);
+	for (int32 PrevBreak = 0, NameBreak = NameBreakIterator->MoveToNext(); NameBreak != INDEX_NONE; NameBreak = NameBreakIterator->MoveToNext())
 	{
-		if (FChar::IsUpper(Char))
+		const int32 OrigPythonizedNameLen = PythonizedName.Len();
+
+		// Append an underscore if this was a break between two parts of the identifier, *and* the previous character isn't already an underscore
+		if (OrigPythonizedNameLen > 0 && PythonizedName[OrigPythonizedNameLen - 1] != TEXT('_'))
 		{
-			if (!Result.IsEmpty() && !Result.EndsWith(TEXT("_")))
-			{
-				Result.AppendChar(TEXT('_'));
-			}
-			Result.AppendChar(FChar::ToLower(Char));
+			PythonizedName += TEXT('_');
 		}
-		else
+
+		// Append this part of the identifier
+		PythonizedName.AppendChars(&InName[PrevBreak], NameBreak - PrevBreak);
+
+		// Remove any trailing underscores in the last part of the identifier
+		while (PythonizedName.Len() > OrigPythonizedNameLen)
 		{
-			Result.AppendChar(Char);
+			const int32 CharIndex = PythonizedName.Len() - 1;
+			if (PythonizedName[CharIndex] != TEXT('_'))
+			{
+				break;
+			}
+			PythonizedName.RemoveAt(CharIndex, 1, false);
 		}
+
+		PrevBreak = NameBreak;
+	}
+	NameBreakIterator->ClearString();
+
+	if (InNameCase == EPythonizeNameCase::Lower)
+	{
+		PythonizedName.ToLowerInline();
 	}
-	return Result;
+	else if (InNameCase == EPythonizeNameCase::Upper)
+	{
+		PythonizedName.ToUpperInline();
+	}
+
+	// Don't allow the name to conflict with a keyword
+	if (ReservedKeywords.Contains(PythonizedName))
+	{
+		PythonizedName += TEXT('_');
+	}
+
+	return PythonizedName;
 }
 
 FString RigVMPythonUtils::TransformToPythonString(const FTransform& Transform)
@@ -77,7 +148,7 @@ FString RigVMPythonUtils::EnumValueToPythonString(UEnum* Enum, int64 Value)
 	return FString::Printf(
 		EnumValueFormat,
 		*EnumName,
-		*NameToPep8(Enum->GetNameStringByValue((int64)Value)).ToUpper()
+		*PythonizeName(Enum->GetNameStringByValue((int64)Value), EPythonizeNameCase::Upper)
 	);
 }
 
diff --git a/Engine/Source/Developer/RigVMDeveloper/Public/RigVMPythonUtils.h b/Engine/Source/Developer/RigVMDeveloper/Public/RigVMPythonUtils.h
--- a/Engine/Source/Developer/RigVMDeveloper/Public/RigVMPythonUtils.h
+++ b/Engine/Source/Developer/RigVMDeveloper/Public/RigVMPythonUtils.h
@@ -15,7 +15,16 @@ class UEnum;
 
 namespace RigVMPythonUtils
 {
-	RIGVMDEVELOPER_API FString NameToPep8(const FString& Name);
+	/** How should PythonizeName adjust the final name? */
+	enum EPythonizeNameCase : uint8
+	{
+		/** lower_snake_case */
+		Lower,
+		/** UPPER_SNAKE_CASE */
+		Upper,
+	};
+	
+	RIGVMDEVELOPER_API FString PythonizeName(FStringView InName, const EPythonizeNameCase InNameCase = EPythonizeNameCase::Lower);
 
 	RIGVMDEVELOPER_API FString TransformToPythonString(const FTransform& Transform);
 
