Control Rig: Fix can't resolve copy/pasted function

diff --git a/Engine/Source/Developer/RigVMDeveloper/Private/RigVMModel/RigVMController.cpp b/Engine/Source/Developer/RigVMDeveloper/Private/RigVMModel/RigVMController.cpp
--- a/Engine/Source/Developer/RigVMDeveloper/Private/RigVMModel/RigVMController.cpp
+++ b/Engine/Source/Developer/RigVMDeveloper/Private/RigVMModel/RigVMController.cpp
@@ -11102,6 +11102,7 @@ URigVMFunctionReferenceNode* URigVMController::AddFunctionReferenceNode(URigVMLi
 	FRigVMControllerCompileBracketScope CompileScope(this);
 	
 	RepopulatePinsOnNode(FunctionRefNode, false, false);
+	UpdateTemplateNodePinTypes(FunctionRefNode, false);
 
 	Notify(ERigVMGraphNotifType::NodeAdded, FunctionRefNode);
 
