Control Rig: Fix unexpected behavior with constant value folding

diff --git a/Engine/Source/Developer/RigVMDeveloper/Private/RigVMCompiler/RigVMAST.cpp b/Engine/Source/Developer/RigVMDeveloper/Private/RigVMCompiler/RigVMAST.cpp
--- a/Engine/Source/Developer/RigVMDeveloper/Private/RigVMCompiler/RigVMAST.cpp
+++ b/Engine/Source/Developer/RigVMDeveloper/Private/RigVMCompiler/RigVMAST.cpp
@@ -1833,12 +1833,32 @@ bool FRigVMParserAST::FoldConstantValuesToLiterals(TArray<URigVMGraph*> InGraphs
 		FRigVMASTProxy& ProxyToCompute = NodesToCompute[NodeToComputeIndex];
 		if(URigVMNode* NodeToCompute = ProxyToCompute.GetSubject<URigVMNode>())
 		{
+			// first find all of the source node of this node directly
 			TArray<URigVMNode*> SourceNodes = NodeToCompute->GetLinkedSourceNodes();
 			for(URigVMNode* SourceNode : SourceNodes)
 			{
 				FRigVMASTProxy SourceProxy = ProxyToCompute.GetSibling(SourceNode);
 				NodesToCompute.AddUnique(SourceProxy);
 			}
+
+			// if the node is a library node - also add the return node
+			if(URigVMLibraryNode* LibraryToCompute = Cast<URigVMLibraryNode>(NodeToCompute))
+			{
+				if(URigVMNode* ReturnNodeToCompute = LibraryToCompute->GetReturnNode())
+				{
+					const FRigVMASTProxy LibraryProxy = ProxyToCompute.GetSibling(LibraryToCompute);
+					const FRigVMASTProxy ReturnProxy = LibraryProxy.GetChild(ReturnNodeToCompute);
+					NodesToCompute.AddUnique(ReturnProxy);
+				}
+			}
+      
+			// if the node is an entry node - also add host library node
+			if(NodeToCompute->IsA<URigVMFunctionEntryNode>())
+			{
+				check(ProxyToCompute.GetCallstack().Num() > 1);
+				const FRigVMASTProxy LibraryProxy = ProxyToCompute.GetParent();
+				NodesToCompute.AddUnique(LibraryProxy);
+			}
 		}
 	}
 
