Control Rig: Fix crash when connecting last pin of function

diff --git a/Engine/Source/Developer/RigVMDeveloper/Private/RigVMCompiler/RigVMAST.cpp b/Engine/Source/Developer/RigVMDeveloper/Private/RigVMCompiler/RigVMAST.cpp
--- a/Engine/Source/Developer/RigVMDeveloper/Private/RigVMCompiler/RigVMAST.cpp
+++ b/Engine/Source/Developer/RigVMDeveloper/Private/RigVMCompiler/RigVMAST.cpp
@@ -1830,7 +1830,7 @@ bool FRigVMParserAST::FoldConstantValuesToLiterals(TArray<URigVMGraph*> InGraphs
 	// add all of the additional nodes driving these
 	for(int32 NodeToComputeIndex = 0; NodeToComputeIndex < NodesToCompute.Num(); NodeToComputeIndex++)
 	{
-		FRigVMASTProxy& ProxyToCompute = NodesToCompute[NodeToComputeIndex];
+		FRigVMASTProxy ProxyToCompute = NodesToCompute[NodeToComputeIndex];
 		if(URigVMNode* NodeToCompute = ProxyToCompute.GetSubject<URigVMNode>())
 		{
 			// first find all of the source node of this node directly
