Control Rig: Fix python script generated with spaces in variable names

diff --git a/Engine/Source/Developer/RigVMDeveloper/Private/RigVMPythonUtils.cpp b/Engine/Source/Developer/RigVMDeveloper/Private/RigVMPythonUtils.cpp
--- a/Engine/Source/Developer/RigVMDeveloper/Private/RigVMPythonUtils.cpp
+++ b/Engine/Source/Developer/RigVMDeveloper/Private/RigVMPythonUtils.cpp
@@ -51,8 +51,12 @@ FString RigVMPythonUtils::PythonizeName(FStringView InName, const RigVMPythonUti
 		TEXT("property"),
 	};
 
+	// Remove spaces
+	FString Name = InName.GetData();
+	Name.ReplaceCharInline(' ', '_');
+
 	FString PythonizedName;
-	PythonizedName.Reserve(InName.Len() + 10);
+	PythonizedName.Reserve(Name.Len() + 10);
 
 	static TSharedPtr<IBreakIterator> NameBreakIterator;
 	if (!NameBreakIterator.IsValid())
@@ -60,7 +64,7 @@ FString RigVMPythonUtils::PythonizeName(FStringView InName, const RigVMPythonUti
 		NameBreakIterator = FBreakIterator::CreateCamelCaseBreakIterator();
 	}
 
-	NameBreakIterator->SetStringRef(InName);
+	NameBreakIterator->SetStringRef(Name);
 	for (int32 PrevBreak = 0, NameBreak = NameBreakIterator->MoveToNext(); NameBreak != INDEX_NONE; NameBreak = NameBreakIterator->MoveToNext())
 	{
 		const int32 OrigPythonizedNameLen = PythonizedName.Len();
@@ -72,7 +76,7 @@ FString RigVMPythonUtils::PythonizeName(FStringView InName, const RigVMPythonUti
 		}
 
 		// Append this part of the identifier
-		PythonizedName.AppendChars(&InName[PrevBreak], NameBreak - PrevBreak);
+		PythonizedName.AppendChars(&Name[PrevBreak], NameBreak - PrevBreak);
 
 		// Remove any trailing underscores in the last part of the identifier
 		while (PythonizedName.Len() > OrigPythonizedNameLen)
