Control Rig: Support for casting

diff --git a/Engine/Plugins/Animation/ControlRig/Source/ControlRig/Private/Units/Execution/RigUnit_Item.cpp b/Engine/Plugins/Animation/ControlRig/Source/ControlRig/Private/Units/Execution/RigUnit_Item.cpp
--- a/Engine/Plugins/Animation/ControlRig/Source/ControlRig/Private/Units/Execution/RigUnit_Item.cpp
+++ b/Engine/Plugins/Animation/ControlRig/Source/ControlRig/Private/Units/Execution/RigUnit_Item.cpp
@@ -104,3 +104,8 @@ IMPLEMENT_RIGUNIT_AUTOMATION_TEST(FRigUnit_ItemReplace)
 	return true;
 }
 #endif
+
+FRigUnit_ItemToName_Execute()
+{
+	Result = Value.Name;
+}
\ No newline at end of file
diff --git a/Engine/Plugins/Animation/ControlRig/Source/ControlRig/Private/Units/Execution/RigUnit_Item.h b/Engine/Plugins/Animation/ControlRig/Source/ControlRig/Private/Units/Execution/RigUnit_Item.h
--- a/Engine/Plugins/Animation/ControlRig/Source/ControlRig/Private/Units/Execution/RigUnit_Item.h
+++ b/Engine/Plugins/Animation/ControlRig/Source/ControlRig/Private/Units/Execution/RigUnit_Item.h
@@ -189,4 +189,28 @@ struct CONTROLRIG_API FRigUnit_ItemTypeNotEquals : public FRigUnit_ItemBase
 
 	UPROPERTY(meta = (Output))
 	bool Result;
-};
\ No newline at end of file
+};
+
+/**
+ * Casts the provided item key to its name
+ */
+USTRUCT(meta=(DisplayName="To Name", Keywords="", TemplateName="Cast"))
+struct CONTROLRIG_API FRigUnit_ItemToName : public FRigUnit_ItemBase
+{
+	GENERATED_BODY()
+
+	FRigUnit_ItemToName()
+	{
+		Value = FRigElementKey();
+		Result = NAME_None;
+	}
+
+	RIGVM_METHOD()
+	virtual void Execute(const FRigUnitContext& Context) override;
+
+	UPROPERTY(meta = (Input, ExpandByDefault))
+	FRigElementKey Value;
+
+	UPROPERTY(meta = (Output))
+	FName Result;
+};
diff --git a/Engine/Plugins/Animation/ControlRig/Source/ControlRig/Private/Units/Math/RigUnit_MathColor.cpp b/Engine/Plugins/Animation/ControlRig/Source/ControlRig/Private/Units/Math/RigUnit_MathColor.cpp
--- a/Engine/Plugins/Animation/ControlRig/Source/ControlRig/Private/Units/Math/RigUnit_MathColor.cpp
+++ b/Engine/Plugins/Animation/ControlRig/Source/ControlRig/Private/Units/Math/RigUnit_MathColor.cpp
@@ -12,6 +12,12 @@ FRigUnit_MathColorFromFloat_Execute()
 	Result = FLinearColor(Value, Value, Value);
 }
 
+FRigUnit_MathColorFromDouble_Execute()
+{
+	DECLARE_SCOPE_HIERARCHICAL_COUNTER_RIGUNIT()
+	Result = FLinearColor((float)Value, (float)Value, (float)Value);
+}
+
 FRigUnit_MathColorAdd_Execute()
 {
     DECLARE_SCOPE_HIERARCHICAL_COUNTER_RIGUNIT()
diff --git a/Engine/Plugins/Animation/ControlRig/Source/ControlRig/Private/Units/Math/RigUnit_MathColor.h b/Engine/Plugins/Animation/ControlRig/Source/ControlRig/Private/Units/Math/RigUnit_MathColor.h
--- a/Engine/Plugins/Animation/ControlRig/Source/ControlRig/Private/Units/Math/RigUnit_MathColor.h
+++ b/Engine/Plugins/Animation/ControlRig/Source/ControlRig/Private/Units/Math/RigUnit_MathColor.h
@@ -75,6 +75,29 @@ struct CONTROLRIG_API FRigUnit_MathColorFromFloat : public FRigUnit_MathColorBas
 	FLinearColor Result;
 };
 
+/**
+ * Makes a vector from a single double
+ */
+USTRUCT(meta=(DisplayName="From Double", TemplateName="Cast", Keywords="Make,Construct"))
+struct CONTROLRIG_API FRigUnit_MathColorFromDouble : public FRigUnit_MathColorBase
+{
+	GENERATED_BODY()
+	RIGVM_METHOD()
+	virtual void Execute(const FRigUnitContext& Context) override;
+
+	FRigUnit_MathColorFromDouble()
+	{
+		Value = 0.0;
+		Result = FLinearColor::Black;
+	}
+
+	UPROPERTY(meta=(Input))
+	double Value;
+
+	UPROPERTY(meta=(Output))
+	FLinearColor Result;
+};
+
 /**
  * Returns the sum of the two values
  */
diff --git a/Engine/Plugins/Animation/ControlRig/Source/ControlRig/Private/Units/Math/RigUnit_MathInt.cpp b/Engine/Plugins/Animation/ControlRig/Source/ControlRig/Private/Units/Math/RigUnit_MathInt.cpp
--- a/Engine/Plugins/Animation/ControlRig/Source/ControlRig/Private/Units/Math/RigUnit_MathInt.cpp
+++ b/Engine/Plugins/Animation/ControlRig/Source/ControlRig/Private/Units/Math/RigUnit_MathInt.cpp
@@ -84,6 +84,12 @@ FRigUnit_MathIntToFloat_Execute()
 	Result = (float)Value;
 }
 
+FRigUnit_MathIntToDouble_Execute()
+{
+	DECLARE_SCOPE_HIERARCHICAL_COUNTER_RIGUNIT()
+	Result = (double)Value;
+}
+
 FRigUnit_MathIntSign_Execute()
 {
     DECLARE_SCOPE_HIERARCHICAL_COUNTER_RIGUNIT()
diff --git a/Engine/Plugins/Animation/ControlRig/Source/ControlRig/Private/Units/Math/RigUnit_MathInt.h b/Engine/Plugins/Animation/ControlRig/Source/ControlRig/Private/Units/Math/RigUnit_MathInt.h
--- a/Engine/Plugins/Animation/ControlRig/Source/ControlRig/Private/Units/Math/RigUnit_MathInt.h
+++ b/Engine/Plugins/Animation/ControlRig/Source/ControlRig/Private/Units/Math/RigUnit_MathInt.h
@@ -235,6 +235,30 @@ struct CONTROLRIG_API FRigUnit_MathIntToFloat : public FRigUnit_MathIntBase
 	float Result;
 };
 
+/**
+ * Returns the int cast to a float
+ */
+USTRUCT(meta=(DisplayName="To Double", TemplateName="Cast"))
+struct CONTROLRIG_API FRigUnit_MathIntToDouble : public FRigUnit_MathIntBase
+{
+	GENERATED_BODY()
+
+	FRigUnit_MathIntToDouble()
+	{
+		Value = 0;
+		Result = 0.0;
+	}
+
+	RIGVM_METHOD()
+	virtual void Execute(const FRigUnitContext& Context) override;
+
+	UPROPERTY(meta=(Input))
+	int32 Value;
+
+	UPROPERTY(meta=(Output))
+	double Result;
+};
+
 /**
  * Returns the sign of the value (+1 for >= 0, -1 for < 0)
  */
diff --git a/Engine/Plugins/Animation/ControlRig/Source/ControlRig/Private/Units/Math/RigUnit_MathVector.cpp b/Engine/Plugins/Animation/ControlRig/Source/ControlRig/Private/Units/Math/RigUnit_MathVector.cpp
--- a/Engine/Plugins/Animation/ControlRig/Source/ControlRig/Private/Units/Math/RigUnit_MathVector.cpp
+++ b/Engine/Plugins/Animation/ControlRig/Source/ControlRig/Private/Units/Math/RigUnit_MathVector.cpp
@@ -14,6 +14,12 @@ FRigUnit_MathVectorFromFloat_Execute()
 	Result = FVector(Value, Value, Value);
 }
 
+FRigUnit_MathVectorFromDouble_Execute()
+{
+	DECLARE_SCOPE_HIERARCHICAL_COUNTER_RIGUNIT()
+	Result = FVector((float)Value, (float)Value, (float)Value);
+}
+
 FRigUnit_MathVectorAdd_Execute()
 {
     DECLARE_SCOPE_HIERARCHICAL_COUNTER_RIGUNIT()
diff --git a/Engine/Plugins/Animation/ControlRig/Source/ControlRig/Private/Units/Math/RigUnit_MathVector.h b/Engine/Plugins/Animation/ControlRig/Source/ControlRig/Private/Units/Math/RigUnit_MathVector.h
--- a/Engine/Plugins/Animation/ControlRig/Source/ControlRig/Private/Units/Math/RigUnit_MathVector.h
+++ b/Engine/Plugins/Animation/ControlRig/Source/ControlRig/Private/Units/Math/RigUnit_MathVector.h
@@ -92,6 +92,29 @@ struct CONTROLRIG_API FRigUnit_MathVectorFromFloat : public FRigUnit_MathVectorB
 	FVector Result;
 };
 
+/**
+ * Makes a vector from a single double
+ */
+USTRUCT(meta=(DisplayName="From Float", TemplateName="Cast", Keywords="Make,Construct"))
+struct CONTROLRIG_API FRigUnit_MathVectorFromDouble : public FRigUnit_MathVectorBase
+{
+	GENERATED_BODY()
+	RIGVM_METHOD()
+	virtual void Execute(const FRigUnitContext& Context) override;
+
+	FRigUnit_MathVectorFromDouble()
+	{
+		Value = 0.0;
+		Result = FVector::ZeroVector;
+	}
+
+	UPROPERTY(meta=(Input))
+	double Value;
+
+	UPROPERTY(meta=(Output))
+	FVector Result;
+};
+
 /**
  * Returns the sum of the two values
  */
diff --git a/Engine/Plugins/Animation/ControlRig/Source/ControlRigDeveloper/Private/Graph/ControlRigGraphSchema.cpp b/Engine/Plugins/Animation/ControlRig/Source/ControlRigDeveloper/Private/Graph/ControlRigGraphSchema.cpp
--- a/Engine/Plugins/Animation/ControlRig/Source/ControlRigDeveloper/Private/Graph/ControlRigGraphSchema.cpp
+++ b/Engine/Plugins/Animation/ControlRig/Source/ControlRigDeveloper/Private/Graph/ControlRigGraphSchema.cpp
@@ -1794,33 +1794,45 @@ bool UControlRigGraphSchema::ArePinsCompatible(const UEdGraphPin* PinA, const UE
 		{
 			if(const UControlRigGraphNode* RigNode = Cast<UControlRigGraphNode>(InPin->GetOwningNode()))
 			{
-				if(const FRigVMTemplate* Template = RigNode->GetTemplate())
+				FString CPPType;
+				UObject* CPPTypeObject = nullptr;
+				if(RigVMTypeUtils::CPPTypeFromPinType(InPinType, CPPType, &CPPTypeObject))
 				{
-					FString CPPType;
-					UObject* CPPTypeObject = nullptr;
-					if(RigVMTypeUtils::CPPTypeFromPinType(InPinType, CPPType, &CPPTypeObject))
-					{
-						FString PinPath, PinName;
-						URigVMPin::SplitPinPathAtEnd(InPin->GetName(), PinPath, PinName);
+					FString PinPath, PinName;
+					URigVMPin::SplitPinPathAtEnd(InPin->GetName(), PinPath, PinName);
 						
-						if((InPin->ParentPin != nullptr) &&
-							(InPin->ParentPin->ParentPin == nullptr) &&
-							(InPin->ParentPin->PinType.ContainerType == EPinContainerType::Array))
-						{
-							URigVMPin::SplitPinPathAtEnd(InPin->ParentPin->GetName(), PinPath, PinName);
-							CPPType = RigVMTypeUtils::ArrayTypeFromBaseType(CPPType);
-						}
+					if((InPin->ParentPin != nullptr) &&
+						(InPin->ParentPin->ParentPin == nullptr) &&
+						(InPin->ParentPin->PinType.ContainerType == EPinContainerType::Array))
+					{
+						URigVMPin::SplitPinPathAtEnd(InPin->ParentPin->GetName(), PinPath, PinName);
+						CPPType = RigVMTypeUtils::ArrayTypeFromBaseType(CPPType);
+					}
 
-						const FRigVMTemplateArgumentType Type(*CPPType, CPPTypeObject);
-						const int32 TypeIndex = FRigVMRegistry::Get().GetTypeIndex(Type);
-						if(!Template->ArgumentSupportsTypeIndex(*PinName, TypeIndex))
+					const FRigVMTemplateArgumentType Type(*CPPType, CPPTypeObject);
+					const TRigVMTypeIndex TypeIndex = FRigVMRegistry::Get().GetTypeIndex(Type);
+					if(const FRigVMTemplate* Template = RigNode->GetTemplate())
+					{
+						if(const FRigVMTemplateArgument* Argument = Template->FindArgument(*PinName))
 						{
-							return false;
+							if(Argument->SupportsTypeIndex(TypeIndex))
+							{
+								return true;
+							}
+					
+							const TArray<TRigVMTypeIndex>& AvailableCasts = RigVMTypeUtils::GetAvailableCasts(TypeIndex, InPin->Direction == EGPD_Output);
+							for(const TRigVMTypeIndex& AvailableCast : AvailableCasts)
+							{
+								if(Argument->SupportsTypeIndex(AvailableCast))
+								{
+									return true;
+								}
+							}
 						}
 					}
 				}
 			}
-			return true;
+			return false;
 		};
 
 		auto IsTemplateNodePin = [](const UEdGraphPin* InPin)
@@ -1868,7 +1880,15 @@ bool UControlRigGraphSchema::ArePinsCompatible(const UEdGraphPin* PinA, const UE
 		}
 	}
 
-	return GetDefault<UEdGraphSchema_K2>()->ArePinsCompatible(PinA, PinB, CallingContext, bIgnoreArray);
+	if(GetDefault<UEdGraphSchema_K2>()->ArePinsCompatible(PinA, PinB, CallingContext, bIgnoreArray))
+	{
+		return true;
+	}
+
+	// also check if there's a cast available for the type
+	const TRigVMTypeIndex TypeIndexA = RigVMTypeUtils::TypeIndexFromPinType(PinA->Direction == EGPD_Output ? PinA->PinType : PinB->PinType);
+	const TRigVMTypeIndex TypeIndexB = RigVMTypeUtils::TypeIndexFromPinType(PinA->Direction == EGPD_Output ? PinB->PinType : PinA->PinType);
+	return RigVMTypeUtils::CanCastTypes(TypeIndexA, TypeIndexB);
 }
 
 void UControlRigGraphSchema::RenameNode(UControlRigGraphNode* Node, const FName& InNewNodeName) const
diff --git a/Engine/Source/Developer/RigVMDeveloper/Private/RigVMCompiler/RigVMCompiler.cpp b/Engine/Source/Developer/RigVMDeveloper/Private/RigVMCompiler/RigVMCompiler.cpp
--- a/Engine/Source/Developer/RigVMDeveloper/Private/RigVMCompiler/RigVMCompiler.cpp
+++ b/Engine/Source/Developer/RigVMDeveloper/Private/RigVMCompiler/RigVMCompiler.cpp
@@ -141,6 +141,89 @@ int32 FRigVMCompilerWorkData::FindOrAddPropertyPath(const FRigVMOperand& InOpera
 	return Descriptions.Add(FRigVMPropertyPathDescription(InOperand.GetRegisterIndex(), InHeadCPPType, InSegmentPath));
 }
 
+const FProperty* FRigVMCompilerWorkData::GetPropertyForOperand(const FRigVMOperand& InOperand) const
+{
+	if(!InOperand.IsValid())
+	{
+		return nullptr;
+	}
+	
+	check(!bSetupMemory);
+
+	auto GetPropertyFromMemory = [this](const URigVMMemoryStorage* InMemory, const FRigVMOperand& InOperand)
+	{
+		if(InOperand.GetRegisterOffset() == INDEX_NONE)
+		{
+			return  InMemory->GetProperty(InOperand.GetRegisterIndex());
+		}
+		if(!InMemory->GetPropertyPaths().IsValidIndex(InOperand.GetRegisterOffset()))
+		{
+			if(URigVMMemoryStorageGeneratorClass* MemoryClass = Cast<URigVMMemoryStorageGeneratorClass>(InMemory->GetClass()))
+			{
+				MemoryClass->PropertyPathDescriptions = PropertyPathDescriptions.FindChecked(InOperand.GetMemoryType());;
+				MemoryClass->RefreshPropertyPaths();
+			}
+		}
+		return InMemory->GetPropertyPaths()[InOperand.GetRegisterOffset()].GetTailProperty();
+	};
+
+	const FProperty* Property = nullptr;
+	switch(InOperand.GetMemoryType())
+	{
+	case ERigVMMemoryType::Literal:
+		{
+			Property = GetPropertyFromMemory(VM->GetLiteralMemory(), InOperand);
+			break;
+		}
+	case ERigVMMemoryType::Work:
+		{
+			Property = GetPropertyFromMemory(VM->GetWorkMemory(), InOperand);
+			break;
+		}
+	case ERigVMMemoryType::Debug:
+		{
+			Property = GetPropertyFromMemory(VM->GetDebugMemory(), InOperand);
+			break;
+		}
+	case ERigVMMemoryType::External:
+		{
+			Property = VM->GetExternalVariables()[InOperand.GetRegisterIndex()].Property;
+			if(InOperand.GetRegisterOffset() != INDEX_NONE)
+			{
+				if(!VM->ExternalPropertyPaths.IsValidIndex(InOperand.GetRegisterOffset()))
+				{
+					VM->ExternalPropertyPathDescriptions = PropertyPathDescriptions.FindChecked(InOperand.GetMemoryType());
+					VM->RefreshExternalPropertyPaths();
+				}
+				Property = VM->ExternalPropertyPaths[InOperand.GetRegisterOffset()].GetTailProperty();
+			}
+			break;
+		}
+	case ERigVMMemoryType::Invalid:
+	default:
+		{
+			break;
+		}
+	}
+
+	return Property;
+}
+
+TRigVMTypeIndex FRigVMCompilerWorkData::GetTypeIndexForOperand(const FRigVMOperand& InOperand) const
+{
+	const FProperty* Property = GetPropertyForOperand(InOperand);
+	if(Property == nullptr)
+	{
+		return INDEX_NONE;
+	}
+
+	FName CPPTypeName(NAME_None);
+	UObject* CPPTypeObject = nullptr;
+	FRigVMExternalVariable::GetTypeFromProperty(Property, CPPTypeName, CPPTypeObject);
+
+	return FRigVMRegistry::Get().GetTypeIndex(CPPTypeName, CPPTypeObject);
+}
+
 URigVMCompiler::URigVMCompiler()
 {
 }
@@ -1941,16 +2024,52 @@ void URigVMCompiler::AddCopyOperator(const FRigVMCopyOp& InOp, const FRigVMAssig
 
 	bool bAddCopyOp = true;
 
+	// check if we need to inject a cast instead of a copy operator
+	const TRigVMTypeIndex SourceTypeIndex = WorkData.GetTypeIndexForOperand(InOp.Source);
+	const TRigVMTypeIndex TargetTypeIndex = WorkData.GetTypeIndexForOperand(InOp.Target);
+	if(SourceTypeIndex != TargetTypeIndex)
+	{
+		// if the type system can't auto cast these types (like float vs double)
+		if(!FRigVMRegistry::Get().CanMatchTypes(SourceTypeIndex, TargetTypeIndex, true))
+		{
+			const FRigVMFunction* CastFunction = RigVMTypeUtils::GetCastForTypeIndices(SourceTypeIndex, TargetTypeIndex);
+			if(CastFunction == nullptr)
+			{
+				const FRigVMRegistry& Registry = FRigVMRegistry::Get();
+				static constexpr TCHAR MissingCastMessage[] = TEXT("Cast (%s -> %s) for Node @@ not found.");
+				const FString& SourceCPPType = Registry.GetType(SourceTypeIndex).CPPType.ToString();
+				const FString& TargetCPPType = Registry.GetType(TargetTypeIndex).CPPType.ToString();
+				Settings.Report(EMessageSeverity::Error, InAssignExpr->GetTargetPin()->GetNode(),
+					FString::Printf(MissingCastMessage, *SourceCPPType, *TargetCPPType));
+				return;
+			}
+
+			check(CastFunction->Arguments.Num() >= 2);
+
+			const FRigVMOperand Source = InOp.Source;
+			const FRigVMOperand Target = InOp.Target;
+
+			const int32 FunctionIndex = WorkData.VM->AddRigVMFunction(CastFunction->Name);
+			WorkData.VM->GetByteCode().AddExecuteOp(FunctionIndex, {Source, Target});
+
+			bAddCopyOp = false;
+		}
+	}
+
+
 	// if we are copying into an array variable
-	if(const URigVMPin* Pin = InTargetExpr->GetPin())
+	if(bAddCopyOp)
 	{
-		if(Pin->IsArray() && Pin->GetNode()->IsA<URigVMVariableNode>())
+		if(const URigVMPin* Pin = InTargetExpr->GetPin())
 		{
-			if(InOp.Source.GetRegisterOffset() == INDEX_NONE &&
-				InOp.Target.GetRegisterOffset() == INDEX_NONE)
-			{ 
-				WorkData.VM->GetByteCode().AddArrayCloneOp(InOp.Source, InOp.Target);
-				bAddCopyOp = false;
+			if(Pin->IsArray() && Pin->GetNode()->IsA<URigVMVariableNode>())
+			{
+				if(InOp.Source.GetRegisterOffset() == INDEX_NONE &&
+					InOp.Target.GetRegisterOffset() == INDEX_NONE)
+				{ 
+					WorkData.VM->GetByteCode().AddArrayCloneOp(InOp.Source, InOp.Target);
+					bAddCopyOp = false;
+				}
 			}
 		}
 	}
diff --git a/Engine/Source/Developer/RigVMDeveloper/Private/RigVMDeveloperTypeUtils.cpp b/Engine/Source/Developer/RigVMDeveloper/Private/RigVMDeveloperTypeUtils.cpp
--- a/Engine/Source/Developer/RigVMDeveloper/Private/RigVMDeveloperTypeUtils.cpp
+++ b/Engine/Source/Developer/RigVMDeveloper/Private/RigVMDeveloperTypeUtils.cpp
@@ -3,6 +3,7 @@
 #include "RigVMDeveloperTypeUtils.h"
 
 #include "Internationalization/StringTableCore.h"
+#include "RigVMModel/RigVMController.h"
 #include "RigVMModel/RigVMVariableDescription.h"
 
 #if WITH_EDITOR
@@ -267,6 +268,16 @@ FRigVMExternalVariable RigVMTypeUtils::ExternalVariableFromCPPType(const FName&
 	return Variable;
 }
 
+FEdGraphPinType RigVMTypeUtils::PinTypeFromTypeIndex(const TRigVMTypeIndex& InTypeIndex)
+{
+	const FRigVMTemplateArgumentType& Type = FRigVMRegistry::Get().GetType(InTypeIndex);
+	if(Type.CPPType.IsValid())
+	{
+		return PinTypeFromCPPType(Type.CPPType, Type.CPPTypeObject);
+	}
+	return FEdGraphPinType();
+}
+
 FEdGraphPinType RigVMTypeUtils::PinTypeFromCPPType(const FName& InCPPType, UObject* InCPPTypeObject)
 {
 	FEdGraphPinType PinType;
@@ -724,6 +735,18 @@ bool RigVMTypeUtils::CPPTypeFromExternalVariable(const FRigVMExternalVariable& I
 	return true;
 }
 
+TRigVMTypeIndex RigVMTypeUtils::TypeIndexFromPinType(const FEdGraphPinType& InPinType)
+{
+	FString CPPType;
+	UObject* CPPTypeObject = nullptr;
+	if(CPPTypeFromPinType(InPinType, CPPType, &CPPTypeObject))
+	{
+		const FRigVMTemplateArgumentType Type(*CPPType, CPPTypeObject);
+		return FRigVMRegistry::Get().GetTypeIndex(Type);
+	}
+	return INDEX_NONE;
+}
+
 bool RigVMTypeUtils::AreCompatible(const FName& InCPPTypeA, UObject* InCPPTypeObjectA, const FName& InCPPTypeB,	UObject* InCPPTypeObjectB)
 {
 	return AreCompatible(PinTypeFromCPPType(InCPPTypeA, InCPPTypeObjectA), PinTypeFromCPPType(InCPPTypeB, InCPPTypeObjectB));
@@ -739,7 +762,7 @@ bool RigVMTypeUtils::AreCompatible(const FEdGraphPinType& InTypeA, const FEdGrap
 	FEdGraphPinType SubPinTypeA = SubPinType(InTypeA, InSegmentPathA);
 	FEdGraphPinType SubPinTypeB = SubPinType(InTypeB, InSegmentPathB);
 
-	// We allow connectiongs between floats and doubles, while EdGraphSchema_K2 does not
+	// We allow connections between floats and doubles, while EdGraphSchema_K2 does not
 	// Every other case is evaluated by UEdGraphSchema_K2::ArePinTypesCompatible
 	if (SubPinTypeA.ContainerType == SubPinTypeB.ContainerType)
 	{
@@ -749,7 +772,94 @@ bool RigVMTypeUtils::AreCompatible(const FEdGraphPinType& InTypeA, const FEdGrap
 			return true;
 		}
 	}
-	
-	return UEdGraphSchema_K2::StaticClass()->GetDefaultObject<UEdGraphSchema_K2>()->ArePinTypesCompatible(SubPinTypeA, SubPinTypeB);	
+
+	if(UEdGraphSchema_K2::StaticClass()->GetDefaultObject<UEdGraphSchema_K2>()->ArePinTypesCompatible(SubPinTypeA, SubPinTypeB))
+	{
+		return true;
+	}
+
+	// also check if there's a cast available for the type
+	const TRigVMTypeIndex TypeIndexA = TypeIndexFromPinType(SubPinTypeA);
+	const TRigVMTypeIndex TypeIndexB = TypeIndexFromPinType(SubPinTypeB);
+	const TArray<TRigVMTypeIndex>& AvailableCasts = GetAvailableCasts(TypeIndexA, true);
+	return AvailableCasts.Contains(TypeIndexB);
+ }
+
+namespace RigVMTypeUtils
+{
+	static const FName CastTemplateValueName = TEXT("Value");
+	static const FName CastTemplateResultName = TEXT("Result");
+	static const FName CastTemplateNotation = TEXT("Cast::Execute(in Value,out Result)");
 }
+
+const TArray<TRigVMTypeIndex>& RigVMTypeUtils::GetAvailableCasts(const TRigVMTypeIndex& InTypeIndex, bool bAsInput)
+{
+	static TMap<int32, TArray<TRigVMTypeIndex>> AvailableInputCastMap;
+	static TMap<int32, TArray<TRigVMTypeIndex>> AvailableOutputCastMap;
+
+	TMap<int32, TArray<TRigVMTypeIndex>>& AvailableCastMap = bAsInput ? AvailableInputCastMap : AvailableOutputCastMap;
+
+	if(const TArray<TRigVMTypeIndex>* AvailableCasts = AvailableCastMap.Find(InTypeIndex))
+	{
+		return *AvailableCasts;
+	}
+
+	// find a specific template called "Cast"
+	const FRigVMRegistry& Registry = FRigVMRegistry::Get();
+	if(const FRigVMTemplate* CastTemplate = Registry.FindTemplate(CastTemplateNotation))
+	{
+		const FRigVMTemplateArgument* ArgumentA = CastTemplate->FindArgument(bAsInput ? CastTemplateValueName : CastTemplateResultName);
+		const FRigVMTemplateArgument* ArgumentB = CastTemplate->FindArgument(bAsInput ? CastTemplateResultName : CastTemplateValueName);
+
+		if(ArgumentA && ArgumentB)
+		{
+			TArray<TRigVMTypeIndex>& AvailableCasts = AvailableCastMap.Add(InTypeIndex);
+			
+			for(int32 Index = 0; Index < ArgumentA->GetTypeIndices().Num(); Index++)
+			{
+				const TRigVMTypeIndex& TypeIndexA = ArgumentA->GetTypeIndices()[Index];
+				if(TypeIndexA == InTypeIndex)
+				{
+					AvailableCasts.Add(ArgumentB->GetTypeIndices()[Index]);
+				}
+			}
+
+			return AvailableCasts;
+		}
+	}
+
+	static const TArray<TRigVMTypeIndex> EmptyCasts;
+	return EmptyCasts;
+}
+
+bool RigVMTypeUtils::CanCastTypes(const TRigVMTypeIndex& InSourceTypeIndex, const TRigVMTypeIndex& InTargetTypeIndex)
+{
+	return GetCastForTypeIndices(InSourceTypeIndex, InTargetTypeIndex) != nullptr;
+}
+
+const FRigVMFunction* RigVMTypeUtils::GetCastForTypeIndices(const TRigVMTypeIndex& InSourceTypeIndex, const TRigVMTypeIndex& InTargetTypeIndex)
+{
+	const FRigVMRegistry& Registry = FRigVMRegistry::Get();
+	if(const FRigVMTemplate* CastTemplate = Registry.FindTemplate(CastTemplateNotation))
+	{
+		const FRigVMTemplateArgument* SourceArgument = CastTemplate->FindArgument(CastTemplateValueName);
+		const FRigVMTemplateArgument* TargetArgument = CastTemplate->FindArgument(CastTemplateResultName);
+
+		if(SourceArgument && TargetArgument)
+		{
+			for(int32 Index = 0; Index < SourceArgument->GetTypeIndices().Num(); Index++)
+			{
+				const TRigVMTypeIndex& SourceTypeIndex = SourceArgument->GetTypeIndices()[Index];
+				const TRigVMTypeIndex& TargetTypeIndex = TargetArgument->GetTypeIndices()[Index];
+				if(SourceTypeIndex == InSourceTypeIndex &&
+					TargetTypeIndex == InTargetTypeIndex)
+				{
+					return CastTemplate->GetPermutation(Index);
+				}
+			}
+		}
+	}
+	return nullptr;
+}
+
 #endif
diff --git a/Engine/Source/Developer/RigVMDeveloper/Private/RigVMModel/RigVMController.cpp b/Engine/Source/Developer/RigVMDeveloper/Private/RigVMModel/RigVMController.cpp
--- a/Engine/Source/Developer/RigVMDeveloper/Private/RigVMModel/RigVMController.cpp
+++ b/Engine/Source/Developer/RigVMDeveloper/Private/RigVMModel/RigVMController.cpp
@@ -70,6 +70,7 @@ URigVMController::URigVMController()
 	, bReportWarningsAndErrors(true)
 	, bIgnoreRerouteCompactnessChanges(false)
 	, UserLinkDirection(ERigVMPinDirection::Invalid)
+	, bEnableTypeCasting(true)
 	, bIsTransacting(false)
 	, bIsRunningUnitTest(false)
 	, bIsFullyResolvingTemplateNode(false)
@@ -88,6 +89,7 @@ URigVMController::URigVMController(const FObjectInitializer& ObjectInitializer)
 	, bReportWarningsAndErrors(true)
 	, bIgnoreRerouteCompactnessChanges(false)
 	, UserLinkDirection(ERigVMPinDirection::Invalid)
+	, bEnableTypeCasting(true)
 	, bIsTransacting(false)
 	, bIsRunningUnitTest(false)
 	, bIsFullyResolvingTemplateNode(false)
@@ -9310,10 +9312,46 @@ bool URigVMController::AddLink(URigVMPin* OutputPin, URigVMPin* InputPin, bool b
 		ExpandPinRecursively(InputPin->GetParentPin(), bSetupUndoRedo);
 	}
 
+	// check if there's a cast available for this type combination
+	bool bIsUsingCast = false;
+	if(bEnableTypeCasting && RigVMTypeUtils::CanCastTypes(OutputPin->GetTypeIndex(), InputPin->GetTypeIndex()))
+	{
+		auto ShouldRelyOnCast = [](const TRigVMTypeIndex& InTypeIndex, URigVMPin* InPin) -> bool
+		{
+			if(InPin->IsRootPin())
+			{
+				if(URigVMTemplateNode* TemplateNode = Cast<URigVMTemplateNode>(InPin->GetNode()))
+				{
+					if(const FRigVMTemplate* Template = TemplateNode->GetTemplate())
+					{
+						if(const FRigVMTemplateArgument* Argument = Template->FindArgument(InPin->GetFName()))
+						{
+							if(!Argument->SupportsTypeIndex(InTypeIndex))
+							{
+								return true;
+							}
+						}
+					}
+				}
+			}
+			return InPin->GetNode()->GetLinks().Num() > 0;
+		};
+		
+		// only rely on a cast if there's already a link
+		if(InUserDirection == ERigVMPinDirection::Input)
+		{
+			bIsUsingCast = ShouldRelyOnCast(InputPin->GetTypeIndex(), OutputPin);
+		}
+		else
+		{
+			bIsUsingCast = ShouldRelyOnCast(OutputPin->GetTypeIndex(), InputPin);
+		}
+	}
+
 	// Before adding the link, let's resolve input and ouput pin types
 	// If templates, we will filter the permutations that support this link
 	// If any links need to be broken before perfoming this connection, try to find them and break them
-	if (!bIsTransacting)
+	if (!bIsTransacting && !bIsUsingCast)
 	{
 		URigVMPin* FirstToResolve = (InUserDirection == ERigVMPinDirection::Input) ? OutputPin : InputPin;
 		URigVMPin* SecondToResolve = (FirstToResolve == OutputPin) ? InputPin : OutputPin;
@@ -9348,16 +9386,19 @@ bool URigVMController::AddLink(URigVMPin* OutputPin, URigVMPin* InputPin, bool b
 	Notify(ERigVMGraphNotifType::LinkAdded, Link);
 
 	// Remove any temporary preferred type that was added
-	for (int32 i=0; i<2; ++i)
+	if(!bIsUsingCast)
 	{
-		URigVMPin* EdgePin = (i==0) ? OutputPin : InputPin;
-		bool& bAddedTemporaryPreferredType = (i==0) ? bAddedTemporaryOutputPreferredType : bAddedTemporaryInputPreferredType;
-		if (bAddedTemporaryPreferredType)
+		for (int32 i=0; i<2; ++i)
 		{
-			if (URigVMTemplateNode* TemplateNode = Cast<URigVMTemplateNode>(EdgePin->GetNode()))
+			URigVMPin* EdgePin = (i==0) ? OutputPin : InputPin;
+			bool& bAddedTemporaryPreferredType = (i==0) ? bAddedTemporaryOutputPreferredType : bAddedTemporaryInputPreferredType;
+			if (bAddedTemporaryPreferredType)
 			{
-				URigVMPin* RootPin = EdgePin->GetRootPin();
-				RemovePreferredType(TemplateNode, RootPin->GetFName(), bSetupUndoRedo);				
+				if (URigVMTemplateNode* TemplateNode = Cast<URigVMTemplateNode>(EdgePin->GetNode()))
+				{
+					URigVMPin* RootPin = EdgePin->GetRootPin();
+					RemovePreferredType(TemplateNode, RootPin->GetFName(), bSetupUndoRedo);				
+				}
 			}
 		}
 	}
@@ -9439,8 +9480,12 @@ bool URigVMController::AddLink(URigVMPin* OutputPin, URigVMPin* InputPin, bool b
 				}
 			}
 		};
-		ResolveTemplateNodeToCommonTypes(OutputPin);
-		ResolveTemplateNodeToCommonTypes(InputPin);
+
+		if(!bIsUsingCast)
+		{
+			ResolveTemplateNodeToCommonTypes(OutputPin);
+			ResolveTemplateNodeToCommonTypes(InputPin);
+		}
 #endif
 		
 		ActionStack->EndAction(Action);
@@ -17209,9 +17254,12 @@ bool URigVMController::PropagateTemplateFilteredTypes(URigVMTemplateNode* InNode
 						else
 						{
 							URigVMLink* Link = Pin->FindLinkForPin(OtherPin);
-							ensureMsgf(!ActionStack->BracketActions.IsEmpty(), TEXT("Unexpected link broken %s in package %s"), *Link->GetPinPathRepresentation(), *GetPackage()->GetPathName());
-							BreakLink(Link->GetSourcePin(), Link->GetTargetPin(), bSetupUndoRedo);
-							return false;
+							if(!bEnableTypeCasting || !RigVMTypeUtils::CanCastTypes(Link->GetSourcePin()->GetTypeIndex(), Link->GetTargetPin()->GetTypeIndex()))
+							{
+								ensureMsgf(!ActionStack->BracketActions.IsEmpty(), TEXT("Unexpected link broken %s in package %s"), *Link->GetPinPathRepresentation(), *GetPackage()->GetPathName());
+								BreakLink(Link->GetSourcePin(), Link->GetTargetPin(), bSetupUndoRedo);
+								return false;
+							}
 						}
 					}
 				}								
@@ -17222,9 +17270,12 @@ bool URigVMController::PropagateTemplateFilteredTypes(URigVMTemplateNode* InNode
 				if (!InNode->FilteredSupportsType(Pin, OtherPin->GetTypeIndex()))
 				{
 					URigVMLink* Link = Pin->FindLinkForPin(OtherPin);
-					ensureMsgf(!ActionStack->BracketActions.IsEmpty(), TEXT("Unexpected link broken %s in package %s"), *Link->GetPinPathRepresentation(), *GetPackage()->GetPathName());
-					BreakLink(Link->GetSourcePin(), Link->GetTargetPin(), bSetupUndoRedo);
-					return false;
+					if(!bEnableTypeCasting || !RigVMTypeUtils::CanCastTypes(Link->GetSourcePin()->GetTypeIndex(), Link->GetTargetPin()->GetTypeIndex()))
+					{
+						ensureMsgf(!ActionStack->BracketActions.IsEmpty(), TEXT("Unexpected link broken %s in package %s"), *Link->GetPinPathRepresentation(), *GetPackage()->GetPathName());
+						BreakLink(Link->GetSourcePin(), Link->GetTargetPin(), bSetupUndoRedo);
+						return false;
+					}
 				}
 			}
 		}
@@ -17937,7 +17988,7 @@ bool URigVMController::RecomputeAllTemplateFilteredPermutations(bool bSetupUndoR
 					Notify(ERigVMGraphNotifType::LinkAdded, Link);
 				}
 			}
-			else
+			else if(!bEnableTypeCasting || !RigVMTypeUtils::CanCastTypes(OutputPin->GetTypeIndex(), InputPin->GetTypeIndex()))
 			{
 				LinksToRemove.Add(Link);
 			}
diff --git a/Engine/Source/Developer/RigVMDeveloper/Private/RigVMModel/RigVMPin.cpp b/Engine/Source/Developer/RigVMDeveloper/Private/RigVMModel/RigVMPin.cpp
--- a/Engine/Source/Developer/RigVMDeveloper/Private/RigVMModel/RigVMPin.cpp
+++ b/Engine/Source/Developer/RigVMDeveloper/Private/RigVMModel/RigVMPin.cpp
@@ -112,8 +112,14 @@ TArray<FString> URigVMPin::SplitDefaultValue(const FString& InDefaultValue)
 		return Parts;
 	}
 
-	ensure(InDefaultValue[0] == TCHAR('('));
-	ensure(InDefaultValue[InDefaultValue.Len() - 1] == TCHAR(')')); 
+	if(InDefaultValue[0] != TCHAR('('))
+	{
+		return Parts;
+	}
+	if(InDefaultValue[InDefaultValue.Len() - 1] != TCHAR(')'))
+	{
+		return Parts;
+	}
 
 	FString Content = InDefaultValue.Mid(1, InDefaultValue.Len() - 2);
 	int32 BraceCount = 0;
@@ -1732,6 +1738,12 @@ bool URigVMPin::CanLink(URigVMPin* InSourcePin, URigVMPin* InTargetPin, FString*
 
 		if (bCPPTypesDiffer)
 		{
+			if(RigVMTypeUtils::CanCastTypes(InSourcePin->GetTypeIndex(), InTargetPin->GetTypeIndex()))
+			{
+				bCPPTypesDiffer = false;
+			}
+
+			if(bCPPTypesDiffer)
 			{
 				auto TemplateNodeSupportsType = [](URigVMPin* InPin, const int32& InTypeIndex, FString* OutFailureReason) -> bool
 				{
diff --git a/Engine/Source/Developer/RigVMDeveloper/Public/RigVMCompiler/RigVMCompiler.h b/Engine/Source/Developer/RigVMDeveloper/Public/RigVMCompiler/RigVMCompiler.h
--- a/Engine/Source/Developer/RigVMDeveloper/Public/RigVMCompiler/RigVMCompiler.h
+++ b/Engine/Source/Developer/RigVMDeveloper/Public/RigVMCompiler/RigVMCompiler.h
@@ -109,6 +109,8 @@ public:
 	FRigVMOperand FindProperty(ERigVMMemoryType InMemoryType, const FName& InName);
 	FRigVMPropertyDescription GetProperty(const FRigVMOperand& InOperand);
 	int32 FindOrAddPropertyPath(const FRigVMOperand& InOperand, const FString& InHeadCPPType, const FString& InSegmentPath);
+	const FProperty* GetPropertyForOperand(const FRigVMOperand& InOperand) const;
+	TRigVMTypeIndex GetTypeIndexForOperand(const FRigVMOperand& InOperand) const;
 
 	TSharedPtr<FRigVMParserAST> AST;
 	
diff --git a/Engine/Source/Developer/RigVMDeveloper/Public/RigVMDeveloperTypeUtils.h b/Engine/Source/Developer/RigVMDeveloper/Public/RigVMDeveloperTypeUtils.h
--- a/Engine/Source/Developer/RigVMDeveloper/Public/RigVMDeveloperTypeUtils.h
+++ b/Engine/Source/Developer/RigVMDeveloper/Public/RigVMDeveloperTypeUtils.h
@@ -6,6 +6,7 @@
 #include "RigVMCore/RigVMExternalVariable.h"
 #include "EdGraphSchema_K2.h"
 #include "RigVMTypeUtils.h"
+#include "RigVMCore/RigVMFunction.h"
 
 struct FRigVMGraphVariableDescription;
 class URigVMPin;
@@ -23,6 +24,8 @@ namespace RigVMTypeUtils
 
 	RIGVMDEVELOPER_API FRigVMExternalVariable ExternalVariableFromCPPType(const FName& InName, const FString& InCPPType, UObject* InCPPTypeObject, bool bInPublic = false, bool bInReadonly = false);
 
+	RIGVMDEVELOPER_API FEdGraphPinType PinTypeFromTypeIndex(const TRigVMTypeIndex& InTypeIndex);
+
 	RIGVMDEVELOPER_API FEdGraphPinType PinTypeFromCPPType(const FName& InCPPType, UObject* InCPPTypeObject);
 
 	RIGVMDEVELOPER_API FEdGraphPinType PinTypeFromExternalVariable(const FRigVMExternalVariable& InExternalVariable);
@@ -45,11 +48,19 @@ namespace RigVMTypeUtils
 
 	RIGVMDEVELOPER_API bool CPPTypeFromExternalVariable(const FRigVMExternalVariable& InExternalVariable, FString& OutCPPType, UObject** OutCPPTypeObject);
 
+	RIGVMDEVELOPER_API TRigVMTypeIndex TypeIndexFromPinType(const FEdGraphPinType& InPinType);
+
 	RIGVMDEVELOPER_API bool AreCompatible(const FName& InCPPTypeA, UObject* InCPPTypeObjectA, const FName& InCPPTypeB, UObject* InCPPTypeObjectB);
 
 	RIGVMDEVELOPER_API bool AreCompatible(const FRigVMExternalVariable& InTypeA, const FRigVMExternalVariable& InTypeB, const FString& InSegmentPathA = FString(), const FString& InSegmentPathB = FString());
 	
 	RIGVMDEVELOPER_API bool AreCompatible(const FEdGraphPinType& InTypeA, const FEdGraphPinType& InTypeB, const FString& InSegmentPathA = FString(), const FString& InSegmentPathB = FString());
 	
+	RIGVMDEVELOPER_API const TArray<TRigVMTypeIndex>& GetAvailableCasts(const TRigVMTypeIndex& InTypeIndex, bool bAsInput);
+
+	RIGVMDEVELOPER_API bool CanCastTypes(const TRigVMTypeIndex& InSourceTypeIndex, const TRigVMTypeIndex& InTargetTypeIndex);
+	
+	RIGVMDEVELOPER_API const FRigVMFunction* GetCastForTypeIndices(const TRigVMTypeIndex& InSourceTypeIndex, const TRigVMTypeIndex& InTargetTypeIndex);
+
 #endif
 }
diff --git a/Engine/Source/Developer/RigVMDeveloper/Public/RigVMModel/RigVMController.h b/Engine/Source/Developer/RigVMDeveloper/Public/RigVMModel/RigVMController.h
--- a/Engine/Source/Developer/RigVMDeveloper/Public/RigVMModel/RigVMController.h
+++ b/Engine/Source/Developer/RigVMDeveloper/Public/RigVMModel/RigVMController.h
@@ -1018,6 +1018,7 @@ public:
 	bool AddLink(URigVMPin* OutputPin, URigVMPin* InputPin, bool bSetupUndoRedo = true, ERigVMPinDirection InUserDirection = ERigVMPinDirection::Invalid);
 	bool BreakLink(URigVMPin* OutputPin, URigVMPin* InputPin, bool bSetupUndoRedo = true);
 	bool BreakAllLinks(URigVMPin* Pin, bool bAsInput, bool bSetupUndoRedo = true);
+	void EnableTypeCasting(bool bEnabled = true) { bEnableTypeCasting = bEnabled; }
 
 private:
 	bool BreakAllLinksRecursive(URigVMPin* Pin, bool bAsInput, bool bTowardsParent, bool bSetupUndoRedo);
@@ -1176,6 +1177,7 @@ private:
 	bool bReportWarningsAndErrors;
 	bool bIgnoreRerouteCompactnessChanges;
 	ERigVMPinDirection UserLinkDirection;
+	bool bEnableTypeCasting;
 
 	// temporary maps used for pin redirection
 	// only valid between Detach & ReattachLinksToPinObjects
diff --git a/Engine/Source/Runtime/RigVM/Public/RigVMCore/RigVM.h b/Engine/Source/Runtime/RigVM/Public/RigVMCore/RigVM.h
--- a/Engine/Source/Runtime/RigVM/Public/RigVMCore/RigVM.h
+++ b/Engine/Source/Runtime/RigVM/Public/RigVMCore/RigVM.h
@@ -826,5 +826,6 @@ private:
 #endif
 
 	friend class URigVMCompiler;
+	friend struct FRigVMCompilerWorkData;
 	friend struct FRigVMCodeGenerator;
 };
diff --git a/Engine/Source/Runtime/RigVM/Public/RigVMCore/RigVMExternalVariable.h b/Engine/Source/Runtime/RigVM/Public/RigVMCore/RigVMExternalVariable.h
--- a/Engine/Source/Runtime/RigVM/Public/RigVMCore/RigVMExternalVariable.h
+++ b/Engine/Source/Runtime/RigVM/Public/RigVMCore/RigVMExternalVariable.h
@@ -79,6 +79,12 @@ struct RIGVM_API FRigVMExternalVariable
 			OutTypeName = *RigVMTypeUtils::GetUniqueStructTypeName(StructProperty->Struct);
 			OutTypeObject = StructProperty->Struct;
 		}
+		else if (const FArrayProperty* ArrayProperty = CastField<FArrayProperty>(InProperty))
+		{
+			FName InnerTypeName = NAME_None;
+			GetTypeFromProperty(ArrayProperty->Inner, InnerTypeName, OutTypeObject);
+			OutTypeName = *RigVMTypeUtils::ArrayTypeFromBaseType(InnerTypeName.ToString());
+		}
 		else if (const FObjectProperty* ObjectProperty = CastField<FObjectProperty>(InProperty))
 		{
 			if(RigVMCore::SupportsUObjects())
diff --git a/Engine/Source/Runtime/RigVM/Public/RigVMCore/RigVMMemoryStorage.h b/Engine/Source/Runtime/RigVM/Public/RigVMCore/RigVMMemoryStorage.h
--- a/Engine/Source/Runtime/RigVM/Public/RigVMCore/RigVMMemoryStorage.h
+++ b/Engine/Source/Runtime/RigVM/Public/RigVMCore/RigVMMemoryStorage.h
@@ -856,6 +856,7 @@ private:
 
 	friend class URigVMMemoryStorage;
 	friend class URigVMCompiler;
+	friend struct FRigVMCompilerWorkData;
 	friend class URigVM;
 	friend struct FRigVMCodeGenerator;
 };
