Control Rig: Fix crash during compilation when if node is used

diff --git a/Engine/Source/Developer/RigVMDeveloper/Private/RigVMCompiler/RigVMCompiler.cpp b/Engine/Source/Developer/RigVMDeveloper/Private/RigVMCompiler/RigVMCompiler.cpp
--- a/Engine/Source/Developer/RigVMDeveloper/Private/RigVMCompiler/RigVMCompiler.cpp
+++ b/Engine/Source/Developer/RigVMDeveloper/Private/RigVMCompiler/RigVMCompiler.cpp
@@ -1072,7 +1072,7 @@ void URigVMCompiler::TraverseForLoop(const FRigVMCallExternExprAST* InExpr, FRig
 	// set the index to 0
 	const FRigVMVarExprAST* IndexExpr = InExpr->FindVarWithPinName(FRigVMStruct::ForLoopIndexPinName);
 	check(IndexExpr);
-	FRigVMOperand IndexOperand = WorkData.ExprToOperand.FindChecked(IndexExpr);
+	FRigVMOperand IndexOperand = WorkData.ExprToOperand.FindChecked(GetSourceVarExpr(IndexExpr));
 	WorkData.VM->GetByteCode().AddZeroOp(IndexOperand);
 	if (Settings.SetupNodeInstructionIndex)
 	{
@@ -1089,7 +1089,7 @@ void URigVMCompiler::TraverseForLoop(const FRigVMCallExternExprAST* InExpr, FRig
 	// set up the jump forward (jump out of the loop)
 	const FRigVMVarExprAST* ContinueLoopExpr = InExpr->FindVarWithPinName(FRigVMStruct::ForLoopContinuePinName);
 	check(ContinueLoopExpr);
-	FRigVMOperand ContinueLoopOperand = WorkData.ExprToOperand.FindChecked(ContinueLoopExpr);
+	FRigVMOperand ContinueLoopOperand = WorkData.ExprToOperand.FindChecked(GetSourceVarExpr(ContinueLoopExpr));
 
 	uint64 JumpToEndByte = WorkData.VM->GetByteCode().AddJumpIfOp(ERigVMOpCode::JumpForwardIf, 0, ContinueLoopOperand, false);
 	int32 JumpToEndInstruction = WorkData.VM->GetByteCode().GetNumInstructions() - 1;
@@ -1101,7 +1101,7 @@ void URigVMCompiler::TraverseForLoop(const FRigVMCallExternExprAST* InExpr, FRig
 	// begin the loop's block
 	const FRigVMVarExprAST* CountExpr = InExpr->FindVarWithPinName(FRigVMStruct::ForLoopCountPinName);
 	check(CountExpr);
-	FRigVMOperand CountOperand = WorkData.ExprToOperand.FindChecked(CountExpr);
+	FRigVMOperand CountOperand = WorkData.ExprToOperand.FindChecked(GetSourceVarExpr(CountExpr));
 	WorkData.VM->GetByteCode().AddBeginBlockOp(CountOperand, IndexOperand);
 	if (Settings.SetupNodeInstructionIndex)
 	{
@@ -1217,7 +1217,7 @@ void URigVMCompiler::TraverseAssign(const FRigVMAssignExprAST* InExpr, FRigVMCom
 		const FRigVMVarExprAST* TargetExpr = InExpr->GetFirstParentOfType(FRigVMVarExprAST::EType::Var)->To<FRigVMVarExprAST>();
 		TargetExpr = GetSourceVarExpr(TargetExpr);
 		
-		FRigVMOperand Target = WorkData.ExprToOperand.FindChecked(TargetExpr);
+		FRigVMOperand Target = WorkData.ExprToOperand.FindChecked(GetSourceVarExpr(TargetExpr));
 		if(Target == Source)
 		{
 			return;
@@ -1344,7 +1344,7 @@ void URigVMCompiler::TraverseBranch(const FRigVMBranchExprAST* InExpr, FRigVMCom
 		ConditionExpr = ConditionExpr->To<FRigVMCachedValueExprAST>()->GetVarExpr();
 	}
 
-	FRigVMOperand& ConditionOperand = WorkData.ExprToOperand.FindChecked(ConditionExpr);
+	FRigVMOperand& ConditionOperand = WorkData.ExprToOperand.FindChecked(GetSourceVarExpr(ConditionExpr));
 
 	// setup the first jump
 	uint64 JumpToFalseByte = WorkData.VM->GetByteCode().AddJumpIfOp(ERigVMOpCode::JumpForwardIf, 1, ConditionOperand, false);
@@ -1407,8 +1407,8 @@ void URigVMCompiler::TraverseIf(const FRigVMIfExprAST* InExpr, FRigVMCompilerWor
 		ConditionExpr = ConditionExpr->To<FRigVMCachedValueExprAST>()->GetVarExpr();
 	}
 
-	FRigVMOperand& ConditionOperand = WorkData.ExprToOperand.FindChecked(ConditionExpr);
-	FRigVMOperand& ResultOperand = WorkData.ExprToOperand.FindChecked(ResultExpr);
+	FRigVMOperand& ConditionOperand = WorkData.ExprToOperand.FindChecked(GetSourceVarExpr(ConditionExpr));
+	FRigVMOperand& ResultOperand = WorkData.ExprToOperand.FindChecked(GetSourceVarExpr(ResultExpr));
 
 	// setup the first jump
 	uint64 JumpToFalseByte = WorkData.VM->GetByteCode().AddJumpIfOp(ERigVMOpCode::JumpForwardIf, 1, ConditionOperand, false);
@@ -1426,7 +1426,7 @@ void URigVMCompiler::TraverseIf(const FRigVMIfExprAST* InExpr, FRigVMCompilerWor
 		TrueExpr = TrueExpr->To<FRigVMCachedValueExprAST>()->GetVarExpr();
 	}
 
-	FRigVMOperand& TrueOperand = WorkData.ExprToOperand.FindChecked(TrueExpr);
+	FRigVMOperand& TrueOperand = WorkData.ExprToOperand.FindChecked(GetSourceVarExpr(TrueExpr));
 
 	WorkData.VM->GetByteCode().AddCopyOp(WorkData.VM->GetCopyOpForOperands(TrueOperand, ResultOperand));
 	if (Settings.SetupNodeInstructionIndex)
@@ -1453,7 +1453,7 @@ void URigVMCompiler::TraverseIf(const FRigVMIfExprAST* InExpr, FRigVMCompilerWor
 		FalseExpr = FalseExpr->To<FRigVMCachedValueExprAST>()->GetVarExpr();
 	}
 
-	FRigVMOperand& FalseOperand = WorkData.ExprToOperand.FindChecked(FalseExpr);
+	FRigVMOperand& FalseOperand = WorkData.ExprToOperand.FindChecked(GetSourceVarExpr(FalseExpr));
 
 	WorkData.VM->GetByteCode().AddCopyOp(WorkData.VM->GetCopyOpForOperands(FalseOperand, ResultOperand));
 	if (Settings.SetupNodeInstructionIndex)
@@ -1536,8 +1536,8 @@ void URigVMCompiler::TraverseSelect(const FRigVMSelectExprAST* InExpr, FRigVMCom
 		IndexExpr = IndexExpr->To<FRigVMCachedValueExprAST>()->GetVarExpr();
 	}
 
-	FRigVMOperand& IndexOperand = WorkData.ExprToOperand.FindChecked(IndexExpr);
-	FRigVMOperand& ResultOperand = WorkData.ExprToOperand.FindChecked(ResultExpr);
+	FRigVMOperand& IndexOperand = WorkData.ExprToOperand.FindChecked(GetSourceVarExpr(IndexExpr));
+	FRigVMOperand& ResultOperand = WorkData.ExprToOperand.FindChecked(GetSourceVarExpr(ResultExpr));
 
 	// setup the jumps for each case
 	TArray<uint64> JumpToCaseBytes;
