Control Rig: Inline function bytecode

diff --git a/Engine/Plugins/Animation/ControlRig/Source/ControlRigDeveloper/Private/ControlRigBlueprint.cpp b/Engine/Plugins/Animation/ControlRig/Source/ControlRigDeveloper/Private/ControlRigBlueprint.cpp
--- a/Engine/Plugins/Animation/ControlRig/Source/ControlRigDeveloper/Private/ControlRigBlueprint.cpp
+++ b/Engine/Plugins/Animation/ControlRig/Source/ControlRigDeveloper/Private/ControlRigBlueprint.cpp
@@ -2099,6 +2099,8 @@ TArray<FString> UControlRigBlueprint::GeneratePythonCommands(const FString InNew
 	InternalCommands.Add(TEXT("library_controller = blueprint.get_controller(library)"));
 	InternalCommands.Add(TEXT("hierarchy = blueprint.hierarchy"));
 	InternalCommands.Add(TEXT("hierarchy_controller = hierarchy.get_controller()"));
+	InternalCommands.Add(TEXT("blueprint.set_auto_vm_recompile(False)"));
+	
 
 	// Hierarchy
 	InternalCommands.Append(Hierarchy->GetController(true)->GeneratePythonCommands());
@@ -2315,6 +2317,8 @@ TArray<FString> UControlRigBlueprint::GeneratePythonCommands(const FString InNew
 		*PreviewMeshPath));
 #endif
 
+	InternalCommands.Add(TEXT("blueprint.set_auto_vm_recompile(True)"));
+
 	// Split multiple commands into different array elements
 	TArray<FString> InnerFunctionCmds;
 	for (FString Cmd : InternalCommands)
diff --git a/Engine/Plugins/Animation/ControlRig/Source/ControlRigDeveloper/Private/Graph/ControlRigGraphNode.cpp b/Engine/Plugins/Animation/ControlRig/Source/ControlRigDeveloper/Private/Graph/ControlRigGraphNode.cpp
--- a/Engine/Plugins/Animation/ControlRig/Source/ControlRigDeveloper/Private/Graph/ControlRigGraphNode.cpp
+++ b/Engine/Plugins/Animation/ControlRig/Source/ControlRigDeveloper/Private/Graph/ControlRigGraphNode.cpp
@@ -1483,7 +1483,7 @@ void UControlRigGraphNode::SetupPinDefaultsFromModel(UEdGraphPin* Pin, const URi
 	}
 
 	const UEdGraphSchema_K2* K2Schema = GetDefault<UEdGraphSchema_K2>();
-	if (InModelPin)
+	if (InModelPin && !IsValid(InModelPin))
 	{
 		if (InModelPin->GetSubPins().Num() > 0)
 		{
diff --git a/Engine/Plugins/Animation/ControlRig/Source/ControlRigDeveloper/Public/ControlRigBlueprint.h b/Engine/Plugins/Animation/ControlRig/Source/ControlRigDeveloper/Public/ControlRigBlueprint.h
--- a/Engine/Plugins/Animation/ControlRig/Source/ControlRigDeveloper/Public/ControlRigBlueprint.h
+++ b/Engine/Plugins/Animation/ControlRig/Source/ControlRigDeveloper/Public/ControlRigBlueprint.h
@@ -293,6 +293,9 @@ public:
 	UFUNCTION(BlueprintCallable, Category = "Control Rig Blueprint")
 	void RequestAutoVMRecompilation();
 
+	UFUNCTION(BlueprintCallable, Category = "Control Rig Blueprint")
+	void SetAutoVMRecompile(bool bAutoRecompile) { bAutoRecompileVM = bAutoRecompile; }
+
 	void IncrementVMRecompileBracket();
 	void DecrementVMRecompileBracket();
 
diff --git a/Engine/Source/Developer/RigVMDeveloper/Private/RigVMCompiler/RigVMAST.cpp b/Engine/Source/Developer/RigVMDeveloper/Private/RigVMCompiler/RigVMAST.cpp
--- a/Engine/Source/Developer/RigVMDeveloper/Private/RigVMCompiler/RigVMAST.cpp
+++ b/Engine/Source/Developer/RigVMDeveloper/Private/RigVMCompiler/RigVMAST.cpp
@@ -48,6 +48,10 @@ FName FRigVMExprAST::GetTypeName() const
 		{
 			return TEXT("[.Call..]");
 		}
+		case EType::InlineFunction:
+		{
+			return TEXT("[.Inline..]");
+		}
 		case EType::NoOp:
 		{
 			return TEXT("[.NoOp..]");
@@ -754,9 +758,15 @@ void FRigVMParserASTSettings::Report(EMessageSeverity::Type InSeverity, UObject*
 const TArray<FRigVMASTProxy> FRigVMParserAST::EmptyProxyArray;
 
 FRigVMParserAST::FRigVMParserAST(TArray<URigVMGraph*> InGraphs, URigVMController* InController, const FRigVMParserASTSettings& InSettings, const TArray<FRigVMExternalVariable>& InExternalVariables, const TArray<FRigVMUserDataArray>& InRigVMUserData)
+	: LibraryNodeBeingCompiled(nullptr)
 {
 	DECLARE_SCOPE_HIERARCHICAL_COUNTER_FUNC()
 
+	if (InGraphs.Num() == 1 && InGraphs[0]->GetTypedOuter<URigVMFunctionLibrary>() != nullptr)
+	{
+		LibraryNodeBeingCompiled = Cast<URigVMLibraryNode>(InGraphs[0]->GetOuter());
+	}
+
 	Settings = InSettings;
 	ObsoleteBlock = nullptr;
 	LastCycleCheckExpr = nullptr;
@@ -767,12 +777,48 @@ FRigVMParserAST::FRigVMParserAST(TArray<URigVMGraph*> InGraphs, URigVMController
 	// construct the inlined nodes and links information
 	Inline(InGraphs);
 
-	for (const FRigVMASTProxy& NodeProxy : NodeProxies)
+	if (LibraryNodeBeingCompiled == nullptr)
 	{
-		URigVMNode* Node = NodeProxy.GetSubjectChecked<URigVMNode>();
-		if(Node->IsEvent())
+		for (const FRigVMASTProxy& NodeProxy : NodeProxies)
 		{
-			TraverseMutableNode(NodeProxy, nullptr);
+			URigVMNode* Node = NodeProxy.GetSubjectChecked<URigVMNode>();
+			if(Node->IsEvent())
+			{
+				TraverseMutableNode(NodeProxy, nullptr);
+			}
+		}
+	}
+	else
+	{
+		URigVMFunctionEntryNode* Entry = LibraryNodeBeingCompiled->GetEntryNode();
+		URigVMFunctionReturnNode* Return = LibraryNodeBeingCompiled->GetReturnNode();
+		if (Entry && Entry->IsMutable())
+		{
+			for (const FRigVMASTProxy& NodeProxy : NodeProxies)
+			{
+				URigVMNode* Node = NodeProxy.GetSubjectChecked<URigVMNode>();
+				if(Node == Entry)
+				{
+					TraverseMutableNode(NodeProxy, nullptr);
+					break;
+				}
+			}
+		}
+		else if (Return)
+		{
+			for (const FRigVMASTProxy& NodeProxy : NodeProxies)
+			{
+				URigVMNode* Node = NodeProxy.GetSubjectChecked<URigVMNode>();
+				if(Node == Return)
+				{
+					TraverseNode(NodeProxy, nullptr);
+					break;
+				}
+			}
+		}
+		else
+		{
+			return;
 		}
 	}
 
@@ -839,6 +885,7 @@ FRigVMParserAST::FRigVMParserAST(TArray<URigVMGraph*> InGraphs, URigVMController
 }
 
 FRigVMParserAST::FRigVMParserAST(TArray<URigVMGraph*> InGraphs, const TArray<FRigVMASTProxy>& InNodesToCompute)
+	: LibraryNodeBeingCompiled(nullptr)
 {
 	check(!InGraphs.IsEmpty());
 
@@ -995,7 +1042,11 @@ FRigVMExprAST* FRigVMParserAST::CreateExpressionForNode(const FRigVMASTProxy& In
 	}
 	else
 	{
-		if (InNodeProxy.IsA<URigVMRerouteNode>() ||
+		if (InNodeProxy.IsA<URigVMFunctionReferenceNode>())
+		{
+			NodeExpr = MakeExpr<FRigVMInlineFunctionExprAST>(InNodeProxy);
+		}
+		else if (InNodeProxy.IsA<URigVMRerouteNode>() ||
 			InNodeProxy.IsA<URigVMVariableNode>() ||
 			InNodeProxy.IsA<URigVMEnumNode>() ||
 			InNodeProxy.IsA<URigVMLibraryNode>() ||
@@ -1278,6 +1329,7 @@ FRigVMExprAST* FRigVMParserAST::TraverseLink(int32 InLinkIndex, FRigVMExprAST* I
 	{
 		// if this is a copy expression - we should require the copy to use a ref instead
 		if (NodeExpr->IsA(FRigVMExprAST::EType::CallExtern) ||
+			NodeExpr->IsA(FRigVMExprAST::EType::InlineFunction) ||
 			NodeExpr->IsA(FRigVMExprAST::EType::If) ||
 			NodeExpr->IsA(FRigVMExprAST::EType::Select) ||
 			NodeExpr->IsA(FRigVMExprAST::EType::Array))
@@ -1320,6 +1372,19 @@ FRigVMExprAST* FRigVMParserAST::TraverseLink(int32 InLinkIndex, FRigVMExprAST* I
 			}
 			checkNoEntry();
 		}
+		else if (LibraryNodeBeingCompiled &&
+			SourceRootPin->GetNode()->IsA<URigVMFunctionEntryNode>() &&
+			SourceRootPin->GetTypedOuter<URigVMLibraryNode>() == LibraryNodeBeingCompiled)
+		{
+			FRigVMASTProxy SourceRootProxy = SourceProxy;
+			if (SourceRootPin != SourcePin)
+			{
+				SourceRootProxy = FRigVMASTProxy::MakeFromUObject(SourceRootPin);
+			}
+			FRigVMVarExprAST* SourcePinExpr = MakeExpr<FRigVMVarExprAST>(FRigVMExprAST::EType::Var, SourceRootProxy);
+			AssignExpr->ReplaceChild(NodeExpr, SourcePinExpr);
+			return AssignExpr;
+		}			
 	}
 
 	return AssignExpr;
@@ -1540,6 +1605,25 @@ void FRigVMParserAST::FoldNoOps()
 						continue;
 					}
 				}
+				if (LibraryNodeBeingCompiled != nullptr)
+				{
+					if (URigVMFunctionEntryNode* EntryNode = Cast<URigVMFunctionEntryNode>(Node))
+					{
+						if (EntryNode->GetTypedOuter<URigVMLibraryNode>() == LibraryNodeBeingCompiled &&
+							EntryNode->IsMutable())
+						{
+							continue;
+						}
+					}
+					if (URigVMFunctionReturnNode* ReturnNode = Cast<URigVMFunctionReturnNode>(Node))
+					{
+						if (ReturnNode->GetTypedOuter<URigVMLibraryNode>() == LibraryNodeBeingCompiled &&
+							!ReturnNode->IsMutable())
+						{
+							continue;
+						}
+					}
+				}
 			}
 			// copy since we are changing the content during iteration below
 			TArray<FRigVMExprAST*> Children = Expression->Children;
@@ -1913,7 +1997,7 @@ bool FRigVMParserAST::FoldConstantValuesToLiterals(TArray<URigVMGraph*> InGraphs
 			}
 		}
 
-		FString PinHash = URigVMCompiler::GetPinHash(RootPin, RootVarExpr, false);
+		FString PinHash = URigVMCompiler::GetPinHash(RootPin, RootVarExpr, false, LibraryNodeBeingCompiled);
 		const FRigVMOperand* OperandPtr = Operands.Find(PinHash);
 		if(OperandPtr == nullptr)
 		{
@@ -2383,6 +2467,12 @@ bool FRigVMParserAST::CanLink(URigVMPin* InSourcePin, URigVMPin* InTargetPin, FS
 		return false;
 	}
 
+	// If the source node is an entry of a function, no need to cycle check
+	if (SourceNode->IsA<URigVMFunctionEntryNode>() && SourceNode->GetTypedOuter<URigVMFunctionLibrary>() != nullptr)
+	{
+		return true;
+	}
+
 	if (SourceBlock == TargetBlock ||
 		SourceBlock->Contains(TargetBlock) ||
 		TargetBlock->Contains(SourceBlock) ||
@@ -2610,6 +2700,14 @@ FString FRigVMParserAST::DumpDot() const
 					}
 					break;
 				}
+				case FRigVMExprAST::EType::InlineFunction:
+				{
+					if (URigVMFunctionReferenceNode* Node = Cast<URigVMFunctionReferenceNode>(InExpr->To<FRigVMInlineFunctionExprAST>()->GetNode()))
+					{
+						Label = FString::Printf(TEXT("Inline %s"), *Node->GetReferencedNode()->GetName());
+					}
+					break;
+				}
 				case FRigVMExprAST::EType::NoOp:
 				{
 					Label = TEXT("NoOp");
@@ -2655,6 +2753,7 @@ FString FRigVMParserAST::DumpDot() const
 				case FRigVMExprAST::EType::Assign:
 				case FRigVMExprAST::EType::Copy:
 				case FRigVMExprAST::EType::CallExtern:
+				case FRigVMExprAST::EType::InlineFunction:
 				case FRigVMExprAST::EType::If:
 				case FRigVMExprAST::EType::Select:
 				case FRigVMExprAST::EType::NoOp:
@@ -2909,8 +3008,9 @@ void FRigVMParserAST::Inline(TArray<URigVMGraph*> InGraphs, const TArray<FRigVMA
 		TArray<FRigVMASTLinkDescription>* Links;
 		TArray<FRigVMASTProxy> LibraryNodeCallstack;
 		FRigVMParserASTSettings* Settings;
+		URigVMLibraryNode* LibraryNodeBeingCompiled;
 
-		static bool ShouldRecursePin(const FRigVMASTProxy& InPinProxy)
+		static bool ShouldRecursePin(const FRigVMASTProxy& InPinProxy, LocalPinTraversalInfo& OutTraversalInfo)
 		{
 			URigVMPin* Pin = InPinProxy.GetSubjectChecked<URigVMPin>();
 			URigVMNode* Node = Pin->GetNode();
@@ -2918,6 +3018,25 @@ void FRigVMParserAST::Inline(TArray<URigVMGraph*> InGraphs, const TArray<FRigVMA
 			{
 				return VarNode->IsInputArgument();
 			}
+
+			// If its an interface node of the library we are compiling, don't recurse
+			if (OutTraversalInfo.LibraryNodeBeingCompiled != nullptr)
+			{
+				if (Node->IsA<URigVMFunctionEntryNode>() ||
+					Node->IsA<URigVMFunctionReturnNode>())
+				{
+					if (Node->GetTypedOuter<URigVMLibraryNode>() == OutTraversalInfo.LibraryNodeBeingCompiled)
+					{
+						return false;
+					}
+				}
+			}
+
+			// If its a function reference, don't recurse
+			if (Node->IsA<URigVMFunctionReferenceNode>())
+			{
+				return false;				
+			}
 			
 			return Node->IsA<URigVMRerouteNode>() ||
 				Node->IsA<URigVMLibraryNode>() ||
@@ -2925,14 +3044,14 @@ void FRigVMParserAST::Inline(TArray<URigVMGraph*> InGraphs, const TArray<FRigVMA
 				Node->IsA<URigVMFunctionReturnNode>();
 		}
 
-		static bool IsValidPinForAST(const FRigVMASTProxy& InPinProxy)
+		static bool IsValidPinForAST(const FRigVMASTProxy& InPinProxy, LocalPinTraversalInfo& OutTraversalInfo)
 		{
-			return !ShouldRecursePin(InPinProxy);
+			return !ShouldRecursePin(InPinProxy, OutTraversalInfo);
 		}
 
-		static bool IsValidLinkForAST(const FRigVMASTProxy& InSourcePinProxy, const FRigVMASTProxy& InTargetPinProxy)
+		static bool IsValidLinkForAST(const FRigVMASTProxy& InSourcePinProxy, const FRigVMASTProxy& InTargetPinProxy, LocalPinTraversalInfo& OutTraversalInfo)
 		{
-			return IsValidPinForAST(InSourcePinProxy) && IsValidPinForAST(InTargetPinProxy);
+			return IsValidPinForAST(InSourcePinProxy, OutTraversalInfo) && IsValidPinForAST(InTargetPinProxy, OutTraversalInfo);
 		}
 
 		static FRigVMASTProxy FindSourcePin(const FRigVMASTProxy& InPinProxy, LocalPinTraversalInfo& OutTraversalInfo)
@@ -2952,17 +3071,17 @@ void FRigVMParserAST::Inline(TArray<URigVMGraph*> InGraphs, const TArray<FRigVMA
 					Pin->GetDirection() == ERigVMPinDirection::IO)
 				{
 					URigVMNode* Node = Pin->GetNode();
-					if (URigVMLibraryNode* LibraryNode = Cast<URigVMLibraryNode>(Node))
+					if (URigVMCollapseNode* CollapseNode = Cast<URigVMCollapseNode>(Node))
 					{
-						FRigVMASTProxy LibraryNodeProxy = InPinProxy.GetSibling(LibraryNode);
-						if (!OutTraversalInfo.LibraryNodeCallstack.Contains(LibraryNodeProxy))
+						FRigVMASTProxy CollapseNodeProxy = InPinProxy.GetSibling(CollapseNode);
+						if (!OutTraversalInfo.LibraryNodeCallstack.Contains(CollapseNodeProxy))
 						{
-							if (URigVMFunctionReturnNode* ReturnNode = LibraryNode->GetReturnNode())
+							if (URigVMFunctionReturnNode* ReturnNode = CollapseNode->GetReturnNode())
 							{
 								if (URigVMPin* ReturnPin = ReturnNode->FindPin(Pin->GetName()))
 								{
-									OutTraversalInfo.LibraryNodeCallstack.Push(LibraryNodeProxy);
-									FRigVMASTProxy ReturnPinProxy = LibraryNodeProxy.GetChild(ReturnPin);
+									OutTraversalInfo.LibraryNodeCallstack.Push(CollapseNodeProxy);
+									FRigVMASTProxy ReturnPinProxy = CollapseNodeProxy.GetChild(ReturnPin);
 									FRigVMASTProxy SourcePinProxy = FindSourcePin(ReturnPinProxy, OutTraversalInfo);
 									SourcePinProxy = SourcePinProxy.IsValid() ? SourcePinProxy : ReturnPinProxy;
 									if(bStoreSourcePinOnMap)
@@ -2998,6 +3117,10 @@ void FRigVMParserAST::Inline(TArray<URigVMGraph*> InGraphs, const TArray<FRigVMA
 					}
 					else if (URigVMFunctionEntryNode* EntryNode = Cast<URigVMFunctionEntryNode>(Node))
 					{
+						if (EntryNode->GetGraph()->GetOuter() == OutTraversalInfo.LibraryNodeBeingCompiled)
+						{
+							return FRigVMASTProxy(); 
+						}
 						for (int32 LibraryNodeIndex = OutTraversalInfo.LibraryNodeCallstack.Num() - 1; LibraryNodeIndex >= 0; LibraryNodeIndex--)
 						{
 							FRigVMASTProxy LibraryNodeProxy = OutTraversalInfo.LibraryNodeCallstack[LibraryNodeIndex];
@@ -3062,6 +3185,11 @@ void FRigVMParserAST::Inline(TArray<URigVMGraph*> InGraphs, const TArray<FRigVMA
 				{
 					if (URigVMFunctionEntryNode* EntryNode = Cast<URigVMFunctionEntryNode>(ChildPin->GetNode()))
 					{
+						if (EntryNode->GetGraph()->GetOuter() == OutTraversalInfo.LibraryNodeBeingCompiled)
+						{
+							return FRigVMASTProxy(); 
+						}
+						
 						// rather than relying on the other we are going to query what's in the call stack.
 						// for collapse nodes that's not a different, but for function ref nodes the outer
 						// node is the definition and not the reference node - which sits in the callstack.
@@ -3081,7 +3209,7 @@ void FRigVMParserAST::Inline(TArray<URigVMGraph*> InGraphs, const TArray<FRigVMA
 						if(ChildPin != InPinProxy.GetSubject<URigVMPin>())
 						{
 							const FRigVMASTProxy ChildPinProxy = InPinProxy.GetSibling(ChildPin); 
-							if (ShouldRecursePin(ChildPinProxy))
+							if (ShouldRecursePin(ChildPinProxy, OutTraversalInfo))
 							{
 								FRigVMASTProxy SourceSourcePinProxy = FindSourcePin(ChildPinProxy, OutTraversalInfo);
 								if(SourceSourcePinProxy.IsValid())
@@ -3132,7 +3260,7 @@ void FRigVMParserAST::Inline(TArray<URigVMGraph*> InGraphs, const TArray<FRigVMA
 				{
 					SourcePinProxy = InPinProxy.GetSibling(SourcePin);
 
-					if (ShouldRecursePin(SourcePinProxy))
+					if (ShouldRecursePin(SourcePinProxy, OutTraversalInfo))
 					{
 						FRigVMASTProxy SourceSourcePinProxy = FindSourcePin(SourcePinProxy, OutTraversalInfo);
 						if(SourceSourcePinProxy.IsValid())
@@ -3155,7 +3283,7 @@ void FRigVMParserAST::Inline(TArray<URigVMGraph*> InGraphs, const TArray<FRigVMA
 						FRigVMASTProxy& ParentSourcePinProxy = *ParentSourcePinProxyPtr;
 						if (ParentSourcePinProxy.IsValid())
 						{
-							if (!ShouldRecursePin(ParentSourcePinProxy))
+							if (!ShouldRecursePin(ParentSourcePinProxy, OutTraversalInfo))
 							{
 								// only discard results here if we haven't crossed a collapse node boundary
 								if(ParentSourcePinProxy.GetSubjectChecked<URigVMPin>()->GetGraph() == ChildPin->GetGraph())
@@ -3183,7 +3311,7 @@ void FRigVMParserAST::Inline(TArray<URigVMGraph*> InGraphs, const TArray<FRigVMA
 					{
 						SourcePinProxy = SourcePinProxy.GetSibling(SourceSubPin);
 
-						if (ShouldRecursePin(SourcePinProxy))
+						if (ShouldRecursePin(SourcePinProxy, OutTraversalInfo))
 						{
 							FRigVMASTProxy SourceSourceSubPinProxy = FindSourcePin(SourcePinProxy, OutTraversalInfo);
 							if (SourceSourceSubPinProxy.IsValid())
@@ -3229,7 +3357,7 @@ void FRigVMParserAST::Inline(TArray<URigVMGraph*> InGraphs, const TArray<FRigVMA
 				URigVMPin* SourcePin = SourcePinProxy.GetSubjectChecked<URigVMPin>();
 				URigVMNode* SourceNode = SourcePin->GetNode();
 				if (SourceNode->IsA<URigVMRerouteNode>() ||
-					SourceNode->IsA<URigVMLibraryNode>() ||
+					SourceNode->IsA<URigVMCollapseNode>() ||
 					SourceNode->IsA<URigVMFunctionReturnNode>())
 				{
 					// for arrays - if there are sub-pins on the determined source, we need to walk those as well
@@ -3255,7 +3383,7 @@ void FRigVMParserAST::Inline(TArray<URigVMGraph*> InGraphs, const TArray<FRigVMA
 				}
 				else
 				{
-					if (IsValidLinkForAST(SourcePinProxy, InPinProxyForMap))
+					if (IsValidLinkForAST(SourcePinProxy, InPinProxyForMap, OutTraversalInfo))
 					{
 						FRigVMASTLinkDescription Link(SourcePinProxy, InPinProxyForMap, InSegmentPath);
 						Link.LinkIndex = OutTraversalInfo.Links->Num();
@@ -3280,12 +3408,32 @@ void FRigVMParserAST::Inline(TArray<URigVMGraph*> InGraphs, const TArray<FRigVMA
 
 		static void VisitNode(const FRigVMASTProxy& InNodeProxy, LocalPinTraversalInfo& OutTraversalInfo)
 		{
-			if (InNodeProxy.IsA<URigVMRerouteNode>() ||
-				InNodeProxy.IsA<URigVMFunctionEntryNode>() ||
-				InNodeProxy.IsA<URigVMFunctionReturnNode>())
+			if (InNodeProxy.IsA<URigVMRerouteNode>())
 			{
 				return;
 			}
+			
+			const bool bIsCompilingFunction = OutTraversalInfo.LibraryNodeBeingCompiled != nullptr;
+			if (bIsCompilingFunction)
+			{
+				if (InNodeProxy.IsA<URigVMFunctionEntryNode>() ||
+				   InNodeProxy.IsA<URigVMFunctionReturnNode>())
+				{
+					URigVMNode* Node = InNodeProxy.GetSubjectChecked<URigVMNode>();
+					if (Node->GetTypedOuter<URigVMLibraryNode>() != OutTraversalInfo.LibraryNodeBeingCompiled)
+					{
+						return;
+					}
+				}
+			}
+			else
+			{
+				if (InNodeProxy.IsA<URigVMFunctionEntryNode>() ||
+				   InNodeProxy.IsA<URigVMFunctionReturnNode>())
+				{
+					return;
+				}
+			}
 
 			if (URigVMLibraryNode* LibraryNode = InNodeProxy.GetSubject<URigVMLibraryNode>())
 			{
@@ -3313,15 +3461,26 @@ void FRigVMParserAST::Inline(TArray<URigVMGraph*> InGraphs, const TArray<FRigVMA
 					return;
 				}
 
-				OutTraversalInfo.LibraryNodeCallstack.Push(InNodeProxy);
-				TArray<URigVMNode*> ContainedNodes = LibraryNode->GetContainedNodes();
-				for (URigVMNode* ContainedNode : ContainedNodes)
+				if (URigVMFunctionReferenceNode* FunctionRef = Cast<URigVMFunctionReferenceNode>(LibraryNode))
 				{
-					// create a proxy which uses the previous node as a callstack
-					FRigVMASTProxy ContainedNodeProxy = InNodeProxy.GetChild(ContainedNode);
-					VisitNode(ContainedNodeProxy, OutTraversalInfo);
+					for (URigVMPin* Pin : FunctionRef->GetPins())
+					{
+						FRigVMASTProxy PinProxy = InNodeProxy.GetSibling(Pin);
+						LocalPinTraversalInfo::VisitPin(PinProxy, OutTraversalInfo);
+					}
 				}
-				OutTraversalInfo.LibraryNodeCallstack.Pop();
+				else
+				{
+					OutTraversalInfo.LibraryNodeCallstack.Push(InNodeProxy);
+					TArray<URigVMNode*> ContainedNodes = LibraryNode->GetContainedNodes();
+					for (URigVMNode* ContainedNode : ContainedNodes)
+					{
+						// create a proxy which uses the previous node as a callstack
+						FRigVMASTProxy ContainedNodeProxy = InNodeProxy.GetChild(ContainedNode);
+						VisitNode(ContainedNodeProxy, OutTraversalInfo);
+					}
+					OutTraversalInfo.LibraryNodeCallstack.Pop();
+				}				
 			}
 			else
 			{
@@ -3352,6 +3511,7 @@ void FRigVMParserAST::Inline(TArray<URigVMGraph*> InGraphs, const TArray<FRigVMA
 	TraversalInfo.SourceLinkIndices = &SourceLinkIndices;
 	TraversalInfo.Links = &Links;
 	TraversalInfo.Settings = &Settings;
+	TraversalInfo.LibraryNodeBeingCompiled = this->LibraryNodeBeingCompiled;
 
 	for (const FRigVMASTProxy& NodeProxy : NodeProxies)
 	{
diff --git a/Engine/Source/Developer/RigVMDeveloper/Private/RigVMCompiler/RigVMCompiler.cpp b/Engine/Source/Developer/RigVMDeveloper/Private/RigVMCompiler/RigVMCompiler.cpp
--- a/Engine/Source/Developer/RigVMDeveloper/Private/RigVMCompiler/RigVMCompiler.cpp
+++ b/Engine/Source/Developer/RigVMDeveloper/Private/RigVMCompiler/RigVMCompiler.cpp
@@ -225,6 +225,7 @@ TRigVMTypeIndex FRigVMCompilerWorkData::GetTypeIndexForOperand(const FRigVMOpera
 }
 
 URigVMCompiler::URigVMCompiler()
+	: CurrentCompilationFunction(nullptr)
 {
 }
 
@@ -346,9 +347,9 @@ bool URigVMCompiler::Compile(TArray<URigVMGraph*> InGraphs, URigVMController* In
 				}
 			}
 
-			if(URigVMLibraryNode* LibraryNode = Cast<URigVMLibraryNode>(ModelNode))
+			if(URigVMCollapseNode* CollapseNode = Cast<URigVMCollapseNode>(ModelNode))
 			{
-				if(URigVMGraph* ContainedGraph = LibraryNode->GetContainedGraph())
+				if(URigVMGraph* ContainedGraph = CollapseNode->GetContainedGraph())
 				{
 					VisitedGraphs.AddUnique(ContainedGraph);
 				}
@@ -466,6 +467,112 @@ bool URigVMCompiler::Compile(TArray<URigVMGraph*> InGraphs, URigVMController* In
 		OutOperands->Add(Hash, Operand);
 	}
 
+	URigVMFunctionLibrary* FunctionLibrary = InGraphs[0]->GetDefaultFunctionLibrary();
+	if (CurrentCompilationFunction == nullptr)
+	{
+		ClearFunctionCompilationData();
+
+		// We need to compile the functions in order of dependency (leaf functions with no dependency first)
+		TArray<URigVMLibraryNode*> FunctionsToCompile;
+		if (FunctionLibrary)
+		{
+			FunctionsToCompile = FunctionLibrary->GetFunctions();
+		}
+
+		// Add public functions from other assets
+		// todo: this should be a check if the compilation data exists
+		{
+			TArray<URigVMNode*> Nodes = InGraphs[0]->GetNodes();
+			for (int32 i=0; i<Nodes.Num(); ++i)
+			{
+				if (URigVMFunctionReferenceNode* ReferenceNode = Cast<URigVMFunctionReferenceNode>(Nodes[i]))
+				{
+					FunctionsToCompile.AddUnique(ReferenceNode->GetReferencedNode());
+				}
+				if (URigVMCollapseNode* CollapseNode = Cast<URigVMCollapseNode>(Nodes[i]))
+				{
+					Nodes.Append(CollapseNode->GetContainedGraph()->GetNodes());
+				}
+			}
+		}
+		while (CompiledFunctions.Num() < FunctionsToCompile.Num())
+		{
+			for (int32 FunctionIndex=0; FunctionIndex<FunctionsToCompile.Num(); ++FunctionIndex)
+			{
+				URigVMLibraryNode* FunctionToCompile = FunctionsToCompile[FunctionIndex];
+				if (CompiledFunctions.Contains(FunctionToCompile))
+				{
+					continue;
+				}
+			
+				bool bCanBeCompiled = true;
+				TArray<URigVMNode*> Nodes = FunctionToCompile->GetContainedNodes();
+				for (int32 i=0; i<Nodes.Num() && bCanBeCompiled; ++i)
+				{
+					if (URigVMFunctionReferenceNode* FunctionReference = Cast<URigVMFunctionReferenceNode>(Nodes[i]))
+					{
+						if (const URigVMFunctionLibrary* ReferencedFunctionLibrary = FunctionReference->GetReferencedNode()->GetLibrary())
+						{
+							if (ReferencedFunctionLibrary == FunctionLibrary)
+							{
+								if (!CompiledFunctions.Contains(FunctionReference->GetReferencedNode()))
+								{
+									bCanBeCompiled = false;
+									break;
+								}
+							}
+							else
+							{
+								// Public function.
+								// todo Try to find its compilation data
+								bool bFoundCompilationData = true;
+								if (!CompiledFunctions.Contains(FunctionReference->GetReferencedNode()))
+								{
+									FunctionsToCompile.AddUnique(FunctionReference->GetReferencedNode());
+									bCanBeCompiled = false;
+									break;
+								}
+
+								if (!bFoundCompilationData)
+								{
+									static const FString FunctionCompilationErrorMessage = TEXT("Could not find compilation data for public function @@.");
+									Settings.ASTSettings.Report(EMessageSeverity::Error, FunctionReference, FunctionCompilationErrorMessage);
+									bEncounteredGraphError = true;
+									bCanBeCompiled = false;
+									break;
+								}
+							}
+						}
+					}
+
+					if (bCanBeCompiled)
+					{
+						if (URigVMLibraryNode* LibraryNode = Cast<URigVMLibraryNode>(Nodes[i]))
+						{
+							Nodes.Append(LibraryNode->GetContainedNodes());	
+						}
+					}
+				}
+
+				if (bCanBeCompiled)
+				{
+					if (!CompileFunction(FunctionToCompile, InController, UserData))
+					{
+						static const FString FunctionCompilationErrorMessage = TEXT("Function @@ failed to compile.");
+						Settings.ASTSettings.Report(EMessageSeverity::Error, FunctionToCompile, FunctionCompilationErrorMessage);
+						bEncounteredGraphError = true;
+						break;
+					}
+				}
+			}
+			
+			if(bEncounteredGraphError)
+			{
+				return false;
+			}
+		}		
+	}
+
 	FRigVMCompilerWorkData WorkData;
 
 	WorkData.AST = InAST;
@@ -502,75 +609,32 @@ bool URigVMCompiler::Compile(TArray<URigVMGraph*> InGraphs, URigVMController* In
 
 	UE_LOG_RIGVMMEMORY(TEXT("RigVMCompiler: Begin '%s'..."), *InGraph->GetPathName());
 
-#if WITH_EDITOR
-	// If in editor, make sure we visit all the graphs to initialize local variables
-	// in case the user wants to edit default values
-	URigVMFunctionLibrary* FunctionLibrary = InGraphs[0]->GetDefaultFunctionLibrary();
-	if (FunctionLibrary)
+	// If we are compiling a function, we want the first registers to represent the interface pins (in the order of the pins)
+	// so they can be replaced when inlining the function
+	if (CurrentCompilationFunction)
 	{
-		for (URigVMLibraryNode* LibraryNode : FunctionLibrary->GetFunctions())
+		URigVMFunctionEntryNode* EntryNode = CurrentCompilationFunction->GetEntryNode();
+		URigVMFunctionReturnNode* ReturnNode = CurrentCompilationFunction->GetReturnNode();
+		for (URigVMPin* Pin : CurrentCompilationFunction->GetPins())
 		{
+			URigVMPin* InterfacePin = nullptr;
+			if (Pin->GetDirection() == ERigVMPinDirection::Input ||
+				Pin->GetDirection() == ERigVMPinDirection::IO)
 			{
-				FRigVMControllerGraphGuard Guard(InController, LibraryNode->GetContainedGraph(), false);
-				// make sure variables are up to date before validating other things.
-				// that is, make sure their cpp type and type object agree with each other
-				InController->EnsureLocalVariableValidity();
+				checkf(EntryNode, TEXT("Could not find entry node when compiling function %s in %s"), *CurrentCompilationFunction->GetPathName(), *GetPackage()->GetPathName());
+				InterfacePin = EntryNode->FindPin(Pin->GetName());
 			}
-
-			for (FRigVMGraphVariableDescription& Variable : LibraryNode->GetContainedGraph()->LocalVariables)
+			else
 			{
-				FString Path = FString::Printf(TEXT("LocalVariableDefault::%s|%s::Const"), *LibraryNode->GetContainedGraph()->GetGraphName(), *Variable.Name.ToString());
-				FRigVMOperand Operand = WorkData.AddProperty(ERigVMMemoryType::Literal, *Path, Variable.CPPType, Variable.CPPTypeObject, Variable.DefaultValue);
-				WorkData.PinPathToOperand->Add(Path, Operand);
-
-				for (const FRigVMExternalVariable& ExternalVariable : InExternalVariables)
-				{
-					if (ExternalVariable.Name == Variable.Name)
-					{
-						ReportWarningf(TEXT("Blueprint variable %s is being shadowed by a local variable in function %s"), *ExternalVariable.Name.ToString(), *LibraryNode->GetName());
-					}
-				}
+				checkf(ReturnNode, TEXT("Could not find return node when compiling function %s in %s"), *CurrentCompilationFunction->GetPathName(), *GetPackage()->GetPathName());
+				InterfacePin = ReturnNode->FindPin(Pin->GetName());
 			}
+			checkf(InterfacePin, TEXT("Could not find interface pin %s when compiling function %s in %s"), *Pin->GetName(), *CurrentCompilationFunction->GetPathName(), *GetPackage()->GetPathName());
+			FRigVMASTProxy PinProxy = FRigVMASTProxy::MakeFromUObject(InterfacePin);
+			FRigVMVarExprAST* TempVarExpr = WorkData.AST->MakeExpr<FRigVMVarExprAST>(FRigVMExprAST::EType::Var, PinProxy);
+			FindOrAddRegister(TempVarExpr, WorkData, false);
 		}
 	}
-#endif
-
-	// Look for all local variables to create the register with the default value in the literal memory
-	int32 IndexLocalVariable = 0;
-	for(URigVMGraph* VisitedGraph : VisitedGraphs)
-	{
-		for (const FRigVMGraphVariableDescription& LocalVariable : VisitedGraph->LocalVariables)
-		{
-			auto AddDefaultValueOperand = [&](URigVMPin* Pin)
-			{
-				FRigVMASTProxy PinProxy = FRigVMASTProxy::MakeFromUObject(Pin);
-				FRigVMVarExprAST* TempVarExpr = WorkData.AST->MakeExpr<FRigVMVarExprAST>(FRigVMExprAST::EType::Literal, PinProxy);
-				FRigVMOperand Operand = FindOrAddRegister(TempVarExpr, WorkData, false);
-
-				check(Operand.GetMemoryType() == ERigVMMemoryType::Literal);
-				TArray<FRigVMPropertyDescription>& LiteralProperties = WorkData.PropertyDescriptions.FindChecked(Operand.GetMemoryType());
-				LiteralProperties[Operand.GetRegisterIndex()].DefaultValue = LocalVariable.DefaultValue;
-			};
-			
-			// To create the default value in the literal memory, we need to find a pin in a variable node (or bounded to a local variable) that
-			// uses this local variable
-			for (URigVMNode* Node : VisitedGraph->GetNodes())
-			{
-				if (URigVMVariableNode* VariableNode = Cast<URigVMVariableNode>(Node))
-				{
-					if (URigVMPin* Pin = VariableNode->FindPin(URigVMVariableNode::VariableName))
-					{
-						if (Pin->GetDefaultValue() == LocalVariable.Name.ToString())
-						{
-							URigVMPin* ValuePin = VariableNode->FindPin(URigVMVariableNode::ValueName);
-							AddDefaultValueOperand(ValuePin);
-							break;
-						}
-					}
-				}
-			}
-		}
-	}	
 
 	if(Settings.EnablePinWatches)
 	{
@@ -652,13 +716,16 @@ bool URigVMCompiler::Compile(TArray<URigVMGraph*> InGraphs, URigVMController* In
 		TraverseExpression(RootExpr, WorkData);
 	}
 
-	if (WorkData.VM->GetByteCode().GetInstructions().Num() == 0)
+	if (!CurrentCompilationFunction)
 	{
-		WorkData.VM->GetByteCode().AddExitOp();
+		if (WorkData.VM->GetByteCode().GetInstructions().Num() == 0)
+		{
+			WorkData.VM->GetByteCode().AddExitOp();
+		}
+	
+		WorkData.VM->GetByteCode().AlignByteCode();
 	}
 
-	WorkData.VM->GetByteCode().AlignByteCode();
-
 	// setup debug registers after all other registers have been created
 	if(Settings.EnablePinWatches)
 	{
@@ -695,9 +762,90 @@ bool URigVMCompiler::Compile(TArray<URigVMGraph*> InGraphs, URigVMController* In
 		WorkData.VM->ExternalPropertyPathDescriptions = *Descriptions;
 	}
 
+	// Store function compile data
+	if (CurrentCompilationFunction)
+	{
+		FunctionCompilationData Data;
+		Data.ByteCode = WorkData.VM->ByteCodeStorage;
+		Data.FunctionNames = WorkData.VM->FunctionNamesStorage;
+		Data.PropertyDescriptions = WorkData.PropertyDescriptions;
+		Data.PropertyPathDescriptions = WorkData.PropertyPathDescriptions;
+		for (const TPair<FString, FRigVMOperand> Pair : (*WorkData.PinPathToOperand))
+		{
+			static const FString VariablePrefix = TEXT("Variable::");
+			if (Pair.Key.StartsWith(VariablePrefix))
+			{
+				ensure(Pair.Value.GetMemoryType() == ERigVMMemoryType::External);
+				FString VariableName = Pair.Key.RightChop(VariablePrefix.Len());
+				Data.ExternalRegisterIndexToVariable.Add(Pair.Value.GetRegisterIndex(), *VariableName);
+			}
+		}
+		CompiledFunctions.Add(CurrentCompilationFunction, Data);
+	}
+
 	return true;
 }
 
+bool URigVMCompiler::CompileFunction(URigVMLibraryNode* InLibraryNode, URigVMController* InController, const TArray<FRigVMUserDataArray>& InRigVMUserData)
+{
+	TGuardValue<URigVMLibraryNode*> CompilationGuard(CurrentCompilationFunction, InLibraryNode);
+	FRigVMControllerGraphGuard ControllerGraphGuard(InController, InLibraryNode->GetContainedGraph(), false);
+
+	TArray<FRigVMExternalVariable> ExternalVariables;
+	if (InController->GetExternalVariablesDelegate.IsBound())
+	{
+		ExternalVariables = InController->GetExternalVariablesDelegate.Execute(InLibraryNode->GetContainedGraph());
+	}
+	TMap<FString, FRigVMOperand> Operands;
+	URigVM* TempVM = NewObject<URigVM>(InLibraryNode->GetContainedGraph());
+	const bool bSuccess = Compile({InLibraryNode->GetContainedGraph()}, InController, TempVM, ExternalVariables, InRigVMUserData, &Operands);
+	TempVM->Rename(nullptr, GetTransientPackage(), REN_ForceNoResetLoaders | REN_DoNotDirty | REN_DontCreateRedirectors | REN_NonTransactional);
+	TempVM->MarkAsGarbage();
+
+	return bSuccess;
+}
+
+void URigVMCompiler::ClearFunctionCompilationData()
+{
+	TArray<TSoftObjectPtr<URigVMLibraryNode>> Functions;
+	CompiledFunctions.GetKeys(Functions);
+	for (TSoftObjectPtr<URigVMLibraryNode> Function : Functions)
+	{
+		RemoveFunctionCompilationData(Function);
+	}
+	CompiledFunctions.Reset();
+}
+
+void URigVMCompiler::RemoveFunctionCompilationData(TSoftObjectPtr<URigVMLibraryNode> LibraryNode)
+{
+	FunctionCompilationData* CompilationData = CompiledFunctions.Find(LibraryNode);
+	if (CompilationData)
+	{
+		CompiledFunctions.Remove(LibraryNode);
+	}
+
+	if (!LibraryNode.IsValid())
+	{
+		return;
+	}
+
+	// Clear the compilation data of any function referencing the removed function
+	URigVMFunctionLibrary* Library = LibraryNode->GetLibrary();
+	URigVMBuildData* BuildData = Library->GetBuildData();
+	if (const FRigVMFunctionReferenceArray* FunctionReferences = BuildData->FindFunctionReferences(LibraryNode.Get()))
+	{
+		for (int32 i=0; i<FunctionReferences->Num(); ++i)
+		{
+			const TSoftObjectPtr<URigVMFunctionReferenceNode> Reference = FunctionReferences->operator[](i);
+			if (Reference.IsValid())
+			{
+				URigVMLibraryNode* ReferencedNode = Reference->GetReferencedNode();
+				// todo: remove compilation data of this function
+			}
+		}
+	}
+}
+
 void URigVMCompiler::TraverseExpression(const FRigVMExprAST* InExpr, FRigVMCompilerWorkData& WorkData)
 {
 	if (WorkData.ExprToSkip.Contains(InExpr))
@@ -711,8 +859,6 @@ void URigVMCompiler::TraverseExpression(const FRigVMExprAST* InExpr, FRigVMCompi
 	}
 	WorkData.ExprComplete.Add(InExpr, true);
 
-	InitializeLocalVariables(InExpr, WorkData);	
-
 	switch (InExpr->GetType())
 	{
 		case FRigVMExprAST::EType::Block:
@@ -740,6 +886,12 @@ void URigVMCompiler::TraverseExpression(const FRigVMExprAST* InExpr, FRigVMCompi
 			TraverseCallExtern(CallExternExpr, WorkData);
 			break;
 		}
+		case FRigVMExprAST::EType::InlineFunction:
+		{
+			const FRigVMInlineFunctionExprAST* InlineExpr = InExpr->To<FRigVMInlineFunctionExprAST>();
+			TraverseInlineFunction(InlineExpr, WorkData);
+			break;
+		}
 		case FRigVMExprAST::EType::NoOp:
 		{
 			TraverseNoOp(InExpr->To<FRigVMNoOpExprAST>(), WorkData);
@@ -1046,6 +1198,186 @@ int32 URigVMCompiler::TraverseCallExtern(const FRigVMCallExternExprAST* InExpr,
 	return InstructionIndex;
 }
 
+int32 URigVMCompiler::TraverseInlineFunction(const FRigVMInlineFunctionExprAST* InExpr, FRigVMCompilerWorkData& WorkData)
+{
+	URigVMNode* Node = InExpr->GetNode();
+	URigVMFunctionReferenceNode* FunctionReferenceNode = Cast<URigVMFunctionReferenceNode>(Node);
+	if(!ValidateNode(FunctionReferenceNode, false))
+	{
+		return INDEX_NONE;
+	}
+	
+	int32 InstructionIndexStart = INDEX_NONE;
+	int32 InstructionIndexEnd = INDEX_NONE;
+
+	const URigVMLibraryNode* ReferencedLibrary = FunctionReferenceNode->GetReferencedNode();
+	if (!CompiledFunctions.Contains(ReferencedLibrary))
+	{
+		return INDEX_NONE;
+	}
+	const FunctionCompilationData& FunctionCompilationData = CompiledFunctions.FindChecked(ReferencedLibrary);
+	const FRigVMByteCode& FunctionByteCode = FunctionCompilationData.ByteCode;
+	
+	if (WorkData.bSetupMemory)
+	{
+		TraverseChildren(InExpr, WorkData);
+
+		// Add internal operands (not the ones represented by interface pins)
+		TArray<ERigVMMemoryType> MemoryTypes;
+		FunctionCompilationData.PropertyDescriptions.GetKeys(MemoryTypes);
+		for (const ERigVMMemoryType& MemoryType : MemoryTypes)
+		{
+			const TArray<FRigVMPropertyDescription>& Properties = FunctionCompilationData.PropertyDescriptions[MemoryType];
+			int32 StartIndex = MemoryType == ERigVMMemoryType::Work ? ReferencedLibrary->GetPins().Num() : 0;
+			for (int32 PropertyIndex = StartIndex; PropertyIndex < Properties.Num(); ++PropertyIndex)
+			{
+				const FRigVMPropertyDescription& Description = Properties[PropertyIndex];
+				FString NewName = Description.Name.ToString();
+				static const FString FunctionLibraryPrefix = TEXT("FunctionLibrary");
+				NewName = FString::Printf(TEXT("%s%s"), *FunctionReferenceNode->GetNodePath(), *NewName.RightChop(FunctionLibraryPrefix.Len()));
+				FRigVMOperand Operand = WorkData.AddProperty(MemoryType, *NewName, Description.CPPType, Description.CPPTypeObject, Description.DefaultValue);
+				FRigVMCompilerWorkData::FFunctionRegisterData Data = {FunctionReferenceNode, MemoryType, PropertyIndex};
+				WorkData.FunctionRegisterToOperand.Add(Data, Operand);
+
+				// @todo Try to reuse literal operands				
+			}
+		}		
+	}
+	else
+	{
+		TArray<FRigVMOperand> Operands;		
+		for(const URigVMPin* Pin : FunctionReferenceNode->GetPins())
+		{
+			const FRigVMExprAST* ChildExpr = InExpr->FindExprWithPinName(Pin->GetFName());
+			checkf(ChildExpr, TEXT("Found unexpected opaque argument for %s while inlining function %s in package %s"), *InExpr->Name.ToString(), *FunctionReferenceNode->GetPathName(), *GetPackage()->GetPathName());			
+			if (ChildExpr->GetType() == FRigVMExprAST::EType::CachedValue)
+			{
+				Operands.Add(WorkData.ExprToOperand.FindChecked(GetSourceVarExpr(ChildExpr->To<FRigVMCachedValueExprAST>()->GetVarExpr())));
+			}
+			else if (ChildExpr->IsA(FRigVMExprAST::EType::Var))
+			{
+				Operands.Add(WorkData.ExprToOperand.FindChecked(GetSourceVarExpr(ChildExpr->To<FRigVMVarExprAST>())));
+			}
+			else
+			{
+				break;
+			}
+		}
+
+		TraverseChildren(InExpr, WorkData);
+
+		// Inline the bytecode from the function
+		FRigVMByteCode& ByteCode = WorkData.VM->GetByteCode();
+		InstructionIndexStart = ByteCode.GetNumInstructions();
+		FRigVMInstructionArray OldInstructions = ByteCode.GetInstructions();
+		ByteCode.InlineFunction(&FunctionByteCode, Operands);
+		InstructionIndexEnd = ByteCode.GetNumInstructions() - 1;
+		FRigVMCallstack FuncRefCallstack = InExpr->GetProxy().GetCallstack();
+
+		// For each instruction, substitute the operand for the one used in the current bytecode
+		const FRigVMInstructionArray FunctionInstructions = FunctionByteCode.GetInstructions();
+		FRigVMInstructionArray Instructions = ByteCode.GetInstructions();
+		for (int32 i=InstructionIndexStart; i<=InstructionIndexEnd; ++i)
+		{
+			const FRigVMInstruction& Instruction = Instructions[i];
+			const FRigVMOperandArray OperandArray = ByteCode.GetOperandsForOp(Instruction);
+			uint64 OperandsIndex = ByteCode.GetFirstOperandByteIndex(Instruction);
+			for (int32 j=0; j<OperandArray.Num(); ++j)
+			{
+				FRigVMOperand* Operand = (FRigVMOperand*)(ByteCode.ByteCode.GetData() + OperandsIndex + j*sizeof(FRigVMOperand));
+				ERigVMMemoryType OriginalMemoryType = Operand->GetMemoryType();
+
+				// Remap the variable: find the operand index of the outer variable
+				if (Operand->GetMemoryType() == ERigVMMemoryType::External)
+				{
+					const FName& InnerVariableName = FunctionCompilationData.ExternalRegisterIndexToVariable[Operand->GetRegisterIndex()];
+					const FName& OuterVariableName = FunctionReferenceNode->GetVariableMap()[InnerVariableName];
+					const FRigVMOperand& OuterOperand = WorkData.PinPathToOperand->FindChecked(FString::Printf(TEXT("Variable::%s"), *OuterVariableName.ToString()));
+					Operand->RegisterIndex = OuterOperand.RegisterIndex;
+				}
+				// Operand is an interface pin: replace the index and memory type
+				else if (Operand->GetMemoryType() == ERigVMMemoryType::Work &&
+					Operand->GetRegisterIndex() < ReferencedLibrary->GetPins().Num())
+				{
+					Operand->MemoryType = Operands[Operand->GetRegisterIndex()].MemoryType;
+					Operand->RegisterIndex = Operands[Operand->GetRegisterIndex()].RegisterIndex;
+				}
+				else
+				{
+					// Operand is internal
+					// Replace with added Operand
+					FRigVMCompilerWorkData::FFunctionRegisterData Data = {FunctionReferenceNode, Operand->GetMemoryType(), Operand->GetRegisterIndex()};
+					FRigVMOperand NewOperand = WorkData.FunctionRegisterToOperand.FindChecked(Data);
+					Operand->MemoryType = NewOperand.MemoryType;
+					Operand->RegisterIndex = NewOperand.RegisterIndex;
+				}
+
+				// For all operands, check to see if we need to add a property path
+				if (Operand->GetRegisterOffset() != INDEX_NONE)
+				{
+					const FRigVMPropertyPathDescription& Description = FunctionCompilationData.PropertyPathDescriptions[OriginalMemoryType][Operand->GetRegisterOffset()];
+					Operand->RegisterOffset = WorkData.FindOrAddPropertyPath(*Operand, Description.HeadCPPType, Description.SegmentPath);
+				}
+			}
+
+			if (Instruction.OpCode >= ERigVMOpCode::Execute_0_Operands && Instruction.OpCode <= ERigVMOpCode::Execute_64_Operands)
+			{
+				FRigVMExecuteOp& Op = ByteCode.GetOpAt<FRigVMExecuteOp>(Instruction);
+				const int32 FunctionIndex = WorkData.VM->AddRigVMFunction(FunctionCompilationData.FunctionNames[Op.FunctionIndex].ToString());
+				Op.FunctionIndex = FunctionIndex;
+			}
+
+			if (Settings.SetupNodeInstructionIndex)
+			{
+				if (const TArray<UObject*>* Callstack = FunctionByteCode.GetCallstackForInstruction(i-InstructionIndexStart))
+				{
+					if (Callstack->Num() > 1)
+					{
+						FRigVMCallstack InstructionCallstack = FuncRefCallstack;
+						InstructionCallstack.Stack.Append(&(*Callstack)[1], Callstack->Num()-1);
+						WorkData.VM->GetByteCode().SetSubject(i, InstructionCallstack.GetCallPath(), InstructionCallstack.GetStack());
+					}
+				}
+			}
+		}
+		
+
+#if WITH_EDITORONLY_DATA
+		TArray<FRigVMOperand> InputsOperands, OutputOperands;
+
+		for(const URigVMPin* InputPin : Node->GetPins())
+		{
+			if(InputPin->IsExecuteContext())
+			{
+				continue;
+			}
+
+			const FRigVMOperand& Operand = Operands[InputPin->GetPinIndex()];
+
+			if(InputPin->GetDirection() == ERigVMPinDirection::Output || InputPin->GetDirection() == ERigVMPinDirection::IO)
+			{
+				OutputOperands.Add(Operand);
+			}
+
+			if(InputPin->GetDirection() != ERigVMPinDirection::Input && InputPin->GetDirection() != ERigVMPinDirection::IO)
+			{
+				continue;
+			}
+
+			InputsOperands.Add(Operand);
+		}
+
+		WorkData.VM->GetByteCode().SetOperandsForInstruction(
+			InstructionIndexStart,
+			FRigVMOperandArray(InputsOperands.GetData(), InputsOperands.Num()),
+			FRigVMOperandArray(OutputOperands.GetData(), OutputOperands.Num()));
+
+#endif
+	}
+
+	return InstructionIndexEnd;
+}
+
 void URigVMCompiler::TraverseForLoop(const FRigVMCallExternExprAST* InExpr, FRigVMCompilerWorkData& WorkData)
 {
 	if (WorkData.bSetupMemory)
@@ -1145,6 +1477,39 @@ void URigVMCompiler::TraverseForLoop(const FRigVMCallExternExprAST* InExpr, FRig
 
 void URigVMCompiler::TraverseNoOp(const FRigVMNoOpExprAST* InExpr, FRigVMCompilerWorkData& WorkData)
 {
+	if (CurrentCompilationFunction && InExpr->NumParents() == 0)
+	{
+		// Initialize local variables
+		if (WorkData.bSetupMemory)
+		{
+			TArray<FRigVMGraphVariableDescription> LocalVariables = CurrentCompilationFunction->GetContainedGraph()->GetLocalVariables();
+			for (const FRigVMGraphVariableDescription& Variable : LocalVariables)
+			{
+				FString Path = FString::Printf(TEXT("LocalVariableDefault::%s|%s::Const"), *CurrentCompilationFunction->GetContainedGraph()->GetGraphName(), *Variable.Name.ToString());
+				FRigVMOperand Operand = WorkData.AddProperty(ERigVMMemoryType::Literal, *Path, Variable.CPPType, Variable.CPPTypeObject, Variable.DefaultValue);
+				WorkData.PinPathToOperand->Add(Path, Operand);
+			}
+		}
+		else
+		{
+			TArray<FRigVMGraphVariableDescription> LocalVariables = CurrentCompilationFunction->GetContainedGraph()->GetLocalVariables();
+			for (const FRigVMGraphVariableDescription& Variable : LocalVariables)
+			{
+				FString TargetPath = FString::Printf(TEXT("LocalVariable::%s|%s"), *CurrentCompilationFunction->GetContainedGraph()->GetGraphName(), *Variable.Name.ToString());
+				FString SourcePath = FString::Printf(TEXT("LocalVariableDefault::%s|%s::Const"), *CurrentCompilationFunction->GetContainedGraph()->GetGraphName(), *Variable.Name.ToString());
+				FRigVMOperand* TargetPtr = WorkData.PinPathToOperand->Find(TargetPath);
+				FRigVMOperand* SourcePtr = WorkData.PinPathToOperand->Find(SourcePath);
+				if (SourcePtr && TargetPtr)
+				{
+					const FRigVMOperand& Source = *SourcePtr;
+					const FRigVMOperand& Target = *TargetPtr;
+	
+					WorkData.VM->GetByteCode().AddCopyOp(WorkData.VM->GetCopyOpForOperands(Source, Target));						
+				}
+			}
+		}
+	}
+
 	TraverseChildren(InExpr, WorkData);
 }
 
@@ -2110,115 +2475,7 @@ void URigVMCompiler::AddCopyOperator(
 	AddCopyOperator(CopyOpInfo.Op, CopyOpInfo.AssignExpr, CopyOpInfo.SourceExpr, CopyOpInfo.TargetExpr, WorkData, bDelayCopyOperations);
 }
 
-void URigVMCompiler::InitializeLocalVariables(const FRigVMExprAST* InExpr, FRigVMCompilerWorkData& WorkData)
-{
-	// Initialize local variables if we are entering a new graph
-	if (!WorkData.bSetupMemory)
-	{
-		FRigVMByteCode& ByteCode = WorkData.VM->GetByteCode();
-		const FRigVMASTProxy* Proxy = nullptr;
-		switch (InExpr->GetType())
-		{
-			case FRigVMExprAST::EType::CallExtern:
-			{
-				Proxy = &InExpr->To<FRigVMCallExternExprAST>()->GetProxy();
-				break;
-			}
-			case FRigVMExprAST::EType::NoOp:
-			{
-				Proxy = &InExpr->To<FRigVMNoOpExprAST>()->GetProxy();
-				break;
-			}
-			case FRigVMExprAST::EType::Var:
-			{
-				Proxy = &InExpr->To<FRigVMVarExprAST>()->GetProxy();
-				break;
-			}
-			case FRigVMExprAST::EType::Literal:
-			{
-				Proxy = &InExpr->To<FRigVMLiteralExprAST>()->GetProxy();
-				break;
-			}
-			case FRigVMExprAST::EType::ExternalVar:
-			{
-				Proxy = &InExpr->To<FRigVMExternalVarExprAST>()->GetProxy();
-				break;
-			}
-			case FRigVMExprAST::EType::Branch:
-			{
-				Proxy = &InExpr->To<FRigVMBranchExprAST>()->GetProxy();
-				break;
-			}
-			case FRigVMExprAST::EType::If:
-			{
-				Proxy = &InExpr->To<FRigVMIfExprAST>()->GetProxy();
-				break;
-			}
-			case FRigVMExprAST::EType::Select:
-			{
-				Proxy = &InExpr->To<FRigVMSelectExprAST>()->GetProxy();
-				break;
-			}
-			case FRigVMExprAST::EType::Array:
-			{
-				Proxy = &InExpr->To<FRigVMArrayExprAST>()->GetProxy();
-				break;
-			}
-		}
-
-		if(Proxy != nullptr)
-		{
-			const FRigVMCallstack& Callstack = Proxy->GetCallstack();
-			ensure(Callstack.Num() > 0);
-
-			// Find all function references in the callstack and initialize their local variables if necessary
-			for (int32 SubjectIndex=0; SubjectIndex<Callstack.Num(); ++SubjectIndex)
-			{
-				if (const URigVMLibraryNode* Node = Cast<const URigVMLibraryNode>(Callstack[SubjectIndex]))
-				{
-					// Check if this is the first time we are accessing this function reference
-					bool bFound = false;
-					for (int32 i=ByteCode.GetNumInstructions()-1; i>0; --i)
-					{
-						const TArray<UObject*>* PreviousCallstack = ByteCode.GetCallstackForInstruction(i);
-						if (PreviousCallstack && PreviousCallstack->Contains(Node))
-						{
-							bFound = true;
-							break;
-						}
-					}
-
-					// If it is the first time we access this function reference, initialize all local variables
-					if (!bFound)
-					{
-						for (FRigVMGraphVariableDescription Variable : Node->GetContainedGraph()->LocalVariables)
-						{
-							const FRigVMCallstack LocalCallstack = Callstack.GetCallStackUpTo(SubjectIndex);
-							FString TargetPath = FString::Printf(TEXT("LocalVariable::%s|%s"), *LocalCallstack.GetCallPath(), *Variable.Name.ToString());
-							FString SourcePath = FString::Printf(TEXT("LocalVariableDefault::%s|%s::Const"), *Node->GetContainedGraph()->GetGraphName(), *Variable.Name.ToString());
-							FRigVMOperand* TargetPtr = WorkData.PinPathToOperand->Find(TargetPath);
-							FRigVMOperand* SourcePtr = WorkData.PinPathToOperand->Find(SourcePath);
-							if (SourcePtr && TargetPtr) 
-							{
-								const FRigVMOperand& Source = *SourcePtr;
-								const FRigVMOperand& Target = *TargetPtr;
-								
-								ByteCode.AddCopyOp(WorkData.VM->GetCopyOpForOperands(Source, Target));
-								if(Settings.SetupNodeInstructionIndex)
-								{
-									const int32 InstructionIndex = ByteCode.GetNumInstructions() - 1;
-									WorkData.VM->GetByteCode().SetSubject(InstructionIndex, Callstack.GetCallPath(), Callstack.GetStack());
-								}							
-							}					
-						}
-					}
-				}
-			}
-		}
-	}
-}
-
-FString URigVMCompiler::GetPinHashImpl(const URigVMPin* InPin, const FRigVMVarExprAST* InVarExpr, bool bIsDebugValue, const FRigVMASTProxy& InPinProxy)
+FString URigVMCompiler::GetPinHashImpl(const URigVMPin* InPin, const FRigVMVarExprAST* InVarExpr, bool bIsDebugValue, URigVMLibraryNode* FunctionCompiling, const FRigVMASTProxy& InPinProxy)
 {
 	FString Prefix = bIsDebugValue ? TEXT("DebugWatch:") : TEXT("");
 	FString Suffix;
@@ -2233,6 +2490,7 @@ FString URigVMCompiler::GetPinHashImpl(const URigVMPin* InPin, const FRigVMVarEx
 	bool bIsExecutePin = false;
 	bool bIsLiteral = false;
 	bool bIsVariable = false;
+	bool bIsFunctionInterfacePin = false;
 
 	if (InVarExpr != nullptr && !bIsDebugValue)
 	{
@@ -2248,7 +2506,7 @@ FString URigVMCompiler::GetPinHashImpl(const URigVMPin* InPin, const FRigVMVarEx
 		{
 			if(SourceVarExpr != InVarExpr)
 			{
-				return GetPinHash(SourceVarExpr->GetPin(), SourceVarExpr, bIsDebugValue);
+				return GetPinHash(SourceVarExpr->GetPin(), SourceVarExpr, bIsDebugValue, FunctionCompiling);
 			}
 		}
 
@@ -2256,10 +2514,13 @@ FString URigVMCompiler::GetPinHashImpl(const URigVMPin* InPin, const FRigVMVarEx
 		bIsLiteral = InVarExpr->GetType() == FRigVMExprAST::EType::Literal;
 
 		bIsVariable = Cast<URigVMVariableNode>(Node) != nullptr || InVarExpr->IsA(FRigVMExprAST::ExternalVar);
+		bIsFunctionInterfacePin = (Cast<URigVMFunctionEntryNode>(Node) || Cast<URigVMFunctionReturnNode>(Node)) &&
+			Node->GetTypedOuter<URigVMLibraryNode>() == FunctionCompiling;
 
 		// determine if this is an initialization for an IO pin
 		if (!bIsLiteral &&
 			!bIsVariable &&
+			!bIsFunctionInterfacePin &&
 			!bIsExecutePin && (InPin->GetDirection() == ERigVMPinDirection::IO ||
 			(InPin->GetDirection() == ERigVMPinDirection::Input && InPin->GetSourceLinks().Num() == 0)))
 		{
@@ -2289,8 +2550,8 @@ FString URigVMCompiler::GetPinHashImpl(const URigVMPin* InPin, const FRigVMVarEx
 					}
 					else if (InVarExpr)
 					{
-						const FString GraphPath = InVarExpr->GetProxy().GetCallstack().GetCallPath(true);
-						return FString::Printf(TEXT("%sLocalVariable::%s%s"), *Prefix, *GraphPath, *Suffix);					
+						const FString PinPath = InVarExpr->GetProxy().GetCallstack().GetCallPath(true);
+						return FString::Printf(TEXT("%sLocalVariable::%s%s"), *Prefix, *PinPath, *Suffix);					
 					}
 					else
 					{
@@ -2313,8 +2574,7 @@ FString URigVMCompiler::GetPinHashImpl(const URigVMPin* InPin, const FRigVMVarEx
 						}
 
 						// Local variables for root / non-root graphs are in the format "LocalVariable::PathToGraph|VariableName"
-						const FString GraphPath = ParentProxy.GetCallstack().GetCallPath(true);
-						return FString::Printf(TEXT("%sLocalVariable::%s|%s%s"), *Prefix, *GraphPath, *VariableName.ToString(), *Suffix);
+						return FString::Printf(TEXT("%sLocalVariable::%s|%s%s"), *Prefix, *Node->GetGraph()->GetGraphName(), *VariableName.ToString(), *Suffix);
 					}
 				}
 			}
@@ -2404,10 +2664,10 @@ FString URigVMCompiler::GetPinHashImpl(const URigVMPin* InPin, const FRigVMVarEx
 	return FString::Printf(TEXT("%s%s%s"), *Prefix, *PinPath, *Suffix);
 }
 
-FString URigVMCompiler::GetPinHash(const URigVMPin* InPin, const FRigVMVarExprAST* InVarExpr, bool bIsDebugValue, const FRigVMASTProxy& InPinProxy)
+FString URigVMCompiler::GetPinHash(const URigVMPin* InPin, const FRigVMVarExprAST* InVarExpr, bool bIsDebugValue, URigVMLibraryNode* FunctionCompiling, const FRigVMASTProxy& InPinProxy)
 {
-	const FString Hash = GetPinHashImpl(InPin, InVarExpr, bIsDebugValue, InPinProxy);
-	if(!bIsDebugValue)
+	const FString Hash = GetPinHashImpl(InPin, InVarExpr, bIsDebugValue, FunctionCompiling, InPinProxy);
+	if(!bIsDebugValue && FunctionCompiling == nullptr)
 	{
 		ensureMsgf(!Hash.Contains(TEXT("FunctionLibrary::")), TEXT("A library path should never be part of a pin hash %s."), *Hash);
 	}
@@ -2518,10 +2778,10 @@ void URigVMCompiler::MarkDebugWatch(bool bRequired, URigVMPin* InPin, URigVM* Ou
 			}
 		}
 
-		FString PinHash = GetPinHash(Pin, VarExpression, false);
+		FString PinHash = GetPinHash(Pin, VarExpression, false, CurrentCompilationFunction);
 		if(!OutOperands->Contains(PinHash))
 		{
-			PinHash = GetPinHash(SourcePin, VarExpression, false);
+			PinHash = GetPinHash(SourcePin, VarExpression, false, CurrentCompilationFunction);
 		}
 
 		if(const FRigVMOperand* Operand = OutOperands->Find(PinHash))
@@ -2530,7 +2790,7 @@ void URigVMCompiler::MarkDebugWatch(bool bRequired, URigVMPin* InPin, URigVM* Ou
 			FRigVMVarExprAST TempVarExpr(FRigVMExprAST::EType::Var, PinProxy);
 			TempVarExpr.ParserPtr = InRuntimeAST.Get();
 
-			const FString DebugPinHash = GetPinHash(Pin, &TempVarExpr, true);
+			const FString DebugPinHash = GetPinHash(Pin, &TempVarExpr, true, CurrentCompilationFunction);
 			const FRigVMOperand* DebugOperand = OutOperands->Find(DebugPinHash);
 			if(DebugOperand)
 			{
@@ -2698,7 +2958,7 @@ FRigVMOperand URigVMCompiler::FindOrAddRegister(const FRigVMVarExprAST* InVarExp
 	
 	FString CPPType = Pin->GetCPPType();
 	FString BaseCPPType = Pin->IsArray() ? Pin->GetArrayElementCppType() : CPPType;
-	FString Hash = GetPinHash(Pin, InVarExpr, bIsDebugValue);
+	FString Hash = GetPinHash(Pin, InVarExpr, bIsDebugValue, CurrentCompilationFunction);
 	FRigVMOperand Operand;
 	FString RegisterKey = Hash;
 
@@ -2744,7 +3004,7 @@ FRigVMOperand URigVMCompiler::FindOrAddRegister(const FRigVMVarExprAST* InVarExp
 			{
 				if (const URigVMPin* VirtualPin = Cast<URigVMPin>(Proxy.GetSubject()))
 				{
-					const FString VirtualPinHash = GetPinHash(VirtualPin, InVarExpr, bIsDebugValue, Proxy);
+					const FString VirtualPinHash = GetPinHash(VirtualPin, InVarExpr, bIsDebugValue, CurrentCompilationFunction, Proxy);
 					HashesWithSharedOperand.Add(VirtualPinHash);
 					if (Pin != VirtualPin)
 					{
diff --git a/Engine/Source/Developer/RigVMDeveloper/Public/RigVMCompiler/RigVMAST.h b/Engine/Source/Developer/RigVMDeveloper/Public/RigVMCompiler/RigVMAST.h
--- a/Engine/Source/Developer/RigVMDeveloper/Public/RigVMCompiler/RigVMAST.h
+++ b/Engine/Source/Developer/RigVMDeveloper/Public/RigVMCompiler/RigVMAST.h
@@ -16,6 +16,7 @@ class FRigVMBlockExprAST;
 class FRigVMEntryExprAST;
 class FRigVMInvokeEntryExprAST;
 class FRigVMCallExternExprAST;
+class FRigVMInlineFunctionExprAST;
 class FRigVMNoOpExprAST;
 class FRigVMVarExprAST;
 class FRigVMLiteralExprAST;
@@ -34,6 +35,7 @@ class URigVMLink;
 class URigVMNode;
 class URigVMGraph;
 class URigVMController;
+class URigVMLibraryNode;
 
 /*
  * A structure to describe a link between two proxies
@@ -82,6 +84,7 @@ public:
 		Entry,
 		InvokeEntry,
 		CallExtern,
+		InlineFunction,
 		NoOp,
 		Var,
 		Literal,
@@ -338,6 +341,17 @@ FORCEINLINE const FRigVMCallExternExprAST* FRigVMExprAST::To() const
 	return (const FRigVMCallExternExprAST*)this;
 }
 
+// specialized cast for type checking
+// for a InlineFunction / FRigVMInlineFunctionExprAST expression
+// will raise if types are not compatible
+// @return this expression cast to FRigVMInlineFunctionExprAST
+template<>
+FORCEINLINE const FRigVMInlineFunctionExprAST* FRigVMExprAST::To() const
+{
+	ensure(IsA(EType::InlineFunction));
+	return (const FRigVMInlineFunctionExprAST*)this;
+}
+
 // specialized cast for type checking
 // for a NoOp / FRigVMNoOpExprAST expression
 // will raise if types are not compatible
@@ -681,6 +695,37 @@ private:
 	friend class FRigVMParserAST;
 };
 
+/*
+ *
+ */
+class RIGVMDEVELOPER_API FRigVMInlineFunctionExprAST: public FRigVMNodeExprAST
+{
+public:
+
+	// virtual destructor
+	virtual ~FRigVMInlineFunctionExprAST() {}
+
+	// disable copy constructor
+	FRigVMInlineFunctionExprAST(const FRigVMInlineFunctionExprAST&) = delete;
+
+	// overload of the type checking mechanism
+	virtual bool IsA(EType InType) const override
+	{
+		return InType == EType::InlineFunction;
+	};
+
+protected:
+
+	// default constructor (protected so that only parser can access it)
+	FRigVMInlineFunctionExprAST(const FRigVMASTProxy& InNodeProxy)
+		: FRigVMNodeExprAST(EType::InlineFunction, InNodeProxy)
+	{}
+
+private:
+
+	friend class FRigVMParserAST;
+};
+
 /*
  * An abstract syntax tree no-op expression represents an expression which is
  * relevant for the structure of the tree (for grouping for example) but which 
@@ -1532,6 +1577,8 @@ private:
 	FRigVMParserASTSettings Settings;
 	mutable TMap<TTuple<const FRigVMExprAST*,const FRigVMExprAST*>, int32> MinIndexOfChildWithinParent;
 
+	URigVMLibraryNode* LibraryNodeBeingCompiled;
+
 	friend class FRigVMExprAST;
 	friend class FRigVMAssignExprAST;
 	friend class URigVMCompiler;
diff --git a/Engine/Source/Developer/RigVMDeveloper/Public/RigVMCompiler/RigVMASTProxy.h b/Engine/Source/Developer/RigVMDeveloper/Public/RigVMCompiler/RigVMASTProxy.h
--- a/Engine/Source/Developer/RigVMDeveloper/Public/RigVMCompiler/RigVMASTProxy.h
+++ b/Engine/Source/Developer/RigVMDeveloper/Public/RigVMCompiler/RigVMASTProxy.h
@@ -119,6 +119,7 @@ private:
 	friend class FRigVMParserAST;
 	friend class FRigVMCallstackGuard;
 	friend class FRigVMASTProxy;
+	friend class URigVMCompiler;
 };
 
 /*
diff --git a/Engine/Source/Developer/RigVMDeveloper/Public/RigVMCompiler/RigVMCompiler.h b/Engine/Source/Developer/RigVMDeveloper/Public/RigVMCompiler/RigVMCompiler.h
--- a/Engine/Source/Developer/RigVMDeveloper/Public/RigVMCompiler/RigVMCompiler.h
+++ b/Engine/Source/Developer/RigVMDeveloper/Public/RigVMCompiler/RigVMCompiler.h
@@ -10,6 +10,7 @@
 #include "RigVMModel/Nodes/RigVMParameterNode.h"
 #include "RigVMModel/Nodes/RigVMCommentNode.h"
 #include "RigVMModel/Nodes/RigVMRerouteNode.h"
+#include "RigVMModel/Nodes/RigVMFunctionReferenceNode.h"
 #include "RigVMCore/RigVM.h"
 #include "RigVMCore/RigVMStruct.h"
 #include "RigVMCompiler/RigVMAST.h"
@@ -100,6 +101,25 @@ public:
 	const FRigVMASTProxySourceMap* ProxySources;
 	FRigVMASTProxyTargetsMap ProxyTargets;
 
+	struct FFunctionRegisterData
+	{
+		TSoftObjectPtr<URigVMFunctionReferenceNode> ReferenceNode;
+		ERigVMMemoryType MemoryType;
+		int32 RegisterIndex;
+
+		friend inline uint32 GetTypeHash(const FFunctionRegisterData& Data)
+		{
+			uint32 Result = HashCombine(::GetTypeHash(Data.ReferenceNode.IsValid() ? Data.ReferenceNode.Get() : 0), ::GetTypeHash(Data.MemoryType));
+			return HashCombine(Result, ::GetTypeHash(Data.RegisterIndex));
+		}
+
+		bool operator==(const FFunctionRegisterData& Other) const
+		{
+			return ReferenceNode == Other.ReferenceNode && MemoryType == Other.MemoryType && RegisterIndex == Other.RegisterIndex;
+		}
+	};
+	TMap<FFunctionRegisterData, FRigVMOperand> FunctionRegisterToOperand; 
+
 	TArray<URigVMPin*> WatchedPins;
 	
 	TMap<ERigVMMemoryType, TArray<FRigVMPropertyPathDescription>> PropertyPathDescriptions;
@@ -146,8 +166,18 @@ public:
 
 	bool Compile(TArray<URigVMGraph*> InGraphs, URigVMController* InController, URigVM* OutVM, const TArray<FRigVMExternalVariable>& InExternalVariables, const TArray<FRigVMUserDataArray>& InRigVMUserData, TMap<FString, FRigVMOperand>* OutOperands, TSharedPtr<FRigVMParserAST> InAST = TSharedPtr<FRigVMParserAST>());
 
+	struct FunctionCompilationData
+	{
+		FRigVMByteCode ByteCode;
+		TArray<FName> FunctionNames;
+		TMap<ERigVMMemoryType, TArray<FRigVMPropertyDescription>> PropertyDescriptions;
+		TMap<ERigVMMemoryType, TArray<FRigVMPropertyPathDescription>> PropertyPathDescriptions;
+		TMap<int32, FName> ExternalRegisterIndexToVariable;
+	};
+	TMap<TSoftObjectPtr<URigVMLibraryNode>, FunctionCompilationData> CompiledFunctions;
+
 	static UScriptStruct* GetScriptStructForCPPType(const FString& InCPPType);
-	static FString GetPinHash(const URigVMPin* InPin, const FRigVMVarExprAST* InVarExpr, bool bIsDebugValue = false, const FRigVMASTProxy& InPinProxy = FRigVMASTProxy());
+	static FString GetPinHash(const URigVMPin* InPin, const FRigVMVarExprAST* InVarExpr, bool bIsDebugValue = false, URigVMLibraryNode* FunctionCompiling = nullptr, const FRigVMASTProxy& InPinProxy = FRigVMASTProxy());
 
 	// follows assignment expressions to find the source ref counted containers
 	// since ref counted containers are not copied for assignments.
@@ -158,16 +188,22 @@ public:
 
 private:
 
+	bool CompileFunction(URigVMLibraryNode* InLibraryNode, URigVMController* InController, const TArray<FRigVMUserDataArray>& InRigVMUserData);
+	void ClearFunctionCompilationData();
+	void RemoveFunctionCompilationData(TSoftObjectPtr<URigVMLibraryNode> LibraryNode);
+	URigVMLibraryNode* CurrentCompilationFunction;
+
 	TArray<URigVMPin*> GetLinkedPins(URigVMPin* InPin, bool bInputs = true, bool bOutputs = true, bool bRecursive = true);
 	uint16 GetElementSizeFromCPPType(const FString& InCPPType, UScriptStruct* InScriptStruct);
 
-	static FString GetPinHashImpl(const URigVMPin* InPin, const FRigVMVarExprAST* InVarExpr, bool bIsDebugValue = false, const FRigVMASTProxy& InPinProxy = FRigVMASTProxy());
+	static FString GetPinHashImpl(const URigVMPin* InPin, const FRigVMVarExprAST* InVarExpr, bool bIsDebugValue = false, URigVMLibraryNode* FunctionCompiling = nullptr, const FRigVMASTProxy& InPinProxy = FRigVMASTProxy());
 
 	void TraverseExpression(const FRigVMExprAST* InExpr, FRigVMCompilerWorkData& WorkData);
 	void TraverseChildren(const FRigVMExprAST* InExpr, FRigVMCompilerWorkData& WorkData);
 	void TraverseBlock(const FRigVMBlockExprAST* InExpr, FRigVMCompilerWorkData& WorkData);
 	void TraverseEntry(const FRigVMEntryExprAST* InExpr, FRigVMCompilerWorkData& WorkData);
 	int32 TraverseCallExtern(const FRigVMCallExternExprAST* InExpr, FRigVMCompilerWorkData& WorkData);
+	int32 TraverseInlineFunction(const FRigVMInlineFunctionExprAST* InExpr, FRigVMCompilerWorkData& WorkData);
 	void TraverseForLoop(const FRigVMCallExternExprAST* InExpr, FRigVMCompilerWorkData& WorkData);
 	void TraverseNoOp(const FRigVMNoOpExprAST* InExpr, FRigVMCompilerWorkData& WorkData);
 	void TraverseVar(const FRigVMVarExprAST* InExpr, FRigVMCompilerWorkData& WorkData);
@@ -196,8 +232,6 @@ private:
 		FRigVMCompilerWorkData& WorkData,
 		bool bDelayCopyOperations = true);
 
-	void InitializeLocalVariables(const FRigVMExprAST* InExpr, FRigVMCompilerWorkData& WorkData);
-
 	FRigVMOperand FindOrAddRegister(const FRigVMVarExprAST* InVarExpr, FRigVMCompilerWorkData& WorkData, bool bIsDebugValue = false);
 	const FRigVMCompilerWorkData::FRigVMASTProxyArray& FindProxiesWithSharedOperand(const FRigVMVarExprAST* InVarExpr, FRigVMCompilerWorkData& WorkData);
 
diff --git a/Engine/Source/Runtime/RigVM/Private/RigVMCore/RigVM.cpp b/Engine/Source/Runtime/RigVM/Private/RigVMCore/RigVM.cpp
--- a/Engine/Source/Runtime/RigVM/Private/RigVMCore/RigVM.cpp
+++ b/Engine/Source/Runtime/RigVM/Private/RigVMCore/RigVM.cpp
@@ -3209,6 +3209,7 @@ TArray<FString> URigVM::DumpByteCodeAsTextArray(const TArray<int32>& InInstructi
 
 FString URigVM::DumpByteCodeAsText(const TArray<int32>& InInstructionOrder, bool bIncludeLineNumbers)
 {
+	RefreshExternalPropertyPaths();
 	return FString::Join(DumpByteCodeAsTextArray(InInstructionOrder, bIncludeLineNumbers), TEXT("\n"));
 }
 
diff --git a/Engine/Source/Runtime/RigVM/Private/RigVMCore/RigVMByteCode.cpp b/Engine/Source/Runtime/RigVM/Private/RigVMCore/RigVMByteCode.cpp
--- a/Engine/Source/Runtime/RigVM/Private/RigVMCore/RigVMByteCode.cpp
+++ b/Engine/Source/Runtime/RigVM/Private/RigVMCore/RigVMByteCode.cpp
@@ -1228,6 +1228,16 @@ uint64 FRigVMByteCode::AddExecuteOp(uint16 InFunctionIndex, const FRigVMOperandA
 	return OpByteIndex;
 }
 
+uint64 FRigVMByteCode::InlineFunction(const FRigVMByteCode* FunctionByteCode, const FRigVMOperandArray& InOperands)
+{
+	check(FunctionByteCode);
+	uint64 OpByteIndex = ByteCode.Num();
+	ByteCode.Append(FunctionByteCode->ByteCode);
+	NumInstructions += FunctionByteCode->NumInstructions;
+	
+	return OpByteIndex;
+}
+
 uint64 FRigVMByteCode::AddExitOp()
 {
 	FRigVMBaseOp Op(ERigVMOpCode::Exit);
@@ -1733,6 +1743,94 @@ FRigVMOperandArray FRigVMByteCode::GetOperandsForOp(const FRigVMInstruction& InI
 	return FRigVMOperandArray();
 }
 
+uint64 FRigVMByteCode::GetFirstOperandByteIndex(const FRigVMInstruction& InInstruction) const
+{
+	if (InInstruction.OpCode >= ERigVMOpCode::Execute_0_Operands && InInstruction.OpCode <= ERigVMOpCode::Execute_64_Operands)
+	{
+		const uint64 ByteCodeIndex = InInstruction.ByteCodeIndex;
+		// if the bytecode is not aligned the OperandAlignment needs to be 0
+		ensure(bByteCodeIsAligned || InInstruction.OperandAlignment == 0);
+		return ByteCodeIndex + sizeof(FRigVMExecuteOp) + (uint64)InInstruction.OperandAlignment;
+	}
+
+	switch(InInstruction.OpCode)
+	{
+		case ERigVMOpCode::Copy:
+		{
+			const FRigVMCopyOp& Op = GetOpAt<FRigVMCopyOp>(InInstruction.ByteCodeIndex);
+			return InInstruction.ByteCodeIndex + ((uint64)&Op.Source - (uint64)&Op);
+		}
+		case ERigVMOpCode::Zero:
+		case ERigVMOpCode::BoolFalse:
+		case ERigVMOpCode::BoolTrue:
+		case ERigVMOpCode::Increment:
+		case ERigVMOpCode::Decrement:
+		case ERigVMOpCode::ArrayReset:
+		case ERigVMOpCode::ArrayReverse:
+		{
+			const FRigVMUnaryOp& Op = GetOpAt<FRigVMUnaryOp>(InInstruction.ByteCodeIndex);
+			return InInstruction.ByteCodeIndex + ((uint64)&Op.Arg - (uint64)&Op);
+		}
+		case ERigVMOpCode::Equals:
+		case ERigVMOpCode::NotEquals:
+		{
+			const FRigVMComparisonOp& Op = GetOpAt<FRigVMComparisonOp>(InInstruction.ByteCodeIndex);
+			return InInstruction.ByteCodeIndex + ((uint64)&Op.A - (uint64)&Op);
+		}
+		case ERigVMOpCode::JumpAbsoluteIf:
+		case ERigVMOpCode::JumpForwardIf:
+		case ERigVMOpCode::JumpBackwardIf:
+		{
+			const FRigVMJumpIfOp& Op = GetOpAt<FRigVMJumpIfOp>(InInstruction.ByteCodeIndex);
+			return InInstruction.ByteCodeIndex + ((uint64)&Op.Arg - (uint64)&Op);
+		}
+		case ERigVMOpCode::BeginBlock:
+		case ERigVMOpCode::ArrayGetNum:
+		case ERigVMOpCode::ArraySetNum:
+		case ERigVMOpCode::ArrayAppend:
+		case ERigVMOpCode::ArrayClone:
+		case ERigVMOpCode::ArrayRemove:
+		case ERigVMOpCode::ArrayUnion:
+		{
+			const FRigVMBinaryOp& Op = GetOpAt<FRigVMBinaryOp>(InInstruction.ByteCodeIndex);
+			return InInstruction.ByteCodeIndex + ((uint64)&Op.ArgA - (uint64)&Op);
+		}
+		case ERigVMOpCode::ArrayAdd:
+		case ERigVMOpCode::ArrayGetAtIndex:
+		case ERigVMOpCode::ArraySetAtIndex:
+		case ERigVMOpCode::ArrayInsert:
+		case ERigVMOpCode::ArrayDifference:
+		case ERigVMOpCode::ArrayIntersection:
+		{
+			const FRigVMTernaryOp& Op = GetOpAt<FRigVMTernaryOp>(InInstruction.ByteCodeIndex);
+			return InInstruction.ByteCodeIndex + ((uint64)&Op.ArgA - (uint64)&Op);
+		}
+		case ERigVMOpCode::ArrayFind:
+		{
+			const FRigVMQuaternaryOp& Op = GetOpAt<FRigVMQuaternaryOp>(InInstruction.ByteCodeIndex);
+			return InInstruction.ByteCodeIndex + ((uint64)&Op.ArgA - (uint64)&Op);
+		}
+		case ERigVMOpCode::ArrayIterator:
+		{
+			const FRigVMSenaryOp& Op = GetOpAt<FRigVMSenaryOp>(InInstruction.ByteCodeIndex);
+			return InInstruction.ByteCodeIndex + ((uint64)&Op.ArgA - (uint64)&Op);
+		}
+		case ERigVMOpCode::JumpAbsolute:
+		case ERigVMOpCode::JumpForward:
+		case ERigVMOpCode::JumpBackward:
+		case ERigVMOpCode::ChangeType:
+		case ERigVMOpCode::Exit:
+		case ERigVMOpCode::EndBlock:
+		case ERigVMOpCode::Invalid:
+		default:
+		{
+			break;
+		}
+	}
+
+	return INDEX_NONE;
+}
+
 TArray<int32> FRigVMByteCode::GetInstructionsForOperand(const FRigVMOperand& InOperand) const
 {
 	TArray<int32> InstructionIndices;
diff --git a/Engine/Source/Runtime/RigVM/Public/RigVMCore/RigVMByteCode.h b/Engine/Source/Runtime/RigVM/Public/RigVMCore/RigVMByteCode.h
--- a/Engine/Source/Runtime/RigVM/Public/RigVMCore/RigVMByteCode.h
+++ b/Engine/Source/Runtime/RigVM/Public/RigVMCore/RigVMByteCode.h
@@ -895,6 +895,9 @@ public:
 
 	// adds an execute operator given its function index operands
 	uint64 AddExecuteOp(uint16 InFunctionIndex, const FRigVMOperandArray& InOperands);
+	
+	// adds an execute operator given its function index operands
+	uint64 InlineFunction(const FRigVMByteCode* FunctionByteCode, const FRigVMOperandArray& InOperands);
 
 	// adds a zero operator to zero the memory of a given argument
 	uint64 AddZeroOp(const FRigVMOperand& InArg);
@@ -1056,6 +1059,9 @@ public:
 	// returns all of the operands for a given instruction
 	FRigVMOperandArray GetOperandsForOp(const FRigVMInstruction& InInstruction) const;
 
+	// returns the byte index of the first operand for this instructions
+	uint64 GetFirstOperandByteIndex(const FRigVMInstruction& InInstruction) const;
+
 	// returns all of the operands for a given instruction
 	TArray<int32> GetInstructionsForOperand(const FRigVMOperand& InOperand) const;
 
