Control Rig: Fix compilation error in Linux

diff --git a/Engine/Source/Developer/RigVMDeveloper/Private/RigVMCompiler/RigVMCompiler.cpp b/Engine/Source/Developer/RigVMDeveloper/Private/RigVMCompiler/RigVMCompiler.cpp
--- a/Engine/Source/Developer/RigVMDeveloper/Private/RigVMCompiler/RigVMCompiler.cpp
+++ b/Engine/Source/Developer/RigVMDeveloper/Private/RigVMCompiler/RigVMCompiler.cpp
@@ -770,7 +770,7 @@ bool URigVMCompiler::Compile(TArray<URigVMGraph*> InGraphs, URigVMController* In
 		Data.FunctionNames = WorkData.VM->FunctionNamesStorage;
 		Data.PropertyDescriptions = WorkData.PropertyDescriptions;
 		Data.PropertyPathDescriptions = WorkData.PropertyPathDescriptions;
-		for (const TPair<FString, FRigVMOperand> Pair : (*WorkData.PinPathToOperand))
+		for (const TPair<FString, FRigVMOperand>& Pair : (*WorkData.PinPathToOperand))
 		{
 			static const FString VariablePrefix = TEXT("Variable::");
 			if (Pair.Key.StartsWith(VariablePrefix))
