Control Rig: Fix Lyra cooking

diff --git a/Engine/Source/Developer/RigVMDeveloper/Private/RigVMCompiler/RigVMCompiler.cpp b/Engine/Source/Developer/RigVMDeveloper/Private/RigVMCompiler/RigVMCompiler.cpp
--- a/Engine/Source/Developer/RigVMDeveloper/Private/RigVMCompiler/RigVMCompiler.cpp
+++ b/Engine/Source/Developer/RigVMDeveloper/Private/RigVMCompiler/RigVMCompiler.cpp
@@ -1291,7 +1291,15 @@ int32 URigVMCompiler::TraverseInlineFunction(const FRigVMInlineFunctionExprAST*
 				if (Operand->GetMemoryType() == ERigVMMemoryType::External)
 				{
 					const FName& InnerVariableName = FunctionCompilationData.ExternalRegisterIndexToVariable[Operand->GetRegisterIndex()];
-					const FName& OuterVariableName = FunctionReferenceNode->GetVariableMap()[InnerVariableName];
+					FName OuterVariableName = InnerVariableName;
+					if (const FName* VariableRemapped = FunctionReferenceNode->GetVariableMap().Find(InnerVariableName))
+					{
+						OuterVariableName = *VariableRemapped;
+					}
+					else
+					{
+						ensureMsgf(!FunctionReferenceNode->RequiresVariableRemapping(), TEXT("Could not find variable %s in function reference %s variable map, in package %s\n"), *InnerVariableName.ToString(), *FunctionReferenceNode->GetNodePath(), *GetPackage()->GetPathName());
+					}
 					const FRigVMOperand& OuterOperand = WorkData.PinPathToOperand->FindChecked(FString::Printf(TEXT("Variable::%s"), *OuterVariableName.ToString()));
 					Operand->RegisterIndex = OuterOperand.RegisterIndex;
 				}
