Control Rig: Fix error during CitySample cook

diff --git a/Engine/Source/Developer/RigVMDeveloper/Private/RigVMCompiler/RigVMAST.cpp b/Engine/Source/Developer/RigVMDeveloper/Private/RigVMCompiler/RigVMAST.cpp
--- a/Engine/Source/Developer/RigVMDeveloper/Private/RigVMCompiler/RigVMAST.cpp
+++ b/Engine/Source/Developer/RigVMDeveloper/Private/RigVMCompiler/RigVMAST.cpp
@@ -861,9 +861,12 @@ FRigVMParserAST::FRigVMParserAST(TArray<URigVMGraph*> InGraphs, URigVMController
 	while (bContinueToFoldConstantBranches)
 	{
 		bContinueToFoldConstantBranches = false;
-		if (FoldConstantValuesToLiterals(InGraphs, InController, InExternalVariables, InRigVMUserData))
+		if (LibraryNodeBeingCompiled == nullptr)
 		{
-			bContinueToFoldConstantBranches = true;
+			if (FoldConstantValuesToLiterals(InGraphs, InController, InExternalVariables, InRigVMUserData))
+			{
+				bContinueToFoldConstantBranches = true;
+			}
 		}
 		if (FoldUnreachableBranches(InGraphs))
 		{
diff --git a/Engine/Source/Developer/RigVMDeveloper/Private/RigVMCompiler/RigVMCompiler.cpp b/Engine/Source/Developer/RigVMDeveloper/Private/RigVMCompiler/RigVMCompiler.cpp
--- a/Engine/Source/Developer/RigVMDeveloper/Private/RigVMCompiler/RigVMCompiler.cpp
+++ b/Engine/Source/Developer/RigVMDeveloper/Private/RigVMCompiler/RigVMCompiler.cpp
@@ -3029,7 +3029,7 @@ FRigVMOperand URigVMCompiler::FindOrAddRegister(const FRigVMVarExprAST* InVarExp
 	
 	if (ExistingOperand)
 	{
-		if(ExistingOperand->GetMemoryType() == MemoryType)
+		if(ExistingOperand->GetMemoryType() == MemoryType || ExistingOperand->GetMemoryType() == ERigVMMemoryType::External)
 		{
 			// Add any missing hash that shares this existing operand
 			for (const FString& VirtualPinHash : HashesWithSharedOperand)
