Control Rig: Fix cook EngineTest

diff --git a/Engine/Source/Developer/RigVMDeveloper/Private/RigVMCompiler/RigVMCompiler.cpp b/Engine/Source/Developer/RigVMDeveloper/Private/RigVMCompiler/RigVMCompiler.cpp
--- a/Engine/Source/Developer/RigVMDeveloper/Private/RigVMCompiler/RigVMCompiler.cpp
+++ b/Engine/Source/Developer/RigVMDeveloper/Private/RigVMCompiler/RigVMCompiler.cpp
@@ -799,6 +799,7 @@ bool URigVMCompiler::CompileFunction(URigVMLibraryNode* InLibraryNode, URigVMCon
 	TMap<FString, FRigVMOperand> Operands;
 	URigVM* TempVM = NewObject<URigVM>(InLibraryNode->GetContainedGraph());
 	const bool bSuccess = Compile({InLibraryNode->GetContainedGraph()}, InController, TempVM, ExternalVariables, InRigVMUserData, &Operands);
+	TempVM->ClearMemory();
 	TempVM->Rename(nullptr, GetTransientPackage(), REN_ForceNoResetLoaders | REN_DoNotDirty | REN_DontCreateRedirectors | REN_NonTransactional);
 	TempVM->MarkAsGarbage();
 
@@ -3029,21 +3030,18 @@ FRigVMOperand URigVMCompiler::FindOrAddRegister(const FRigVMVarExprAST* InVarExp
 	
 	if (ExistingOperand)
 	{
-		if(ExistingOperand->GetMemoryType() == MemoryType || ExistingOperand->GetMemoryType() == ERigVMMemoryType::External)
+		// Add any missing hash that shares this existing operand
+		for (const FString& VirtualPinHash : HashesWithSharedOperand)
 		{
-			// Add any missing hash that shares this existing operand
-			for (const FString& VirtualPinHash : HashesWithSharedOperand)
-			{
-				WorkData.PinPathToOperand->Add(VirtualPinHash, *ExistingOperand);
-			}
-			
-			if (!bIsDebugValue)
-			{
-				check(!WorkData.ExprToOperand.Contains(InVarExpr));
-				WorkData.ExprToOperand.Add(InVarExpr, *ExistingOperand);
-			}
-			return *ExistingOperand;
+			WorkData.PinPathToOperand->Add(VirtualPinHash, *ExistingOperand);
+		}
+		
+		if (!bIsDebugValue)
+		{
+			check(!WorkData.ExprToOperand.Contains(InVarExpr));
+			WorkData.ExprToOperand.Add(InVarExpr, *ExistingOperand);
 		}
+		return *ExistingOperand;
 	}
 
 	// create remaining operands / registers
