Control Rig: Fix crash when promoting to variable

diff --git a/Engine/Source/Developer/RigVMDeveloper/Private/RigVMCompiler/RigVMCompiler.cpp b/Engine/Source/Developer/RigVMDeveloper/Private/RigVMCompiler/RigVMCompiler.cpp
--- a/Engine/Source/Developer/RigVMDeveloper/Private/RigVMCompiler/RigVMCompiler.cpp
+++ b/Engine/Source/Developer/RigVMDeveloper/Private/RigVMCompiler/RigVMCompiler.cpp
@@ -461,7 +461,6 @@ bool URigVMCompiler::Compile(TArray<URigVMGraph*> InGraphs, URigVMController* In
 	
 	for (const FRigVMExternalVariable& ExternalVariable : InExternalVariables)
 	{
-		check(ExternalVariable.Property);
 		FRigVMOperand Operand = OutVM->AddExternalVariable(ExternalVariable);
 		FString Hash = FString::Printf(TEXT("Variable::%s"), *ExternalVariable.Name.ToString());
 		OutOperands->Add(Hash, Operand);
