Control Rig: Avoid dereferencing a potentially invalid container pointer

#rnx

diff --git a/Engine/Source/Developer/RigVMDeveloper/Private/RigVMCompiler/RigVMCompiler.cpp b/Engine/Source/Developer/RigVMDeveloper/Private/RigVMCompiler/RigVMCompiler.cpp
--- a/Engine/Source/Developer/RigVMDeveloper/Private/RigVMCompiler/RigVMCompiler.cpp
+++ b/Engine/Source/Developer/RigVMDeveloper/Private/RigVMCompiler/RigVMCompiler.cpp
@@ -2998,8 +2998,8 @@ FRigVMOperand URigVMCompiler::FindOrAddRegister(const FRigVMVarExprAST* InVarExp
 
 	TArray<FString> HashesWithSharedOperand;
 	
-	FRigVMOperand const* ExistingOperand = WorkData.PinPathToOperand->Find(Hash);
-	if (!ExistingOperand)
+	FRigVMOperand const* ExistingOperandPtr = WorkData.PinPathToOperand->Find(Hash);
+	if (!ExistingOperandPtr)
 	{
 		if(Settings.ASTSettings.bFoldAssignments) 		
 		{
@@ -3016,8 +3016,8 @@ FRigVMOperand URigVMCompiler::FindOrAddRegister(const FRigVMVarExprAST* InVarExp
 					HashesWithSharedOperand.Add(VirtualPinHash);
 					if (Pin != VirtualPin)
 					{
-						ExistingOperand = WorkData.PinPathToOperand->Find(VirtualPinHash);
-						if (ExistingOperand)
+						ExistingOperandPtr = WorkData.PinPathToOperand->Find(VirtualPinHash);
+						if (ExistingOperandPtr)
 						{
 							break;
 						}
@@ -3027,20 +3027,23 @@ FRigVMOperand URigVMCompiler::FindOrAddRegister(const FRigVMVarExprAST* InVarExp
 		}
 	}
 	
-	if (ExistingOperand)
+	if (ExistingOperandPtr)
 	{
+		// Dereference the operand pointer here since modifying the PinPathToOperand map will invalidate the pointer.
+		FRigVMOperand ExistingOperand = *ExistingOperandPtr;
+		
 		// Add any missing hash that shares this existing operand
 		for (const FString& VirtualPinHash : HashesWithSharedOperand)
 		{
-			WorkData.PinPathToOperand->Add(VirtualPinHash, *ExistingOperand);
+			WorkData.PinPathToOperand->Add(VirtualPinHash, ExistingOperand);
 		}
 		
 		if (!bIsDebugValue)
 		{
 			check(!WorkData.ExprToOperand.Contains(InVarExpr));
-			WorkData.ExprToOperand.Add(InVarExpr, *ExistingOperand);
+			WorkData.ExprToOperand.Add(InVarExpr, ExistingOperand);
 		}
-		return *ExistingOperand;
+		return ExistingOperand;
 	}
 
 	// create remaining operands / registers
@@ -3118,7 +3121,7 @@ FRigVMOperand URigVMCompiler::FindOrAddRegister(const FRigVMVarExprAST* InVarExp
 					FRigVMPropertyDescription Property = WorkData.GetProperty(Operand);
 					if(Property.IsValid())
 					{
-						if(ExistingOperand == nullptr)
+						if(ExistingOperandPtr == nullptr)
 						{
 							WorkData.PinPathToOperand->Add(Hash, Operand);
 						}
@@ -3183,7 +3186,7 @@ FRigVMOperand URigVMCompiler::FindOrAddRegister(const FRigVMVarExprAST* InVarExp
 	}
 	else
 	{
-		if(ExistingOperand == nullptr)
+		if(ExistingOperandPtr == nullptr)
 		{
 			WorkData.PinPathToOperand->Add(Hash, Operand);
 		}
