Control Rig: Fix wrong loading of Set Transform Array Metadata pin

diff --git a/Engine/Source/Developer/RigVMDeveloper/Private/RigVMModel/RigVMController.cpp b/Engine/Source/Developer/RigVMDeveloper/Private/RigVMModel/RigVMController.cpp
--- a/Engine/Source/Developer/RigVMDeveloper/Private/RigVMModel/RigVMController.cpp
+++ b/Engine/Source/Developer/RigVMDeveloper/Private/RigVMModel/RigVMController.cpp
@@ -14862,7 +14862,7 @@ void URigVMController::RepopulatePinsOnNode(URigVMNode* InNode, bool bFollowCore
 
 			AddNodePin(DispatchNode, Pin);
 
-			if(!Pin->IsWildCard())
+			if(!Pin->IsWildCard() && !Pin->IsArray())
 			{
 				// any serialize default values will be applied later with ApplyPinStates(...)
 				const FString DefaultValue = DispatchNode->GetInitialDefaultValueForPin(Pin->GetFName());
