RigVM: Ignore cast links during template resolval

diff --git a/Engine/Source/Developer/RigVMDeveloper/Private/RigVMModel/RigVMController.cpp b/Engine/Source/Developer/RigVMDeveloper/Private/RigVMModel/RigVMController.cpp
--- a/Engine/Source/Developer/RigVMDeveloper/Private/RigVMModel/RigVMController.cpp
+++ b/Engine/Source/Developer/RigVMDeveloper/Private/RigVMModel/RigVMController.cpp
@@ -9320,7 +9320,7 @@ bool URigVMController::AddLink(URigVMPin* OutputPin, URigVMPin* InputPin, bool b
 		{
 			if(InPin->IsRootPin())
 			{
-				if(URigVMTemplateNode* TemplateNode = Cast<URigVMTemplateNode>(InPin->GetNode()))
+				if(const URigVMTemplateNode* TemplateNode = Cast<URigVMTemplateNode>(InPin->GetNode()))
 				{
 					if(const FRigVMTemplate* Template = TemplateNode->GetTemplate())
 					{
@@ -9334,7 +9334,43 @@ bool URigVMController::AddLink(URigVMPin* OutputPin, URigVMPin* InputPin, bool b
 					}
 				}
 			}
-			return InPin->GetNode()->GetLinks().Num() > 0;
+
+			// we only rely on the cast if the number of links on non-singleton arguments is > 0
+			const TArray<URigVMLink*> Links = InPin->GetNode()->GetLinks();
+			int32 NumLinksRemaining = Links.Num();
+			if(const URigVMTemplateNode* TemplateNode = Cast<URigVMTemplateNode>(InPin->GetNode()))
+			{
+				if(const FRigVMTemplate* Template = TemplateNode->GetTemplate())
+				{
+					for(const URigVMLink* Link : Links)
+					{
+						const URigVMPin* SourcePin = Link->GetSourcePin();
+						const URigVMPin* TargetPin = Link->GetTargetPin();
+
+						if(SourcePin && TargetPin)
+						{
+							const URigVMPin* PinOnTemplateNode = SourcePin->GetNode() == TemplateNode ? SourcePin : TargetPin;
+							if(PinOnTemplateNode->IsRootPin())
+							{
+								const FName ArgumentName = PinOnTemplateNode->GetFName();
+								if(const FRigVMTemplateArgument* Argument = Template->FindArgument(ArgumentName))
+								{
+									if(Argument->IsSingleton())
+									{
+										NumLinksRemaining--;
+									}
+								}
+							}
+							else
+							{
+								NumLinksRemaining--;
+							}
+						}
+					}
+				}
+			}
+
+			return NumLinksRemaining > 0;
 		};
 		
 		// only rely on a cast if there's already a link
@@ -16320,6 +16356,7 @@ bool URigVMController::PrepareTemplatePinForType(URigVMPin* InPin, const TArray<
 			TGuardValue<FRigVMController_RequestBreakLinksDialogDelegate> GuardDelegate(RequestBreakLinksDialogDelegate, nullptr);
 			TGuardValue<bool> GuardNotifications(bSuspendNotifications, true);
 			TGuardValue<bool> SuspendRecomputeTemplates(bSuspendRecomputingTemplateFilters, true);
+			TGuardValue<bool> DisableCasting(bEnableTypeCasting, false);
 			OpenUndoBracket(FString::Printf(TEXT("Resolve wildcard pin %s"), *InPin->GetPinPath()));
 
 			bool bBrokenLinks = false;
@@ -17218,10 +17255,16 @@ bool URigVMController::PropagateTemplateFilteredTypes(URigVMTemplateNode* InNode
 {
 	auto UpdateAndPropagate = [&](URigVMPin* Pin)
 	{
-		TArray<URigVMPin*> OtherPins = Pin->GetLinkedSourcePins();
-		OtherPins.Append(Pin->GetLinkedTargetPins());
-		for (URigVMPin* OtherPin : OtherPins)
+		const TArray<URigVMLink*>& Links = Pin->GetLinks();
+		for (URigVMLink* Link : Links)
 		{
+			// skip links which contain a cast
+			if(LinksWithCast.Contains(Link->GetPinPathRepresentation()))
+			{
+				continue;
+			}
+			
+			URigVMPin* OtherPin = Link->GetSourcePin() == Pin ? Link->GetTargetPin() : Link->GetSourcePin();
 			bool bPropagate = false;
 			bool bIsTemplate = false;
 			if (URigVMTemplateNode* OtherTemplate = Cast<URigVMTemplateNode>(OtherPin->GetNode()))
@@ -17253,7 +17296,6 @@ bool URigVMController::PropagateTemplateFilteredTypes(URigVMTemplateNode* InNode
 						}
 						else
 						{
-							URigVMLink* Link = Pin->FindLinkForPin(OtherPin);
 							if(!bEnableTypeCasting || !RigVMTypeUtils::CanCastTypes(Link->GetSourcePin()->GetTypeIndex(), Link->GetTargetPin()->GetTypeIndex()))
 							{
 								ensureMsgf(!ActionStack->BracketActions.IsEmpty(), TEXT("Unexpected link broken %s in package %s"), *Link->GetPinPathRepresentation(), *GetPackage()->GetPathName());
@@ -17269,7 +17311,6 @@ bool URigVMController::PropagateTemplateFilteredTypes(URigVMTemplateNode* InNode
 			{
 				if (!InNode->FilteredSupportsType(Pin, OtherPin->GetTypeIndex()))
 				{
-					URigVMLink* Link = Pin->FindLinkForPin(OtherPin);
 					if(!bEnableTypeCasting || !RigVMTypeUtils::CanCastTypes(Link->GetSourcePin()->GetTypeIndex(), Link->GetTargetPin()->GetTypeIndex()))
 					{
 						ensureMsgf(!ActionStack->BracketActions.IsEmpty(), TEXT("Unexpected link broken %s in package %s"), *Link->GetPinPathRepresentation(), *GetPackage()->GetPathName());
@@ -17831,6 +17872,41 @@ bool URigVMController::RecomputeAllTemplateFilteredPermutations(bool bSetupUndoR
 	URigVMGraph* Graph = GetGraph();
 	check(Graph);
 
+	// update the transient list of links containing a cast
+	if(bEnableTypeCasting)
+	{
+		LinksWithCast.Reset();
+		for(URigVMLink* Link : Graph->GetLinks())
+		{
+			const URigVMPin* SourcePin = Link->GetSourcePin();
+			const URigVMPin* TargetPin = Link->GetTargetPin();
+
+			if(SourcePin && TargetPin)
+			{
+				const TRigVMTypeIndex SourceTypeIndex = SourcePin->GetTypeIndex();
+				const TRigVMTypeIndex TargetTypeIndex = TargetPin->GetTypeIndex();
+
+				if(SourceTypeIndex != TargetTypeIndex &&
+					RigVMTypeUtils::TypeIndex::WildCard != SourceTypeIndex &&
+					RigVMTypeUtils::TypeIndex::WildCardArray != SourceTypeIndex &&
+					RigVMTypeUtils::TypeIndex::WildCard != TargetTypeIndex &&
+					RigVMTypeUtils::TypeIndex::WildCardArray != TargetTypeIndex &&
+					!FRigVMRegistry::Get().CanMatchTypes(SourceTypeIndex, TargetTypeIndex, true))
+				{
+					ensureMsgf(RigVMTypeUtils::CanCastTypes(SourceTypeIndex, TargetTypeIndex),
+						TEXT("Incompatible types for link %s (%s) -> %s (%s) in %s.")
+						, *SourcePin->GetPinPath(true)
+						, *SourcePin->GetCPPType()
+						, *TargetPin->GetPinPath(true)
+						, *TargetPin->GetCPPType()
+						, *GetPackage()->GetPathName());
+
+					LinksWithCast.Add(Link->GetPinPathRepresentation());
+				}
+			}	
+		}
+	}
+
 	// Save all template pin types, in case we need them after initializing
 	TMap<URigVMPin*, TRigVMTypeIndex> TypesBeforeRecomputing;
 	for (URigVMNode* Node : Graph->GetNodes())
@@ -17869,10 +17945,21 @@ bool URigVMController::RecomputeAllTemplateFilteredPermutations(bool bSetupUndoR
 		InitializeAllTemplateFiltersInGraph(bSetupUndoRedo, false);	
 
 		// Apply all links to update filtered permutations
-		TArray<URigVMLink*> SortedLinks = Graph->GetLinks().FilterByPredicate([](URigVMLink* Link)
+		TArray<URigVMLink*> SortedLinks = Graph->GetLinks().FilterByPredicate([this](URigVMLink* Link)
 		{
 			// Filter out detached links
-			return Link->GetSourcePin()->IsLinkedTo(Link->GetTargetPin());
+			if(!Link->GetSourcePin()->IsLinkedTo(Link->GetTargetPin()))
+			{
+				return false;
+			}
+
+			// filter out links with casts
+			if(LinksWithCast.Contains(Link->GetPinPathRepresentation()))
+			{
+				return false;
+			}
+
+			return true;
 		});
 
 		// Solve for unit nodes first, other templates are more expensive to solve. This way we are avoiding solving
diff --git a/Engine/Source/Developer/RigVMDeveloper/Public/RigVMModel/RigVMController.h b/Engine/Source/Developer/RigVMDeveloper/Public/RigVMModel/RigVMController.h
--- a/Engine/Source/Developer/RigVMDeveloper/Public/RigVMModel/RigVMController.h
+++ b/Engine/Source/Developer/RigVMDeveloper/Public/RigVMModel/RigVMController.h
@@ -1178,6 +1178,7 @@ private:
 	bool bIgnoreRerouteCompactnessChanges;
 	ERigVMPinDirection UserLinkDirection;
 	bool bEnableTypeCasting;
+	mutable TArray<FString> LinksWithCast;
 
 	// temporary maps used for pin redirection
 	// only valid between Detach & ReattachLinksToPinObjects
