Control Rig: Fix some nodes connected to if node not bubbling up when inside a function

diff --git a/Engine/Source/Developer/RigVMDeveloper/Private/RigVMCompiler/RigVMAST.cpp b/Engine/Source/Developer/RigVMDeveloper/Private/RigVMCompiler/RigVMAST.cpp
--- a/Engine/Source/Developer/RigVMDeveloper/Private/RigVMCompiler/RigVMAST.cpp
+++ b/Engine/Source/Developer/RigVMDeveloper/Private/RigVMCompiler/RigVMAST.cpp
@@ -1037,8 +1037,30 @@ FRigVMExprAST* FRigVMParserAST::CreateExpressionForNode(const FRigVMASTProxy& In
 {
 	URigVMNode* Node = InNodeProxy.GetSubjectChecked<URigVMNode>();
 
+	bool bIsEntry = Node->IsEvent();
+	if (LibraryNodeBeingCompiled)
+	{
+		if (Node->GetTypedOuter<URigVMLibraryNode>() == LibraryNodeBeingCompiled)
+		{
+			if (URigVMFunctionEntryNode* EntryNode = Cast<URigVMFunctionEntryNode>(Node))
+			{
+				if (EntryNode->IsMutable())
+				{
+					bIsEntry = true;
+				}
+			}
+			if (URigVMFunctionReturnNode* ReturnNode = Cast<URigVMFunctionReturnNode>(Node))
+			{
+				if (!ReturnNode->IsMutable())
+				{
+					bIsEntry = true;
+				}
+			}
+		}
+	}
+
 	FRigVMExprAST* NodeExpr = nullptr;
-	if (Node->IsEvent())
+	if (bIsEntry)
 	{
 		NodeExpr = MakeExpr<FRigVMEntryExprAST>(InNodeProxy);
 		NodeExpr->Name = Node->GetEventName();
@@ -1453,6 +1475,11 @@ void FRigVMParserAST::InjectExitsToEntries()
 {
 	DECLARE_SCOPE_HIERARCHICAL_COUNTER_FUNC()
 
+	if (LibraryNodeBeingCompiled)
+	{
+		return;
+	}
+	
 	for (FRigVMExprAST* RootExpr : RootExpressions)
 	{
 		if (RootExpr->IsA(FRigVMExprAST::EType::Entry))
diff --git a/Engine/Source/Developer/RigVMDeveloper/Private/RigVMCompiler/RigVMCompiler.cpp b/Engine/Source/Developer/RigVMDeveloper/Private/RigVMCompiler/RigVMCompiler.cpp
--- a/Engine/Source/Developer/RigVMDeveloper/Private/RigVMCompiler/RigVMCompiler.cpp
+++ b/Engine/Source/Developer/RigVMDeveloper/Private/RigVMCompiler/RigVMCompiler.cpp
@@ -498,78 +498,82 @@ bool URigVMCompiler::Compile(TArray<URigVMGraph*> InGraphs, URigVMController* In
 		{
 			for (int32 FunctionIndex=0; FunctionIndex<FunctionsToCompile.Num(); ++FunctionIndex)
 			{
-				URigVMLibraryNode* FunctionToCompile = FunctionsToCompile[FunctionIndex];
-				if (CompiledFunctions.Contains(FunctionToCompile))
+				if (URigVMLibraryNode* FunctionToCompile = FunctionsToCompile[FunctionIndex])
 				{
-					continue;
-				}
+					if (CompiledFunctions.Contains(FunctionToCompile))
+					{
+						continue;
+					}
 			
-				bool bCanBeCompiled = true;
-				TArray<URigVMNode*> Nodes = FunctionToCompile->GetContainedNodes();
-				for (int32 i=0; i<Nodes.Num() && bCanBeCompiled; ++i)
-				{
-					if (URigVMFunctionReferenceNode* FunctionReference = Cast<URigVMFunctionReferenceNode>(Nodes[i]))
+					bool bCanBeCompiled = true;
+					TArray<URigVMNode*> Nodes = FunctionToCompile->GetContainedNodes();
+					for (int32 i=0; i<Nodes.Num() && bCanBeCompiled; ++i)
 					{
-						if (const URigVMFunctionLibrary* ReferencedFunctionLibrary = FunctionReference->GetReferencedNode()->GetLibrary())
+						if (URigVMFunctionReferenceNode* FunctionReference = Cast<URigVMFunctionReferenceNode>(Nodes[i]))
 						{
-							if (ReferencedFunctionLibrary == FunctionLibrary)
+							if (const URigVMFunctionLibrary* ReferencedFunctionLibrary = FunctionReference->GetReferencedNode()->GetLibrary())
 							{
-								if (!CompiledFunctions.Contains(FunctionReference->GetReferencedNode()))
+								if (ReferencedFunctionLibrary == FunctionLibrary)
 								{
-									bCanBeCompiled = false;
-									break;
+									if (!CompiledFunctions.Contains(FunctionReference->GetReferencedNode()))
+									{
+										bCanBeCompiled = false;
+										break;
+									}
 								}
-							}
-							else
-							{
-								// Public function.
-								// todo Try to find its compilation data
-								bool bFoundCompilationData = true;
-								if (!CompiledFunctions.Contains(FunctionReference->GetReferencedNode()))
+								else
 								{
-									FunctionsToCompile.AddUnique(FunctionReference->GetReferencedNode());
-									bCanBeCompiled = false;
-									break;
+									// Public function.
+									// todo Try to find its compilation data
+									bool bFoundCompilationData = true;
+									if (!CompiledFunctions.Contains(FunctionReference->GetReferencedNode()))
+									{
+										FunctionsToCompile.AddUnique(FunctionReference->GetReferencedNode());
+										bCanBeCompiled = false;
+										break;
+									}
+
+									if (!bFoundCompilationData)
+									{
+										static const FString FunctionCompilationErrorMessage = TEXT("Could not find compilation data for public function @@.");
+										Settings.ASTSettings.Report(EMessageSeverity::Error, FunctionReference, FunctionCompilationErrorMessage);
+										bEncounteredGraphError = true;
+										bCanBeCompiled = false;
+										break;
+									}
 								}
+							}
+						}
 
-								if (!bFoundCompilationData)
-								{
-									static const FString FunctionCompilationErrorMessage = TEXT("Could not find compilation data for public function @@.");
-									Settings.ASTSettings.Report(EMessageSeverity::Error, FunctionReference, FunctionCompilationErrorMessage);
-									bEncounteredGraphError = true;
-									bCanBeCompiled = false;
-									break;
-								}
+						if (bCanBeCompiled)
+						{
+							if (URigVMLibraryNode* LibraryNode = Cast<URigVMLibraryNode>(Nodes[i]))
+							{
+								Nodes.Append(LibraryNode->GetContainedNodes());	
 							}
 						}
 					}
 
 					if (bCanBeCompiled)
 					{
-						if (URigVMLibraryNode* LibraryNode = Cast<URigVMLibraryNode>(Nodes[i]))
+						if (!CompileFunction(FunctionToCompile, InController, UserData))
 						{
-							Nodes.Append(LibraryNode->GetContainedNodes());	
+							static const FString FunctionCompilationErrorMessage = TEXT("Function @@ failed to compile.");
+							Settings.ASTSettings.Report(EMessageSeverity::Error, FunctionToCompile, FunctionCompilationErrorMessage);
+							bEncounteredGraphError = true;
+							break;
 						}
 					}
 				}
-
-				if (bCanBeCompiled)
-				{
-					if (!CompileFunction(FunctionToCompile, InController, UserData))
-					{
-						static const FString FunctionCompilationErrorMessage = TEXT("Function @@ failed to compile.");
-						Settings.ASTSettings.Report(EMessageSeverity::Error, FunctionToCompile, FunctionCompilationErrorMessage);
-						bEncounteredGraphError = true;
-						break;
-					}
-				}
 			}
+
+			FunctionsToCompile.RemoveAll([](const URigVMLibraryNode* Node) { return Node == nullptr; });
 			
 			if(bEncounteredGraphError)
 			{
 				return false;
 			}
-		}		
+		}
 	}
 
 	FRigVMCompilerWorkData WorkData;
@@ -984,55 +988,89 @@ void URigVMCompiler::TraverseBlock(const FRigVMBlockExprAST* InExpr, FRigVMCompi
 
 void URigVMCompiler::TraverseEntry(const FRigVMEntryExprAST* InExpr, FRigVMCompilerWorkData& WorkData)
 {
-	URigVMUnitNode* UnitNode = Cast<URigVMUnitNode>(InExpr->GetNode());
-	if(!ValidateNode(UnitNode))
+	if (URigVMUnitNode* UnitNode = Cast<URigVMUnitNode>(InExpr->GetNode()))
 	{
-		return;
-	}
+		if(!ValidateNode(UnitNode))
+		{
+			return;
+		}
 
-	if (WorkData.bSetupMemory)
-	{
-		TSharedPtr<FStructOnScope> DefaultStruct = UnitNode->ConstructStructInstance();
-		TraverseChildren(InExpr, WorkData);
-	}
-	else
-	{
-		TArray<FRigVMOperand> Operands;
-		for (FRigVMExprAST* ChildExpr : *InExpr)
+		if (WorkData.bSetupMemory)
 		{
-			if (ChildExpr->IsA(FRigVMExprAST::EType::Var))
-			{
-				Operands.Add(WorkData.ExprToOperand.FindChecked(GetSourceVarExpr(ChildExpr)));
-			}
-			else
+			TSharedPtr<FStructOnScope> DefaultStruct = UnitNode->ConstructStructInstance();
+			TraverseChildren(InExpr, WorkData);
+		}
+		else
+		{
+			TArray<FRigVMOperand> Operands;
+			for (FRigVMExprAST* ChildExpr : *InExpr)
 			{
-				break;
+				if (ChildExpr->IsA(FRigVMExprAST::EType::Var))
+				{
+					Operands.Add(WorkData.ExprToOperand.FindChecked(GetSourceVarExpr(ChildExpr)));
+				}
+				else
+				{
+					break;
+				}
 			}
-		}
 
-		// setup the instruction
-		int32 FunctionIndex = WorkData.VM->AddRigVMFunction(UnitNode->GetScriptStruct(), UnitNode->GetMethodName());
-		WorkData.VM->GetByteCode().AddExecuteOp(FunctionIndex, Operands);
+			// setup the instruction
+			int32 FunctionIndex = WorkData.VM->AddRigVMFunction(UnitNode->GetScriptStruct(), UnitNode->GetMethodName());
+			WorkData.VM->GetByteCode().AddExecuteOp(FunctionIndex, Operands);
 		
-		int32 EntryInstructionIndex = WorkData.VM->GetByteCode().GetNumInstructions() - 1;
-		FName Entryname = UnitNode->GetEventName();
+			int32 EntryInstructionIndex = WorkData.VM->GetByteCode().GetNumInstructions() - 1;
+			FName Entryname = UnitNode->GetEventName();
 
-		if (WorkData.VM->GetByteCode().FindEntryIndex(Entryname) == INDEX_NONE)
+			if (WorkData.VM->GetByteCode().FindEntryIndex(Entryname) == INDEX_NONE)
+			{
+				FRigVMByteCodeEntry Entry;
+				Entry.Name = Entryname;
+				Entry.InstructionIndex = EntryInstructionIndex;
+				WorkData.VM->GetByteCode().Entries.Add(Entry);
+			}
+
+			if (Settings.SetupNodeInstructionIndex)
+			{
+				const FRigVMCallstack Callstack = InExpr->GetProxy().GetCallstack();
+				WorkData.VM->GetByteCode().SetSubject(EntryInstructionIndex, Callstack.GetCallPath(), Callstack.GetStack());
+			}
+		}
+	}
+	else if (CurrentCompilationFunction && InExpr->NumParents() == 0)
+	{
+		// Initialize local variables
+		if (WorkData.bSetupMemory)
 		{
-			FRigVMByteCodeEntry Entry;
-			Entry.Name = Entryname;
-			Entry.InstructionIndex = EntryInstructionIndex;
-			WorkData.VM->GetByteCode().Entries.Add(Entry);
+			TArray<FRigVMGraphVariableDescription> LocalVariables = CurrentCompilationFunction->GetContainedGraph()->GetLocalVariables();
+			for (const FRigVMGraphVariableDescription& Variable : LocalVariables)
+			{
+				FString Path = FString::Printf(TEXT("LocalVariableDefault::%s|%s::Const"), *CurrentCompilationFunction->GetContainedGraph()->GetGraphName(), *Variable.Name.ToString());
+				FRigVMOperand Operand = WorkData.AddProperty(ERigVMMemoryType::Literal, *Path, Variable.CPPType, Variable.CPPTypeObject, Variable.DefaultValue);
+				WorkData.PinPathToOperand->Add(Path, Operand);
+			}
 		}
-
-		if (Settings.SetupNodeInstructionIndex)
+		else
 		{
-			const FRigVMCallstack Callstack = InExpr->GetProxy().GetCallstack();
-			WorkData.VM->GetByteCode().SetSubject(EntryInstructionIndex, Callstack.GetCallPath(), Callstack.GetStack());
+			TArray<FRigVMGraphVariableDescription> LocalVariables = CurrentCompilationFunction->GetContainedGraph()->GetLocalVariables();
+			for (const FRigVMGraphVariableDescription& Variable : LocalVariables)
+			{
+				FString TargetPath = FString::Printf(TEXT("LocalVariable::%s|%s"), *CurrentCompilationFunction->GetContainedGraph()->GetGraphName(), *Variable.Name.ToString());
+				FString SourcePath = FString::Printf(TEXT("LocalVariableDefault::%s|%s::Const"), *CurrentCompilationFunction->GetContainedGraph()->GetGraphName(), *Variable.Name.ToString());
+				FRigVMOperand* TargetPtr = WorkData.PinPathToOperand->Find(TargetPath);
+				FRigVMOperand* SourcePtr = WorkData.PinPathToOperand->Find(SourcePath);
+				if (SourcePtr && TargetPtr)
+				{
+					const FRigVMOperand& Source = *SourcePtr;
+					const FRigVMOperand& Target = *TargetPtr;
+	
+					WorkData.VM->GetByteCode().AddCopyOp(WorkData.VM->GetCopyOpForOperands(Source, Target));						
+				}
+			}
 		}
-
-		TraverseChildren(InExpr, WorkData);
 	}
+
+	TraverseChildren(InExpr, WorkData);
 }
 
 int32 URigVMCompiler::TraverseCallExtern(const FRigVMCallExternExprAST* InExpr, FRigVMCompilerWorkData& WorkData)
@@ -1485,39 +1523,6 @@ void URigVMCompiler::TraverseForLoop(const FRigVMCallExternExprAST* InExpr, FRig
 
 void URigVMCompiler::TraverseNoOp(const FRigVMNoOpExprAST* InExpr, FRigVMCompilerWorkData& WorkData)
 {
-	if (CurrentCompilationFunction && InExpr->NumParents() == 0)
-	{
-		// Initialize local variables
-		if (WorkData.bSetupMemory)
-		{
-			TArray<FRigVMGraphVariableDescription> LocalVariables = CurrentCompilationFunction->GetContainedGraph()->GetLocalVariables();
-			for (const FRigVMGraphVariableDescription& Variable : LocalVariables)
-			{
-				FString Path = FString::Printf(TEXT("LocalVariableDefault::%s|%s::Const"), *CurrentCompilationFunction->GetContainedGraph()->GetGraphName(), *Variable.Name.ToString());
-				FRigVMOperand Operand = WorkData.AddProperty(ERigVMMemoryType::Literal, *Path, Variable.CPPType, Variable.CPPTypeObject, Variable.DefaultValue);
-				WorkData.PinPathToOperand->Add(Path, Operand);
-			}
-		}
-		else
-		{
-			TArray<FRigVMGraphVariableDescription> LocalVariables = CurrentCompilationFunction->GetContainedGraph()->GetLocalVariables();
-			for (const FRigVMGraphVariableDescription& Variable : LocalVariables)
-			{
-				FString TargetPath = FString::Printf(TEXT("LocalVariable::%s|%s"), *CurrentCompilationFunction->GetContainedGraph()->GetGraphName(), *Variable.Name.ToString());
-				FString SourcePath = FString::Printf(TEXT("LocalVariableDefault::%s|%s::Const"), *CurrentCompilationFunction->GetContainedGraph()->GetGraphName(), *Variable.Name.ToString());
-				FRigVMOperand* TargetPtr = WorkData.PinPathToOperand->Find(TargetPath);
-				FRigVMOperand* SourcePtr = WorkData.PinPathToOperand->Find(SourcePath);
-				if (SourcePtr && TargetPtr)
-				{
-					const FRigVMOperand& Source = *SourcePtr;
-					const FRigVMOperand& Target = *TargetPtr;
-	
-					WorkData.VM->GetByteCode().AddCopyOp(WorkData.VM->GetCopyOpForOperands(Source, Target));						
-				}
-			}
-		}
-	}
-
 	TraverseChildren(InExpr, WorkData);
 }
 
