Control Rig: Fix crash when set variable node is created inside function

diff --git a/Engine/Source/Developer/RigVMDeveloper/Private/RigVMCompiler/RigVMCompiler.cpp b/Engine/Source/Developer/RigVMDeveloper/Private/RigVMCompiler/RigVMCompiler.cpp
--- a/Engine/Source/Developer/RigVMDeveloper/Private/RigVMCompiler/RigVMCompiler.cpp
+++ b/Engine/Source/Developer/RigVMDeveloper/Private/RigVMCompiler/RigVMCompiler.cpp
@@ -214,6 +214,18 @@
 	const FProperty* Property = GetPropertyForOperand(InOperand);
 	if(Property == nullptr)
 	{
+		if (InOperand.GetMemoryType() == ERigVMMemoryType::External)
+		{
+			const TArray<FRigVMExternalVariable>& ExternalVariables = VM->GetExternalVariables();
+			if (ExternalVariables.IsValidIndex(InOperand.GetRegisterIndex()))
+			{
+				const FRigVMExternalVariable& Variable = ExternalVariables[InOperand.GetRegisterIndex()];
+				FString CPPType;
+				UObject* CPPTypeObject;
+				RigVMTypeUtils::CPPTypeFromExternalVariable(Variable, CPPType, &CPPTypeObject);
+				return FRigVMRegistry::Get().GetTypeIndex(*CPPType, CPPTypeObject);
+			}
+		}
 		return INDEX_NONE;
 	}
 
@@ -275,6 +287,117 @@
 	}
 	OutOperands->Reset();
 
+	URigVMFunctionLibrary* FunctionLibrary = InGraphs[0]->GetDefaultFunctionLibrary();
+	bool bEncounteredGraphError = false;
+	if (CurrentCompilationFunction == nullptr)
+	{
+		ClearFunctionCompilationData();
+
+		// We need to compile the functions in order of dependency (leaf functions with no dependency first)
+		TArray<URigVMLibraryNode*> FunctionsToCompile;
+		if (FunctionLibrary)
+		{
+			FunctionsToCompile = FunctionLibrary->GetFunctions();
+		}
+
+		// Add public functions from other assets
+		// todo: this should be a check if the compilation data exists
+		{
+			TArray<URigVMNode*> Nodes = InGraphs[0]->GetNodes();
+			for (int32 i=0; i<Nodes.Num(); ++i)
+			{
+				if (URigVMFunctionReferenceNode* ReferenceNode = Cast<URigVMFunctionReferenceNode>(Nodes[i]))
+				{
+					FunctionsToCompile.AddUnique(ReferenceNode->GetReferencedNode());
+				}
+				if (URigVMCollapseNode* CollapseNode = Cast<URigVMCollapseNode>(Nodes[i]))
+				{
+					Nodes.Append(CollapseNode->GetContainedGraph()->GetNodes());
+				}
+			}
+		}
+		while (CompiledFunctions.Num() < FunctionsToCompile.Num())
+		{
+			for (int32 FunctionIndex=0; FunctionIndex<FunctionsToCompile.Num(); ++FunctionIndex)
+			{
+				if (URigVMLibraryNode* FunctionToCompile = FunctionsToCompile[FunctionIndex])
+				{
+					if (CompiledFunctions.Contains(FunctionToCompile))
+					{
+						continue;
+					}
+			
+					bool bCanBeCompiled = true;
+					TArray<URigVMNode*> Nodes = FunctionToCompile->GetContainedNodes();
+					for (int32 i=0; i<Nodes.Num() && bCanBeCompiled; ++i)
+					{
+						if (URigVMFunctionReferenceNode* FunctionReference = Cast<URigVMFunctionReferenceNode>(Nodes[i]))
+						{
+							if (const URigVMFunctionLibrary* ReferencedFunctionLibrary = FunctionReference->GetReferencedNode()->GetLibrary())
+							{
+								if (ReferencedFunctionLibrary == FunctionLibrary)
+								{
+									if (!CompiledFunctions.Contains(FunctionReference->GetReferencedNode()))
+									{
+										bCanBeCompiled = false;
+										break;
+									}
+								}
+								else
+								{
+									// Public function.
+									// todo Try to find its compilation data
+									bool bFoundCompilationData = true;
+									if (!CompiledFunctions.Contains(FunctionReference->GetReferencedNode()))
+									{
+										FunctionsToCompile.AddUnique(FunctionReference->GetReferencedNode());
+										bCanBeCompiled = false;
+										break;
+									}
+
+									if (!bFoundCompilationData)
+									{
+										static const FString FunctionCompilationErrorMessage = TEXT("Could not find compilation data for public function @@.");
+										Settings.ASTSettings.Report(EMessageSeverity::Error, FunctionReference, FunctionCompilationErrorMessage);
+										bEncounteredGraphError = true;
+										bCanBeCompiled = false;
+										break;
+									}
+								}
+							}
+						}
+
+						if (bCanBeCompiled)
+						{
+							if (URigVMLibraryNode* LibraryNode = Cast<URigVMLibraryNode>(Nodes[i]))
+							{
+								Nodes.Append(LibraryNode->GetContainedNodes());	
+							}
+						}
+					}
+
+					if (bCanBeCompiled)
+					{
+						if (!CompileFunction(FunctionToCompile, InController, UserData))
+						{
+							static const FString FunctionCompilationErrorMessage = TEXT("Function @@ failed to compile.");
+							Settings.ASTSettings.Report(EMessageSeverity::Error, FunctionToCompile, FunctionCompilationErrorMessage);
+							bEncounteredGraphError = true;
+							break;
+						}
+					}
+				}
+			}
+
+			FunctionsToCompile.RemoveAll([](const URigVMLibraryNode* Node) { return Node == nullptr; });
+			
+			if(bEncounteredGraphError)
+			{
+				return false;
+			}
+		}
+	}
+
 #if WITH_EDITOR
 
 	// traverse all graphs and try to clear out orphan pins
@@ -282,7 +405,6 @@
 	TArray<URigVMGraph*> VisitedGraphs;
 	VisitedGraphs.Append(InGraphs);
 
-	bool bEncounteredGraphError = false;
 	for(int32 GraphIndex=0; GraphIndex<VisitedGraphs.Num(); GraphIndex++)
 	{
 		URigVMGraph* VisitedGraph = VisitedGraphs[GraphIndex];
@@ -328,6 +450,34 @@
 						Settings.ASTSettings.Report(EMessageSeverity::Error, ModelNode, UnmappedMessage);
 						bEncounteredGraphError = true;
 					}
+
+					if (const FunctionCompilationData* CompilationData = CompiledFunctions.Find(FunctionReferenceNode->GetReferencedNode()))
+					{
+						for (const TPair<int32, FName>& Pair : CompilationData->ExternalRegisterIndexToVariable)
+						{
+							FName OuterName = FunctionReferenceNode->GetOuterVariableName(Pair.Value);
+							if (OuterName.IsNone())
+							{
+								OuterName = Pair.Value;
+							}
+						
+							if (!InExternalVariables.ContainsByPredicate([OuterName](const FRigVMExternalVariable& ExternalVariable)
+								{
+									return ExternalVariable.Name == OuterName;								
+								}))
+							{
+								static const FString UnmappedMessage = TEXT("Function referenced in @@ using external variable not found in current rig.");
+								Settings.ASTSettings.Report(EMessageSeverity::Error, ModelNode, UnmappedMessage);
+								bEncounteredGraphError = true;
+							}
+						}
+					}
+					else
+					{
+						static const FString UnmappedMessage = TEXT("Node @@ referencing function, but could not find compilation data.");
+						Settings.ASTSettings.Report(EMessageSeverity::Error, ModelNode, UnmappedMessage);
+						bEncounteredGraphError = true;
+					}
 				}
 			}
 
@@ -466,116 +616,6 @@
 		OutOperands->Add(Hash, Operand);
 	}
 
-	URigVMFunctionLibrary* FunctionLibrary = InGraphs[0]->GetDefaultFunctionLibrary();
-	if (CurrentCompilationFunction == nullptr)
-	{
-		ClearFunctionCompilationData();
-
-		// We need to compile the functions in order of dependency (leaf functions with no dependency first)
-		TArray<URigVMLibraryNode*> FunctionsToCompile;
-		if (FunctionLibrary)
-		{
-			FunctionsToCompile = FunctionLibrary->GetFunctions();
-		}
-
-		// Add public functions from other assets
-		// todo: this should be a check if the compilation data exists
-		{
-			TArray<URigVMNode*> Nodes = InGraphs[0]->GetNodes();
-			for (int32 i=0; i<Nodes.Num(); ++i)
-			{
-				if (URigVMFunctionReferenceNode* ReferenceNode = Cast<URigVMFunctionReferenceNode>(Nodes[i]))
-				{
-					FunctionsToCompile.AddUnique(ReferenceNode->GetReferencedNode());
-				}
-				if (URigVMCollapseNode* CollapseNode = Cast<URigVMCollapseNode>(Nodes[i]))
-				{
-					Nodes.Append(CollapseNode->GetContainedGraph()->GetNodes());
-				}
-			}
-		}
-		while (CompiledFunctions.Num() < FunctionsToCompile.Num())
-		{
-			for (int32 FunctionIndex=0; FunctionIndex<FunctionsToCompile.Num(); ++FunctionIndex)
-			{
-				if (URigVMLibraryNode* FunctionToCompile = FunctionsToCompile[FunctionIndex])
-				{
-					if (CompiledFunctions.Contains(FunctionToCompile))
-					{
-						continue;
-					}
-			
-					bool bCanBeCompiled = true;
-					TArray<URigVMNode*> Nodes = FunctionToCompile->GetContainedNodes();
-					for (int32 i=0; i<Nodes.Num() && bCanBeCompiled; ++i)
-					{
-						if (URigVMFunctionReferenceNode* FunctionReference = Cast<URigVMFunctionReferenceNode>(Nodes[i]))
-						{
-							if (const URigVMFunctionLibrary* ReferencedFunctionLibrary = FunctionReference->GetReferencedNode()->GetLibrary())
-							{
-								if (ReferencedFunctionLibrary == FunctionLibrary)
-								{
-									if (!CompiledFunctions.Contains(FunctionReference->GetReferencedNode()))
-									{
-										bCanBeCompiled = false;
-										break;
-									}
-								}
-								else
-								{
-									// Public function.
-									// todo Try to find its compilation data
-									bool bFoundCompilationData = true;
-									if (!CompiledFunctions.Contains(FunctionReference->GetReferencedNode()))
-									{
-										FunctionsToCompile.AddUnique(FunctionReference->GetReferencedNode());
-										bCanBeCompiled = false;
-										break;
-									}
-
-									if (!bFoundCompilationData)
-									{
-										static const FString FunctionCompilationErrorMessage = TEXT("Could not find compilation data for public function @@.");
-										Settings.ASTSettings.Report(EMessageSeverity::Error, FunctionReference, FunctionCompilationErrorMessage);
-										bEncounteredGraphError = true;
-										bCanBeCompiled = false;
-										break;
-									}
-								}
-							}
-						}
-
-						if (bCanBeCompiled)
-						{
-							if (URigVMLibraryNode* LibraryNode = Cast<URigVMLibraryNode>(Nodes[i]))
-							{
-								Nodes.Append(LibraryNode->GetContainedNodes());	
-							}
-						}
-					}
-
-					if (bCanBeCompiled)
-					{
-						if (!CompileFunction(FunctionToCompile, InController, UserData))
-						{
-							static const FString FunctionCompilationErrorMessage = TEXT("Function @@ failed to compile.");
-							Settings.ASTSettings.Report(EMessageSeverity::Error, FunctionToCompile, FunctionCompilationErrorMessage);
-							bEncounteredGraphError = true;
-							break;
-						}
-					}
-				}
-			}
-
-			FunctionsToCompile.RemoveAll([](const URigVMLibraryNode* Node) { return Node == nullptr; });
-			
-			if(bEncounteredGraphError)
-			{
-				return false;
-			}
-		}
-	}
-
 	FRigVMCompilerWorkData WorkData;
 
 	WorkData.AST = InAST;
@@ -773,14 +813,33 @@
 		Data.FunctionNames = WorkData.VM->FunctionNamesStorage;
 		Data.PropertyDescriptions = WorkData.PropertyDescriptions;
 		Data.PropertyPathDescriptions = WorkData.PropertyPathDescriptions;
+
+		// Only add used external registers to the function compilation data
+		FRigVMInstructionArray Instructions = Data.ByteCode.GetInstructions();
+		TSet<int32> UsedExternalVariableRegisters;
+		for (const FRigVMInstruction& Instruction : Instructions)
+		{
+			const FRigVMOperandArray OperandArray = Data.ByteCode.GetOperandsForOp(Instruction);
+			for (const FRigVMOperand& Operand : OperandArray)
+			{
+				if (Operand.GetMemoryType() == ERigVMMemoryType::External)
+				{
+					UsedExternalVariableRegisters.Add(Operand.GetRegisterIndex());					
+				}
+			}			
+		}
+
 		for (const TPair<FString, FRigVMOperand>& Pair : (*WorkData.PinPathToOperand))
 		{
 			static const FString VariablePrefix = TEXT("Variable::");
 			if (Pair.Key.StartsWith(VariablePrefix))
 			{
 				ensure(Pair.Value.GetMemoryType() == ERigVMMemoryType::External);
-				FString VariableName = Pair.Key.RightChop(VariablePrefix.Len());
-				Data.ExternalRegisterIndexToVariable.Add(Pair.Value.GetRegisterIndex(), *VariableName);
+				if (UsedExternalVariableRegisters.Contains(Pair.Value.GetRegisterIndex()))
+				{
+					FString VariableName = Pair.Key.RightChop(VariablePrefix.Len());
+					Data.ExternalRegisterIndexToVariable.Add(Pair.Value.GetRegisterIndex(), *VariableName);
+				}
 			}
 		}
 		CompiledFunctions.Add(CurrentCompilationFunction, Data);
