Control Rig: Replace cast link with cast nodes

diff --git a/Engine/Plugins/Animation/ControlRig/Source/ControlRigDeveloper/Private/ControlRigBlueprint.cpp b/Engine/Plugins/Animation/ControlRig/Source/ControlRigDeveloper/Private/ControlRigBlueprint.cpp
--- a/Engine/Plugins/Animation/ControlRig/Source/ControlRigDeveloper/Private/ControlRigBlueprint.cpp
+++ b/Engine/Plugins/Animation/ControlRig/Source/ControlRigDeveloper/Private/ControlRigBlueprint.cpp
@@ -900,6 +900,7 @@ void UControlRigBlueprint::PostLoad()
 		PatchBoundVariables();
 		PatchPropagateToChildren();
 		PatchParameterNodesOnLoad();
+		PatchLinksWithCast();
 
 #if WITH_EDITOR
 
@@ -4543,6 +4544,64 @@ void UControlRigBlueprint::PatchTemplateNodesWithPreferredPermutation()
 #endif
 }
 
+void UControlRigBlueprint::PatchLinksWithCast()
+{
+#if WITH_EDITOR
+
+	{
+		TGuardValue<bool> DisableAutoCompile(bAutoRecompileVM, false);
+
+		// find all links containing a cast
+		TArray<TTuple<URigVMGraph*,TWeakObjectPtr<URigVMLink>,FString,FString>> LinksWithCast;
+		for (URigVMGraph* Graph : GetAllModels())
+		{
+			for(URigVMLink* Link : Graph->GetLinks())
+			{
+				const URigVMPin* SourcePin = Link->GetSourcePin();
+				const URigVMPin* TargetPin = Link->GetTargetPin();
+				const TRigVMTypeIndex SourceTypeIndex = SourcePin->GetTypeIndex();
+				const TRigVMTypeIndex TargetTypeIndex = TargetPin->GetTypeIndex();
+
+				if(SourceTypeIndex != TargetTypeIndex)
+				{
+					if(!FRigVMRegistry::Get().CanMatchTypes(SourceTypeIndex, TargetTypeIndex, true))
+					{
+						LinksWithCast.Emplace(Graph, TWeakObjectPtr<URigVMLink>(Link), SourcePin->GetPinPath(), TargetPin->GetPinPath());
+					}
+				}
+			}
+		}
+
+		// remove all of those links
+		for(const auto& Tuple : LinksWithCast)
+		{
+			URigVMController* Controller = GetController(Tuple.Get<0>());
+
+			if(URigVMLink* Link = Tuple.Get<1>().Get())
+			{
+				// the link may be detached, attach it first so that removal works.
+				const URigVMPin* SourcePin = Link->GetSourcePin();
+				URigVMPin* TargetPin = Link->GetTargetPin();
+				if(!SourcePin->IsLinkedTo(TargetPin))
+				{
+					const TArray<URigVMLink*> LinksToReattach = {Link};
+					Controller->ReattachLinksToPinObjects(true, &LinksToReattach, true, false, false);
+				}
+			}
+			
+			Controller->BreakLink(Tuple.Get<2>(), Tuple.Get<3>(), false);
+
+			// notify the user that the link has been broken.
+			UE_LOG(LogControlRigDeveloper, Warning,
+				TEXT("A link was removed in %s (%s) - it contained different types on source and target pin (former cast link?)."),
+				*Controller->GetGraph()->GetNodePath(),
+				*URigVMLink::GetPinPathRepresentation(Tuple.Get<2>(), Tuple.Get<3>())
+			);
+		}
+	}
+#endif
+}
+
 void UControlRigBlueprint::PropagatePoseFromInstanceToBP(UControlRig* InControlRig)
 {
 	check(InControlRig);
diff --git a/Engine/Plugins/Animation/ControlRig/Source/ControlRigDeveloper/Private/Graph/ControlRigGraphNode.cpp b/Engine/Plugins/Animation/ControlRig/Source/ControlRigDeveloper/Private/Graph/ControlRigGraphNode.cpp
--- a/Engine/Plugins/Animation/ControlRig/Source/ControlRigDeveloper/Private/Graph/ControlRigGraphNode.cpp
+++ b/Engine/Plugins/Animation/ControlRig/Source/ControlRigDeveloper/Private/Graph/ControlRigGraphNode.cpp
@@ -41,6 +41,8 @@
 
 #define LOCTEXT_NAMESPACE "ControlRigGraphNode"
 
+TAutoConsoleVariable<bool> CVarControlRigDisableCompactNodes(TEXT("ControlRig.Graph.DisableCompactNodes"), false, TEXT("When true all nodes are going to be drawn as full nodes."));
+
 UControlRigGraphNode::UControlRigGraphNode()
 : Dimensions(0.0f, 0.0f)
 , NodeTitle(FText::GetEmpty())
@@ -259,6 +261,8 @@ void UControlRigGraphNode::ReconstructNode_Internal(bool bForce)
 	
 	RewireOldPinsToNewPins(OldPins, Pins);
 
+	DrawAsCompactNodeCache.Reset();
+
 	// Let subclasses do any additional work
 	PostReconstructNode();
 
@@ -897,6 +901,15 @@ bool UControlRigGraphNode::ModelPinAdded_Internal(const URigVMPin* InModelPin)
 {
 	bool bResult = false;
 
+	// conversion nodes don't show sub pins
+	if(!InModelPin->IsRootPin())
+	{
+		if(DrawAsCompactNode())
+		{
+			return false;
+		}
+	}
+
 	UEdGraphPin* InputParentPin = FindGraphPinFromModelPin(InModelPin->GetParentPin(), true);
 	UEdGraphPin* OutputParentPin = FindGraphPinFromModelPin(InModelPin->GetParentPin(), false);
 		
@@ -939,6 +952,60 @@ bool UControlRigGraphNode::ModelPinRemoved(const URigVMPin* InModelPin)
 	return bResult;
 }
 
+bool UControlRigGraphNode::DrawAsCompactNode() const
+{
+	if(CVarControlRigDisableCompactNodes.GetValueOnAnyThread())
+	{
+		return false;
+	}
+	
+	if(!DrawAsCompactNodeCache.IsSet())
+	{
+		DrawAsCompactNodeCache = false;
+		if(const URigVMTemplateNode* TemplateModelNode = Cast<URigVMTemplateNode>(GetModelNode()))
+		{
+			if(TemplateModelNode->GetNotation() == RigVMTypeUtils::GetCastTemplateNotation())
+			{
+				DrawAsCompactNodeCache = true;
+				
+				// if the node has links on any subpin - we can't draw it as a compact node.
+				const TArray<URigVMLink*> Links = TemplateModelNode->GetLinks();
+				if(Links.ContainsByPredicate([TemplateModelNode](const URigVMLink* Link)
+				{
+					if(const URigVMPin* SourcePin = Link->GetSourcePin())
+					{
+						if(SourcePin->GetNode() == TemplateModelNode)
+						{
+							return !SourcePin->IsRootPin();
+						}
+					}
+					if(const URigVMPin* TargetPin = Link->GetTargetPin())
+					{
+						if(TargetPin->GetNode() == TemplateModelNode)
+						{
+							return !TargetPin->IsRootPin();
+						}
+					}
+					return false;
+				}))
+				{
+					DrawAsCompactNodeCache = false;
+				}
+
+				// if the node has any injected nodes - we can't draw it as compact
+				if(TemplateModelNode->GetPins().ContainsByPredicate([](const URigVMPin* Pin)
+				{
+					return Pin->HasInjectedNodes();
+				}))
+				{
+					DrawAsCompactNodeCache = false;
+				}
+			}
+		}
+	}
+	return DrawAsCompactNodeCache.GetValue();
+}
+
 bool UControlRigGraphNode::ModelPinRemoved_Internal(const URigVMPin* InModelPin)
 {
 	auto RemoveGraphPin = [this](const URigVMPin* InModelPin, bool bAsInput)
diff --git a/Engine/Plugins/Animation/ControlRig/Source/ControlRigDeveloper/Private/Graph/ControlRigGraphSchema.cpp b/Engine/Plugins/Animation/ControlRig/Source/ControlRigDeveloper/Private/Graph/ControlRigGraphSchema.cpp
--- a/Engine/Plugins/Animation/ControlRig/Source/ControlRigDeveloper/Private/Graph/ControlRigGraphSchema.cpp
+++ b/Engine/Plugins/Animation/ControlRig/Source/ControlRigDeveloper/Private/Graph/ControlRigGraphSchema.cpp
@@ -714,7 +714,8 @@ bool UControlRigGraphSchema::TryCreateConnection(UEdGraphPin* PinA, UEdGraphPin*
 			}
 #endif
 
-			return Controller->AddLink(PinA->GetName(), PinB->GetName(), true, true, UserLinkDirection);
+			const bool bCreateCastNode = !FSlateApplication::Get().GetModifierKeys().IsAltDown();
+			return Controller->AddLink(PinA->GetName(), PinB->GetName(), true, true, UserLinkDirection, bCreateCastNode);
 		}
 	}
 	return false;
@@ -795,8 +796,9 @@ const FPinConnectionResponse UControlRigGraphSchema::CanCreateConnection(const U
 
 		const FRigVMByteCode* ByteCode = RigNodeA->GetController()->GetCurrentByteCode();
 
+		const bool bEnableTypeCasting = !FSlateApplication::Get().GetModifierKeys().IsAltDown();
 		FString FailureReason;
-		bool bResult = RigNodeA->GetModel()->CanLink(PinA, PinB, &FailureReason, ByteCode, UserLinkDirection);
+		bool bResult = RigNodeA->GetModel()->CanLink(PinA, PinB, &FailureReason, ByteCode, UserLinkDirection, bEnableTypeCasting);
 		if (!bResult)
 		{
 			return FPinConnectionResponse(CONNECT_RESPONSE_DISALLOW, FText::FromString(FailureReason));
@@ -1170,7 +1172,20 @@ FConnectionDrawingPolicy* UControlRigGraphSchema::CreateConnectionDrawingPolicy(
 bool UControlRigGraphSchema::ShouldHidePinDefaultValue(UEdGraphPin* Pin) const
 {
 	// we should hide default values if any of our parents are connected
-	return HasParentConnection_Recursive(Pin);
+	if(HasParentConnection_Recursive(Pin))
+	{
+		return true;
+	}
+
+	if(UControlRigGraphNode* RigGraphNode = Cast<UControlRigGraphNode>(Pin->GetOwningNode()))
+	{
+		if(RigGraphNode->DrawAsCompactNode())
+		{
+			return true;
+		}
+	}
+
+	return false;
 }
 
 bool UControlRigGraphSchema::IsPinBeingWatched(UEdGraphPin const* Pin) const
diff --git a/Engine/Plugins/Animation/ControlRig/Source/ControlRigDeveloper/Public/ControlRigBlueprint.h b/Engine/Plugins/Animation/ControlRig/Source/ControlRigDeveloper/Public/ControlRigBlueprint.h
--- a/Engine/Plugins/Animation/ControlRig/Source/ControlRigDeveloper/Public/ControlRigBlueprint.h
+++ b/Engine/Plugins/Animation/ControlRig/Source/ControlRigDeveloper/Public/ControlRigBlueprint.h
@@ -605,6 +605,7 @@ private:
 	void PatchPropagateToChildren();
 	void PatchParameterNodesOnLoad();
 	void PatchTemplateNodesWithPreferredPermutation();
+	void PatchLinksWithCast();
 
 	TMap<FName, int32> AddedMemberVariableMap;
 	TArray<FBPVariableDescription> LastNewVariables;
diff --git a/Engine/Plugins/Animation/ControlRig/Source/ControlRigDeveloper/Public/Graph/ControlRigGraphNode.h b/Engine/Plugins/Animation/ControlRig/Source/ControlRigDeveloper/Public/Graph/ControlRigGraphNode.h
--- a/Engine/Plugins/Animation/ControlRig/Source/ControlRigDeveloper/Public/Graph/ControlRigGraphNode.h
+++ b/Engine/Plugins/Animation/ControlRig/Source/ControlRigDeveloper/Public/Graph/ControlRigGraphNode.h
@@ -200,6 +200,9 @@ public:
 	/** Called when a model pin is being removed */
 	bool ModelPinRemoved(const URigVMPin* InModelPin);
 
+	/** Returns true if this node is relying on the cast template */
+	bool DrawAsCompactNode() const;
+
 protected:
 
 	FLinearColor GetNodeProfilingColor() const;
@@ -266,6 +269,7 @@ private:
 	FNodePinExpansionChanged NodePinExpansionChanged;
 
 	mutable const FRigVMTemplate* CachedTemplate;
+	mutable TOptional<bool> DrawAsCompactNodeCache;
 
 	friend class SControlRigGraphNode;
 	friend class FControlRigArgumentLayout;
diff --git a/Engine/Plugins/Animation/ControlRig/Source/ControlRigEditor/Private/Graph/ControlRigConnectionDrawingPolicy.cpp b/Engine/Plugins/Animation/ControlRig/Source/ControlRigEditor/Private/Graph/ControlRigConnectionDrawingPolicy.cpp
--- a/Engine/Plugins/Animation/ControlRig/Source/ControlRigEditor/Private/Graph/ControlRigConnectionDrawingPolicy.cpp
+++ b/Engine/Plugins/Animation/ControlRig/Source/ControlRigEditor/Private/Graph/ControlRigConnectionDrawingPolicy.cpp
@@ -112,45 +112,7 @@ void FControlRigConnectionDrawingPolicy::DrawPinGeometries(TMap<TSharedRef<SWidg
 					}
 					else
 					{
-						const float StartFudgeX = 4.0f;
-						const float EndFudgeX = 4.0f;
-						FVector2D StartPoint = FGeometryHelper::VerticalMiddleRightOf(LinkStartWidgetGeometry->Geometry) - FVector2D(StartFudgeX, 0.0f);
-						FVector2D EndPoint = FGeometryHelper::VerticalMiddleLeftOf(LinkEndWidgetGeometry->Geometry) - FVector2D(ArrowRadius.X - EndFudgeX, 0);
-
-						if(Params.bUserFlag1) // indicate a cast using a mid point image
-						{
-							FConnectionParams OtherParams;
-							DetermineWiringStyle(TargetPin, ThePin, OtherParams);
-
-							if(!OtherParams.WireColor.Equals(Params.WireColor, 0.001f))
-							{
-								const FVector2D MidPoint = (StartPoint + EndPoint) * 0.5;
-								const FVector2D OriginalSplineTangent = FVector2D(100.f, 0);
-								const FVector2D InnerStartPoint = StartPoint + OriginalSplineTangent;
-								const FVector2D InnerEndPoint = EndPoint - OriginalSplineTangent;
-
-								const FVector2D DeltaPos = InnerEndPoint - InnerStartPoint;
-								const double TangentLength = FMath::Min<float>(FMath::Abs<float>(DeltaPos.X), Settings->ForwardSplineHorizontalDeltaRange);
-								OtherParams.StartTangent = DeltaPos.GetSafeNormal() * TangentLength;
-
-								DrawSplineWithArrow(MidPoint, EndPoint, OtherParams);
-
-								EndPoint = MidPoint;
-								Params.EndTangent = OtherParams.StartTangent;
-
-								// Draw the icon
-								FSlateDrawElement::MakeBox(
-									DrawElementsList,
-									ArrowLayerID,
-									FPaintGeometry(MidPoint - ArrowRadius, CastImage->ImageSize * ZoomFactor, ZoomFactor),
-									CastImage,
-									ESlateDrawEffect::None,
-									OtherParams.WireColor
-								);
-							}
-						}
-						
-						DrawSplineWithArrow(StartPoint, EndPoint, Params);
+						DrawSplineWithArrow(LinkStartWidgetGeometry->Geometry, LinkEndWidgetGeometry->Geometry, Params);
 					}
 				}
 			}
@@ -369,8 +331,6 @@ void FControlRigConnectionDrawingPolicy::DetermineWiringStyle(UEdGraphPin* Outpu
 				}
 			}
 
-			Params.bUserFlag1 = false;
-
 			if (bVisited)
 			{
 				//Params.bDrawBubbles = true;
@@ -380,17 +340,6 @@ void FControlRigConnectionDrawingPolicy::DetermineWiringStyle(UEdGraphPin* Outpu
 			{
 				Params.WireColor = Params.WireColor * 0.5f;
 			}
-
-			if(!UseLowDetailConnections() && !bInjectionIsSelected)
-			{
-				const TRigVMTypeIndex OutputTypeIndex = OutputModelPin->GetTypeIndex();
-				const TRigVMTypeIndex InputTypeIndex = InputModelPin->GetTypeIndex();
-				if(OutputTypeIndex != InputTypeIndex
-					&& !FRigVMRegistry::Get().CanMatchTypes(OutputTypeIndex, InputTypeIndex, true))
-				{
-					Params.bUserFlag1 = true;
-				}
-			}
 		}
 	}
 }
diff --git a/Engine/Plugins/Animation/ControlRig/Source/ControlRigEditor/Private/Graph/SControlRigGraphNode.cpp b/Engine/Plugins/Animation/ControlRig/Source/ControlRigEditor/Private/Graph/SControlRigGraphNode.cpp
--- a/Engine/Plugins/Animation/ControlRig/Source/ControlRigEditor/Private/Graph/SControlRigGraphNode.cpp
+++ b/Engine/Plugins/Animation/ControlRig/Source/ControlRigEditor/Private/Graph/SControlRigGraphNode.cpp
@@ -8,6 +8,7 @@
 #include "Graph/ControlRigGraphSchema.h"
 #include "Widgets/SBoxPanel.h"
 #include "Widgets/Layout/SSpacer.h"
+#include "Widgets/Layout/SBox.h"
 #include "Widgets/Input/SButton.h"
 #include "SLevelOfDetailBranchNode.h"
 #include "SGraphPanel.h"
@@ -135,6 +136,105 @@ bool SControlRigGraphNode::UseLowDetailPinNames() const
 	return false;
 }
 
+void SControlRigGraphNode::UpdateGraphNode()
+{
+	if(UControlRigGraphNode* RigGraphNode = Cast<UControlRigGraphNode>(GraphNode))
+	{
+		if(RigGraphNode->DrawAsCompactNode())
+		{
+			UpdateCompactNode();
+			return;
+		}
+	}
+	UpdateStandardNode();
+}
+
+void SControlRigGraphNode::UpdateStandardNode()
+{
+	// call super
+	SGraphNode::UpdateGraphNode();
+}
+
+void SControlRigGraphNode::UpdateCompactNode()
+{
+	InputPins.Empty();
+	OutputPins.Empty();
+	
+	// Reset variables that are going to be exposed, in case we are refreshing an already setup node.
+	RightNodeBox.Reset();
+	LeftNodeBox.Reset();
+
+	SetupErrorReporting();
+	
+	TSharedRef<SOverlay> NodeOverlay = SNew(SOverlay);
+	
+	// add optional node specific widget to the overlay:
+	TSharedPtr<SWidget> OverlayWidget = GraphNode->CreateNodeImage();
+	if(OverlayWidget.IsValid())
+	{
+		NodeOverlay->AddSlot()
+		.HAlign(HAlign_Center)
+		.VAlign(VAlign_Center)
+		[
+			SNew( SBox )
+			.WidthOverride( 70.f )
+			.HeightOverride( 70.f )
+			[
+				OverlayWidget.ToSharedRef()
+			]
+		];
+	}
+
+	TSharedRef<SVerticalBox> InnerVerticalBox =
+	SNew(SVerticalBox)
+	+SVerticalBox::Slot()
+	[
+		// NODE CONTENT AREA
+		SNew( SOverlay)
+		+SOverlay::Slot()
+		[
+			SNew(SImage)
+			.Image( FAppStyle::GetBrush("Graph.VarNode.Body") )
+		]
+		+ SOverlay::Slot()
+		[
+			SNew(SImage)
+			.Image( FAppStyle::GetBrush("Graph.VarNode.Gloss") )
+		]
+		+SOverlay::Slot()
+		.Padding( FMargin(0,3) )
+		[
+			NodeOverlay
+		]
+	];
+
+	NodeOverlay->AddSlot()
+	.HAlign(HAlign_Left)
+	.VAlign(VAlign_Center)
+	.Padding(/* left */ 0.f, 0.f, 0.f, /* bottom */ 0.f)
+	[
+		CreateNodeContentArea()
+	];
+
+	this->ContentScale.Bind( this, &SGraphNode::GetContentScale );
+
+	InnerVerticalBox->AddSlot()
+		.AutoHeight()
+		.Padding( FMargin(5.0f, 1.0f) )
+		[
+			ErrorReporting->AsWidget()
+		];
+
+	this->GetOrAddSlot( ENodeZone::Center )
+	.HAlign(HAlign_Center)
+	.VAlign(VAlign_Center)
+	[
+		InnerVerticalBox
+	];
+
+	CreatePinWidgets();
+}
+
 void SControlRigGraphNode::CreateAggregateAddPinButton()
 {
 	if (LeftNodeBox.IsValid() && ModelNode.IsValid() && (ModelNode->IsAggregate() || ModelNode->IsA<URigVMAggregateNode>()))
@@ -373,6 +473,14 @@ void SControlRigGraphNode::CreateStandardPinWidget(UEdGraphPin* CurPin)
 			check(NewPin.IsValid());
 			PinsToKeep.Remove(CurPin);
 		}
+		
+		if(const UControlRigGraphNode* RigGraphNode = Cast<UControlRigGraphNode>(GraphNode))
+		{
+			if(RigGraphNode->DrawAsCompactNode())
+			{
+				NewPin->SetShowLabel(false);
+			}
+		}
 
 		AddPin(NewPin.ToSharedRef());
 	}
@@ -1279,13 +1387,19 @@ void SControlRigGraphNode::UpdatePinTreeView()
 	RootModelPins.Append(ModelNode->GetOrphanedPins());
 	TArray<URigVMPin*> ModelPins;
 
+	const bool bSupportSubPins = !RigGraphNode->DrawAsCompactNode();
+
 	// sort model pins
 	// a) execute IOs, b) IO pins, c) input / visible pins, d) output pins
 	struct Local
 	{
-		static void VisitPinRecursively(URigVMPin* InPin, TArray<URigVMPin*>& OutPins)
+		static void VisitPinRecursively(URigVMPin* InPin, TArray<URigVMPin*>& OutPins, bool bSupportSubPins)
 		{
 			OutPins.Add(InPin);
+			if(!bSupportSubPins)
+			{
+				return;
+			}
 
 			if (InPin->GetCPPType() == TEXT("FRotator"))
 			{
@@ -1301,7 +1415,7 @@ void SControlRigGraphNode::UpdatePinTreeView()
 			{				
 				for (URigVMPin* SubPin : InPin->GetSubPins())
 				{
-					VisitPinRecursively(SubPin, OutPins);
+					VisitPinRecursively(SubPin, OutPins, bSupportSubPins);
 				}
 			}
 		}
@@ -1317,7 +1431,7 @@ void SControlRigGraphNode::UpdatePinTreeView()
 				{
 					if(RootPin->IsExecuteContext() && RootPin->GetDirection() == ERigVMPinDirection::IO)
 					{
-						Local::VisitPinRecursively(RootPin, ModelPins);
+						Local::VisitPinRecursively(RootPin, ModelPins, bSupportSubPins);
 					}
 					break;
 				}
@@ -1325,7 +1439,7 @@ void SControlRigGraphNode::UpdatePinTreeView()
 				{
 					if(!RootPin->IsExecuteContext() && RootPin->GetDirection() == ERigVMPinDirection::IO)
 					{
-						Local::VisitPinRecursively(RootPin, ModelPins);
+						Local::VisitPinRecursively(RootPin, ModelPins, bSupportSubPins);
 					}
 					break;
 				}
@@ -1333,7 +1447,7 @@ void SControlRigGraphNode::UpdatePinTreeView()
 				{
 					if(RootPin->GetDirection() == ERigVMPinDirection::Input || RootPin->GetDirection() == ERigVMPinDirection::Visible)
 					{
-						Local::VisitPinRecursively(RootPin, ModelPins);
+						Local::VisitPinRecursively(RootPin, ModelPins, bSupportSubPins);
 					}
 					break;
 				}
@@ -1342,7 +1456,7 @@ void SControlRigGraphNode::UpdatePinTreeView()
 				{
 					if(RootPin->GetDirection() == ERigVMPinDirection::Output)
 					{
-						Local::VisitPinRecursively(RootPin, ModelPins);
+						Local::VisitPinRecursively(RootPin, ModelPins, bSupportSubPins);
 					}
 					break;
 				}
@@ -1364,6 +1478,12 @@ void SControlRigGraphNode::UpdatePinTreeView()
 		PinInfo.bExpanded = ModelPin->IsExpanded();
 		PinInfo.ModelPinPath = ModelPin->GetPinPath();
 		PinInfo.bAutoHeight = false;
+
+		if(!bSupportSubPins)
+		{
+			PinInfo.bHasChildren = false;
+			PinInfo.bIsContainer = false;
+		}
 		
 		const bool bAskSchemaForEdition = RigSchema && ModelPin->IsStruct() && !ModelPin->IsBoundToVariable();
 		PinInfo.bHideInputWidget = (!ModelPin->IsBoundToVariable()) && PinInfo.bIsContainer;
@@ -1557,48 +1677,120 @@ void SControlRigGraphNode::UpdatePinTreeView()
 		];
 	};
 
+	struct FPinRowInfo
+	{
+		int32 InputPinInfo;
+		int32 OutputPinInfo;
+		
+		FPinRowInfo()
+			: InputPinInfo(INDEX_NONE)
+			, OutputPinInfo(INDEX_NONE)
+		{}
+	};
+
+	TArray<FPinRowInfo> Rows;
+	Rows.Reserve(PinInfos.Num());
+
+	// build a tree with the pin infos in order as before
 	for(const FPinInfo& PinInfo : PinInfos)
 	{
+		FPinRowInfo Row;
 		if(PinInfo.InputPinWidget.IsValid())
 		{
-			if(PinInfo.bHideInputWidget)
+			Row.InputPinInfo = PinInfo.Index;
+		}
+		if(PinInfo.OutputPinWidget.IsValid())
+		{
+			Row.OutputPinInfo = PinInfo.Index;
+		}
+		Rows.Add(Row);
+	}
+
+	// compact the rows - re-use input rows for outputs
+	// if their sub-pins count matches
+	if(RigGraphNode->DrawAsCompactNode())
+	{
+		for(int32 RowIndex = 0; RowIndex < Rows.Num();)
+		{
+			if(Rows[RowIndex].OutputPinInfo == INDEX_NONE)
+			{
+				RowIndex++;
+				continue;
+			}
+
+			bool bFoundMatch = false;
+			for(int32 InputRowIndex = 0; InputRowIndex < RowIndex; InputRowIndex++)
+			{
+				if(Rows[InputRowIndex].OutputPinInfo == INDEX_NONE)
+				{
+					const FPinInfo& OutputPinInfo = PinInfos[Rows[RowIndex].OutputPinInfo];
+					const FPinInfo& InputPinInfo = PinInfos[Rows[InputRowIndex].InputPinInfo];
+					if(OutputPinInfo.bHasChildren == false &&
+						InputPinInfo.bHasChildren == false)
+					{
+						Rows[InputRowIndex].OutputPinInfo = Rows[RowIndex].OutputPinInfo;
+						bFoundMatch = true;
+						break;
+					}
+				}
+			}
+
+			if(bFoundMatch)
+			{
+				Rows.RemoveAt(RowIndex);
+			}
+			else
+			{
+				RowIndex++;
+			}
+		}
+	}
+	
+	for(const FPinRowInfo& Row : Rows)
+	{
+		const FPinInfo& InputPinInfo = PinInfos[Row.InputPinInfo == INDEX_NONE ? Row.OutputPinInfo : Row.InputPinInfo];
+		const FPinInfo& OutputPinInfo = PinInfos[Row.OutputPinInfo == INDEX_NONE ? Row.InputPinInfo : Row.OutputPinInfo];
+
+		if(InputPinInfo.InputPinWidget.IsValid())
+		{
+			if(InputPinInfo.bHideInputWidget)
 			{
-				if(PinInfo.InputPinWidget->GetValueWidget() != SNullWidget::NullWidget)
+				if(InputPinInfo.InputPinWidget->GetValueWidget() != SNullWidget::NullWidget)
 				{
-					PinInfo.InputPinWidget->GetValueWidget()->SetVisibility(EVisibility::Collapsed);
+					InputPinInfo.InputPinWidget->GetValueWidget()->SetVisibility(EVisibility::Collapsed);
 				}
 			}
 				
 			// input pins
-			if(!PinInfo.OutputPinWidget.IsValid())
+			if(!OutputPinInfo.OutputPinWidget.IsValid())
 			{
 				TSharedPtr<SHorizontalBox> SlotLayout;
 				SHorizontalBox::FSlot* FirstSlot = nullptr;
 
-				const float MyEmptySidePadding = PinInfo.bHideInputWidget ? EmptySidePadding : 0.f; 
+				const float MyEmptySidePadding = InputPinInfo.bHideInputWidget ? EmptySidePadding : 0.f; 
 				
 				LeftNodeBox->AddSlot()
                 .HAlign(HAlign_Fill)
 				.VAlign(VAlign_Center)
 				.AutoHeight()
-				.MaxHeight(PinInfo.bAutoHeight ? TAttribute<float>() : MaxHeight)
+				.MaxHeight(InputPinInfo.bAutoHeight ? TAttribute<float>() : MaxHeight)
                 [
                     SAssignNew(SlotLayout, SHorizontalBox)
-                    .Visibility(this, &SControlRigGraphNode::GetPinVisibility, PinInfo.Index)
+                    .Visibility(this, &SControlRigGraphNode::GetPinVisibility, InputPinInfo.Index)
 		            
                     +SHorizontalBox::Slot()
                     .Expose(FirstSlot)
                     .FillWidth(1.f)
                     .HAlign(HAlign_Left)
-                    .Padding(PinWidgetSidePadding, TopPadding, PinInfo.bIsContainer ? 0.f : MyEmptySidePadding, 0.f)
+                    .Padding(PinWidgetSidePadding, TopPadding, InputPinInfo.bIsContainer ? 0.f : MyEmptySidePadding, 0.f)
                     [
-                        PinInfo.InputPinWidget.ToSharedRef()
+                        InputPinInfo.InputPinWidget.ToSharedRef()
                     ]
                 ];
 
-				if(PinInfo.bIsContainer)
+				if(InputPinInfo.bIsContainer)
 				{
-					URigVMPin* ModelPin = ModelNode->GetGraph()->FindPin(PinInfo.ModelPinPath);
+					URigVMPin* ModelPin = ModelNode->GetGraph()->FindPin(InputPinInfo.ModelPinPath);
 					if(ModelPin)
 					{
 						// make sure to minimize the width of the label
@@ -1613,16 +1805,16 @@ void SControlRigGraphNode::UpdatePinTreeView()
 				TSharedPtr<SHorizontalBox> SlotLayout;
 				SHorizontalBox::FSlot* FirstSlot = nullptr;
 
-				PinInfo.OutputPinWidget->SetShowLabel(false);
+				OutputPinInfo.OutputPinWidget->SetShowLabel(false);
 			
 				LeftNodeBox->AddSlot()
                 .HAlign(HAlign_Fill)
 				.VAlign(VAlign_Center)
                 .AutoHeight()
-				.MaxHeight(PinInfo.bAutoHeight ? TAttribute<float>() : MaxHeight)
+				.MaxHeight(InputPinInfo.bAutoHeight ? TAttribute<float>() : MaxHeight)
                 [
                     SAssignNew(SlotLayout, SHorizontalBox)
-                    .Visibility(this, &SControlRigGraphNode::GetPinVisibility, PinInfo.Index)
+                    .Visibility(this, &SControlRigGraphNode::GetPinVisibility, OutputPinInfo.Index)
 
                     +SHorizontalBox::Slot()
 					.Expose(FirstSlot)
@@ -1631,13 +1823,13 @@ void SControlRigGraphNode::UpdatePinTreeView()
                     .VAlign(VAlign_Center)
                     .Padding(PinWidgetSidePadding, TopPadding, 0.f, 0.f)
                     [
-                        PinInfo.InputPinWidget.ToSharedRef()
+                        InputPinInfo.InputPinWidget.ToSharedRef()
                     ]
                 ];
 
-				if(PinInfo.bIsContainer)
+				if(InputPinInfo.bIsContainer)
 				{
-					URigVMPin* ModelPin = ModelNode->GetGraph()->FindPin(PinInfo.ModelPinPath);
+					URigVMPin* ModelPin = ModelNode->GetGraph()->FindPin(InputPinInfo.ModelPinPath);
 					if(ModelPin)
 					{
 						// make sure to minimize the width of the label
@@ -1652,21 +1844,21 @@ void SControlRigGraphNode::UpdatePinTreeView()
 				.VAlign(VAlign_Center)
 				.Padding(0.f, TopPadding, PinWidgetSidePadding, 0.f)
 				[
-					PinInfo.OutputPinWidget.ToSharedRef()
+					OutputPinInfo.OutputPinWidget.ToSharedRef()
 				];
 			}
 		}
 		// output pins
-		else if(PinInfo.OutputPinWidget.IsValid())
+		else if(OutputPinInfo.OutputPinWidget.IsValid())
 		{
 			LeftNodeBox->AddSlot()
             .HAlign(HAlign_Fill)
 			.VAlign(VAlign_Center)
             .AutoHeight()
-			.MaxHeight(PinInfo.bAutoHeight ? TAttribute<float>() : MaxHeight)
+			.MaxHeight(OutputPinInfo.bAutoHeight ? TAttribute<float>() : MaxHeight)
             [
             	SNew(SHorizontalBox)
-	            .Visibility(this, &SControlRigGraphNode::GetPinVisibility, PinInfo.Index)
+	            .Visibility(this, &SControlRigGraphNode::GetPinVisibility, OutputPinInfo.Index)
 	            
 				+SHorizontalBox::Slot()
 				.FillWidth(1.f)
@@ -1674,7 +1866,7 @@ void SControlRigGraphNode::UpdatePinTreeView()
 				.VAlign(VAlign_Center)
 	            .Padding(EmptySidePadding, TopPadding, PinWidgetSidePadding, 0.f)
 				[
-	                PinInfo.OutputPinWidget.ToSharedRef()
+	                OutputPinInfo.OutputPinWidget.ToSharedRef()
 				]
             ];
 		}
diff --git a/Engine/Plugins/Animation/ControlRig/Source/ControlRigEditor/Private/Graph/SControlRigGraphNode.h b/Engine/Plugins/Animation/ControlRig/Source/ControlRigEditor/Private/Graph/SControlRigGraphNode.h
--- a/Engine/Plugins/Animation/ControlRig/Source/ControlRigEditor/Private/Graph/SControlRigGraphNode.h
+++ b/Engine/Plugins/Animation/ControlRig/Source/ControlRigEditor/Private/Graph/SControlRigGraphNode.h
@@ -53,6 +53,9 @@ public:
 	virtual void Tick(const FGeometry& AllottedGeometry, const double InCurrentTime, const float InDeltaTime) override;
 	virtual bool IsHidingPinWidgets() const override { return UseLowDetailNodeContent(); }
 	virtual bool UseLowDetailPinNames() const override;
+	virtual void UpdateGraphNode() override;
+	void UpdateStandardNode();
+	void UpdateCompactNode();
 
 	void CreateAggregateAddPinButton();
 
diff --git a/Engine/Source/Developer/RigVMDeveloper/Private/RigVMCompiler/RigVMAST.cpp b/Engine/Source/Developer/RigVMDeveloper/Private/RigVMCompiler/RigVMAST.cpp
--- a/Engine/Source/Developer/RigVMDeveloper/Private/RigVMCompiler/RigVMAST.cpp
+++ b/Engine/Source/Developer/RigVMDeveloper/Private/RigVMCompiler/RigVMAST.cpp
@@ -3594,6 +3594,7 @@ FString FRigVMParserAST::GetLinkAsString(const FRigVMASTLinkDescription& InLink)
 	static const FString EmptyString;
 	static const FString PeriodString = TEXT(".");
 
-	return FString::Printf(TEXT("%s -> %s%s%s"), *SourcePin->GetPinPath(), *TargetPin->GetPinPath(),
-		*(InLink.SegmentPath.IsEmpty() ? EmptyString : PeriodString), *InLink.SegmentPath);
+	return URigVMLink::GetPinPathRepresentation(SourcePin->GetPinPath(), 
+		 FString::Printf(TEXT("%s%s%s"), *TargetPin->GetPinPath(),
+			*(InLink.SegmentPath.IsEmpty() ? EmptyString : PeriodString), *InLink.SegmentPath));
 }
diff --git a/Engine/Source/Developer/RigVMDeveloper/Private/RigVMCompiler/RigVMCompiler.cpp b/Engine/Source/Developer/RigVMDeveloper/Private/RigVMCompiler/RigVMCompiler.cpp
--- a/Engine/Source/Developer/RigVMDeveloper/Private/RigVMCompiler/RigVMCompiler.cpp
+++ b/Engine/Source/Developer/RigVMDeveloper/Private/RigVMCompiler/RigVMCompiler.cpp
@@ -2473,7 +2473,7 @@ void URigVMCompiler::AddCopyOperator(const FRigVMCopyOp& InOp, const FRigVMAssig
 			if(CastFunction == nullptr)
 			{
 				const FRigVMRegistry& Registry = FRigVMRegistry::Get();
-				static constexpr TCHAR MissingCastMessage[] = TEXT("Cast (%s -> %s) for Node @@ not found.");
+				static constexpr TCHAR MissingCastMessage[] = TEXT("Cast (%s to %s) for Node @@ not found.");
 				const FString& SourceCPPType = Registry.GetType(SourceTypeIndex).CPPType.ToString();
 				const FString& TargetCPPType = Registry.GetType(TargetTypeIndex).CPPType.ToString();
 				Settings.Report(EMessageSeverity::Error, InAssignExpr->GetTargetPin()->GetNode(),
diff --git a/Engine/Source/Developer/RigVMDeveloper/Private/RigVMDeveloperTypeUtils.cpp b/Engine/Source/Developer/RigVMDeveloper/Private/RigVMDeveloperTypeUtils.cpp
--- a/Engine/Source/Developer/RigVMDeveloper/Private/RigVMDeveloperTypeUtils.cpp
+++ b/Engine/Source/Developer/RigVMDeveloper/Private/RigVMDeveloperTypeUtils.cpp
@@ -867,4 +867,19 @@ const FRigVMFunction* RigVMTypeUtils::GetCastForTypeIndices(const TRigVMTypeInde
 	return nullptr;
 }
 
+const FName& RigVMTypeUtils::GetCastTemplateValueName()
+{
+	return CastTemplateValueName;
+}
+
+const FName& RigVMTypeUtils::GetCastTemplateResultName()
+{
+	return CastTemplateResultName;
+}
+
+const FName& RigVMTypeUtils::GetCastTemplateNotation()
+{
+	return CastTemplateNotation;
+}
+
 #endif
diff --git a/Engine/Source/Developer/RigVMDeveloper/Private/RigVMModel/RigVMController.cpp b/Engine/Source/Developer/RigVMDeveloper/Private/RigVMModel/RigVMController.cpp
--- a/Engine/Source/Developer/RigVMDeveloper/Private/RigVMModel/RigVMController.cpp
+++ b/Engine/Source/Developer/RigVMDeveloper/Private/RigVMModel/RigVMController.cpp
@@ -3730,7 +3730,7 @@ TArray<FName> URigVMController::ImportNodesFromText(const FString& InText, bool
 				}
 			}
 
-			ReportErrorf(TEXT("Cannot import link '%s -> %s'."), *CreatedLink->SourcePinPath, *CreatedLink->TargetPinPath);
+			ReportErrorf(TEXT("Cannot import link '%s'."), *URigVMLink::GetPinPathRepresentation(CreatedLink->SourcePinPath, CreatedLink->TargetPinPath));
 			DestroyObject(CreatedLink);
 		}
 
@@ -9104,7 +9104,7 @@ bool URigVMController::PromotePinToVariable(URigVMPin* InPin, bool bCreateVariab
 }
 
 bool URigVMController::AddLink(const FString& InOutputPinPath, const FString& InInputPinPath, bool bSetupUndoRedo,
-	bool bPrintPythonCommand, ERigVMPinDirection InUserDirection)
+	bool bPrintPythonCommand, ERigVMPinDirection InUserDirection, bool bCreateCastNode)
 {
 	if(!IsValidGraph())
 	{
@@ -9148,7 +9148,7 @@ bool URigVMController::AddLink(const FString& InOutputPinPath, const FString& In
 	}
 	InputPin = InputPin->GetPinForLink();
 
-	const bool bSuccess = AddLink(OutputPin, InputPin, bSetupUndoRedo, InUserDirection);
+	const bool bSuccess = AddLink(OutputPin, InputPin, bSetupUndoRedo, InUserDirection, bCreateCastNode);
 	if (bSuccess && bPrintPythonCommand)
 	{
 		const FString GraphName = GetSanitizedGraphName(GetGraph()->GetGraphName());
@@ -9166,7 +9166,7 @@ bool URigVMController::AddLink(const FString& InOutputPinPath, const FString& In
 	return bSuccess;
 }
 
-bool URigVMController::AddLink(URigVMPin* OutputPin, URigVMPin* InputPin, bool bSetupUndoRedo, ERigVMPinDirection InUserDirection)
+bool URigVMController::AddLink(URigVMPin* OutputPin, URigVMPin* InputPin, bool bSetupUndoRedo, ERigVMPinDirection InUserDirection, bool bCreateCastNode)
 {
 	if (!bIsTransacting && !IsGraphEditable())
 	{
@@ -9205,7 +9205,7 @@ bool URigVMController::AddLink(URigVMPin* OutputPin, URigVMPin* InputPin, bool b
 	if (!bIsTransacting)
 	{
 		FString FailureReason;
-		if (!Graph->CanLink(OutputPin, InputPin, &FailureReason, GetCurrentByteCode(), UserLinkDirection))
+		if (!Graph->CanLink(OutputPin, InputPin, &FailureReason, GetCurrentByteCode(), UserLinkDirection, bCreateCastNode))
 		{
 			if(OutputPin->IsExecuteContext() && InputPin->IsExecuteContext())
 			{
@@ -9227,10 +9227,185 @@ bool URigVMController::AddLink(URigVMPin* OutputPin, URigVMPin* InputPin, bool b
 	FRigVMBaseAction Action;
 	if (bSetupUndoRedo)
 	{
-		Action.Title = FString::Printf(TEXT("Add Link"));
+		Action.Title = TEXT("Add Link");
 		ActionStack->BeginAction(Action);
 	}
 
+	// check if we need to inject cast node
+	if(bCreateCastNode &&
+		bEnableTypeCasting &&
+		OutputPin->GetTypeIndex() != InputPin->GetTypeIndex() &&
+		!OutputPin->IsWildCard() &&
+		!InputPin->IsWildCard())
+	{
+		if(!FRigVMRegistry::Get().CanMatchTypes(OutputPin->GetTypeIndex(), InputPin->GetTypeIndex(), true))
+		{
+			// returns true if this template node can resolve to the type without breaking any links
+			auto CanTemplateResolveToType = [](const TRigVMTypeIndex& InTypeIndex, URigVMPin* InPin) -> bool
+			{
+				if(const URigVMTemplateNode* TemplateNode = Cast<URigVMTemplateNode>(InPin->GetNode()))
+				{
+					if(const FRigVMTemplate* Template = TemplateNode->GetTemplate())
+					{
+						if(!InPin->IsRootPin())
+						{
+							return false;
+						}
+						
+						if(const FRigVMTemplateArgument* Argument = Template->FindArgument(InPin->GetFName()))
+						{
+							if(!Argument->SupportsTypeIndex(InTypeIndex))
+							{
+								return false;
+							}
+						}
+						
+						// we only rely on the cast if the number of links on non-singleton arguments is > 0
+						const TArray<URigVMLink*> Links = InPin->GetNode()->GetLinks();
+						int32 NumLinksRemaining = Links.Num();
+						for(const URigVMLink* Link : Links)
+						{
+							const URigVMPin* SourcePin = Link->GetSourcePin();
+							const URigVMPin* TargetPin = Link->GetTargetPin();
+
+							if(SourcePin && TargetPin)
+							{
+								const URigVMPin* PinOnTemplateNode = SourcePin->GetNode() == TemplateNode ? SourcePin : TargetPin;
+								if(PinOnTemplateNode->IsRootPin() && PinOnTemplateNode != InPin)
+								{
+									const FName ArgumentName = PinOnTemplateNode->GetFName();
+									if(const FRigVMTemplateArgument* Argument = Template->FindArgument(ArgumentName))
+									{
+										if(Argument->IsSingleton())
+										{
+											NumLinksRemaining--;
+										}
+									}
+								}
+								else
+								{
+									NumLinksRemaining--;
+								}
+							}
+						}
+						
+						return NumLinksRemaining == 0;
+					}
+				}
+				
+				return false;
+			};
+
+			if(bCreateCastNode)
+			{
+				// todo: this should take into account links which have been detached
+				// during load we don't filter out the links and assume the template node can cast
+				// correctly to each type. we need to look at the links and their map rather then
+				// URigVMPin::GetLinks
+				if(CanTemplateResolveToType(InputPin->GetTypeIndex(), OutputPin) || 
+					CanTemplateResolveToType(OutputPin->GetTypeIndex(), InputPin))
+				{
+					bCreateCastNode = false;
+				}
+			}
+			
+			if(bCreateCastNode)
+			{
+				const FRigVMFunction* CastFunction = RigVMTypeUtils::GetCastForTypeIndices(OutputPin->GetTypeIndex(), InputPin->GetTypeIndex());
+
+				if(CastFunction == nullptr)
+				{
+					// this is potentially a template node which has more types available.
+					// look through the filtered types and find a matching cast.
+				}
+
+				if(CastFunction == nullptr)
+				{
+					if (bSetupUndoRedo)
+					{
+						ActionStack->CancelAction(Action, this);
+					}
+					return false;
+				}
+
+				const FVector2D OutputPinPosition = OutputPin->GetNode()->GetPosition() + FVector2D(150, 40.f) + FVector2D(0, 16) * OutputPin->GetRootPin()->GetPinIndex();
+				const FVector2D InputPinPosition = InputPin->GetNode()->GetPosition() + FVector2D(-75, 40.f) + FVector2D(0, 16) * InputPin->GetRootPin()->GetPinIndex();
+				const FVector2D CastPosition = (OutputPinPosition + InputPinPosition) * 0.5; 
+
+				const URigVMNode* CastNode = nullptr;
+
+				// try to find an existing cast node
+				if(CastNode == nullptr)
+				{
+					for(URigVMLink* ExistingLink : OutputPin->GetLinks())
+					{
+						if(ExistingLink->GetSourcePin() == OutputPin)
+						{
+							if(URigVMUnitNode* ExistingCastNode = Cast<URigVMUnitNode>(ExistingLink->GetTargetPin()->GetNode()))
+							{
+								if(ExistingCastNode->GetScriptStruct() == CastFunction->Struct)
+								{
+									CastNode = ExistingCastNode;
+									break;
+								}
+							}
+						}
+					}
+				}
+
+				if(CastNode == nullptr)
+				{
+					CastNode = AddUnitNode(CastFunction->Struct, CastFunction->GetMethodName(), CastPosition, FString(), bSetupUndoRedo, false);
+				}
+				
+				if(CastNode == nullptr)
+				{
+					if (bSetupUndoRedo)
+					{
+						ActionStack->CancelAction(Action, this);
+					}
+					return false;
+				}
+
+				static const FString CastTemplateValueName = RigVMTypeUtils::GetCastTemplateValueName().ToString();
+				static const FString CastTemplateResultName = RigVMTypeUtils::GetCastTemplateResultName().ToString();
+				URigVMPin* ValuePin = CastNode->FindPin(CastTemplateValueName);
+				URigVMPin* ResultPin = CastNode->FindPin(CastTemplateResultName);
+
+				if(!OutputPin->IsLinkedTo(ValuePin))
+				{
+					if(!AddLink(OutputPin, ValuePin, bSetupUndoRedo))
+					{
+						if (bSetupUndoRedo)
+						{
+							ActionStack->CancelAction(Action, this);
+						}
+						return false;
+					}
+				}
+
+				if(!ResultPin->IsLinkedTo(InputPin))
+				{
+					if(!AddLink(ResultPin, InputPin, bSetupUndoRedo))
+					{
+						if (bSetupUndoRedo)
+						{
+							ActionStack->CancelAction(Action, this);
+						}
+						return false;
+					}
+				}
+
+				if(bSetupUndoRedo)
+				{
+					Action.Title = TEXT("Add Link with Cast");
+					ActionStack->EndAction(Action);
+				}
+				return true;
+			}
+		}
+	}
+
 	if (OutputPin->IsExecuteContext())
 	{
 		BreakAllLinks(OutputPin, false, bSetupUndoRedo);
@@ -9312,82 +9487,10 @@ bool URigVMController::AddLink(URigVMPin* OutputPin, URigVMPin* InputPin, bool b
 		ExpandPinRecursively(InputPin->GetParentPin(), bSetupUndoRedo);
 	}
 
-	// check if there's a cast available for this type combination
-	bool bIsUsingCast = false;
-	if(bEnableTypeCasting && RigVMTypeUtils::CanCastTypes(OutputPin->GetTypeIndex(), InputPin->GetTypeIndex()))
-	{
-		auto ShouldRelyOnCast = [](const TRigVMTypeIndex& InTypeIndex, URigVMPin* InPin) -> bool
-		{
-			if(InPin->IsRootPin())
-			{
-				if(const URigVMTemplateNode* TemplateNode = Cast<URigVMTemplateNode>(InPin->GetNode()))
-				{
-					if(const FRigVMTemplate* Template = TemplateNode->GetTemplate())
-					{
-						if(const FRigVMTemplateArgument* Argument = Template->FindArgument(InPin->GetFName()))
-						{
-							if(!Argument->SupportsTypeIndex(InTypeIndex))
-							{
-								return true;
-							}
-						}
-					}
-				}
-			}
-
-			// we only rely on the cast if the number of links on non-singleton arguments is > 0
-			const TArray<URigVMLink*> Links = InPin->GetNode()->GetLinks();
-			int32 NumLinksRemaining = Links.Num();
-			if(const URigVMTemplateNode* TemplateNode = Cast<URigVMTemplateNode>(InPin->GetNode()))
-			{
-				if(const FRigVMTemplate* Template = TemplateNode->GetTemplate())
-				{
-					for(const URigVMLink* Link : Links)
-					{
-						const URigVMPin* SourcePin = Link->GetSourcePin();
-						const URigVMPin* TargetPin = Link->GetTargetPin();
-
-						if(SourcePin && TargetPin)
-						{
-							const URigVMPin* PinOnTemplateNode = SourcePin->GetNode() == TemplateNode ? SourcePin : TargetPin;
-							if(PinOnTemplateNode->IsRootPin())
-							{
-								const FName ArgumentName = PinOnTemplateNode->GetFName();
-								if(const FRigVMTemplateArgument* Argument = Template->FindArgument(ArgumentName))
-								{
-									if(Argument->IsSingleton())
-									{
-										NumLinksRemaining--;
-									}
-								}
-							}
-							else
-							{
-								NumLinksRemaining--;
-							}
-						}
-					}
-				}
-			}
-
-			return NumLinksRemaining > 0;
-		};
-		
-		// only rely on a cast if there's already a link
-		if(InUserDirection == ERigVMPinDirection::Input)
-		{
-			bIsUsingCast = ShouldRelyOnCast(InputPin->GetTypeIndex(), OutputPin);
-		}
-		else
-		{
-			bIsUsingCast = ShouldRelyOnCast(OutputPin->GetTypeIndex(), InputPin);
-		}
-	}
-
 	// Before adding the link, let's resolve input and ouput pin types
 	// If templates, we will filter the permutations that support this link
 	// If any links need to be broken before perfoming this connection, try to find them and break them
-	if (!bIsTransacting && !bIsUsingCast)
+	if (!bIsTransacting)
 	{
 		URigVMPin* FirstToResolve = (InUserDirection == ERigVMPinDirection::Input) ? OutputPin : InputPin;
 		URigVMPin* SecondToResolve = (FirstToResolve == OutputPin) ? InputPin : OutputPin;
@@ -9422,19 +9525,16 @@ bool URigVMController::AddLink(URigVMPin* OutputPin, URigVMPin* InputPin, bool b
 	Notify(ERigVMGraphNotifType::LinkAdded, Link);
 
 	// Remove any temporary preferred type that was added
-	if(!bIsUsingCast)
+	for (int32 i=0; i<2; ++i)
 	{
-		for (int32 i=0; i<2; ++i)
+		URigVMPin* EdgePin = (i==0) ? OutputPin : InputPin;
+		bool& bAddedTemporaryPreferredType = (i==0) ? bAddedTemporaryOutputPreferredType : bAddedTemporaryInputPreferredType;
+		if (bAddedTemporaryPreferredType)
 		{
-			URigVMPin* EdgePin = (i==0) ? OutputPin : InputPin;
-			bool& bAddedTemporaryPreferredType = (i==0) ? bAddedTemporaryOutputPreferredType : bAddedTemporaryInputPreferredType;
-			if (bAddedTemporaryPreferredType)
+			if (URigVMTemplateNode* TemplateNode = Cast<URigVMTemplateNode>(EdgePin->GetNode()))
 			{
-				if (URigVMTemplateNode* TemplateNode = Cast<URigVMTemplateNode>(EdgePin->GetNode()))
-				{
-					URigVMPin* RootPin = EdgePin->GetRootPin();
-					RemovePreferredType(TemplateNode, RootPin->GetFName(), bSetupUndoRedo);				
-				}
+				URigVMPin* RootPin = EdgePin->GetRootPin();
+				RemovePreferredType(TemplateNode, RootPin->GetFName(), bSetupUndoRedo);				
 			}
 		}
 	}
@@ -9517,11 +9617,8 @@ bool URigVMController::AddLink(URigVMPin* OutputPin, URigVMPin* InputPin, bool b
 			}
 		};
 
-		if(!bIsUsingCast)
-		{
-			ResolveTemplateNodeToCommonTypes(OutputPin);
-			ResolveTemplateNodeToCommonTypes(InputPin);
-		}
+		ResolveTemplateNodeToCommonTypes(OutputPin);
+		ResolveTemplateNodeToCommonTypes(InputPin);
 #endif
 		
 		ActionStack->EndAction(Action);
@@ -14468,7 +14565,7 @@ int32 URigVMController::ReattachLinksToPinObjects(bool bFollowCoreRedirectors, c
 				}
 				else
 				{
-					ReportWarningf(TEXT("Unable to re-create link %s -> %s"), *Link->SourcePinPath, *Link->TargetPinPath);
+					ReportWarningf(TEXT("Unable to re-create link %s"), *URigVMLink::GetPinPathRepresentation(Link->SourcePinPath, Link->TargetPinPath));
 					TargetPin->Links.Remove(Link);
 					SourcePin->Links.Remove(Link);
 					continue;
@@ -14521,7 +14618,7 @@ int32 URigVMController::ReattachLinksToPinObjects(bool bFollowCoreRedirectors, c
 		
 		if (SourcePin == nullptr)
 		{
-			ReportWarningf(TEXT("Unable to re-create link %s -> %s"), *Link->SourcePinPath, *Link->TargetPinPath);
+			ReportWarningf(TEXT("Unable to re-create link %s"), *URigVMLink::GetPinPathRepresentation(Link->SourcePinPath, Link->TargetPinPath));
 			if (TargetPin != nullptr)
 			{
 				TargetPin->Links.Remove(Link);
@@ -14530,7 +14627,7 @@ int32 URigVMController::ReattachLinksToPinObjects(bool bFollowCoreRedirectors, c
 		}
 		if (TargetPin == nullptr)
 		{
-			ReportWarningf(TEXT("Unable to re-create link %s -> %s"), *Link->SourcePinPath, *Link->TargetPinPath);
+			ReportWarningf(TEXT("Unable to re-create link %s"), *URigVMLink::GetPinPathRepresentation(Link->SourcePinPath, Link->TargetPinPath));
 			if (SourcePin != nullptr)
 			{
 				SourcePin->Links.Remove(Link);
@@ -16176,6 +16273,16 @@ bool URigVMController::PrepareTemplatePinForType(URigVMPin* InPin, const TArray<
 
 	FRigVMRegistry& Registry = FRigVMRegistry::Get();
 
+	/*
+	if (InTypeIndices.Num() == 1)
+	{
+		if(Registry.CanMatchTypes(InPin->GetTypeIndex(), InTypeIndices[0], true))
+		{
+			return true;
+		}
+	}
+	*/
+
 	// It might be a subpin of a struct, in that case we just want to make sure
 	// the type required matches the one in the pin
 	if (InPin->IsStructMember())
@@ -16307,7 +16414,7 @@ bool URigVMController::PrepareTemplatePinForType(URigVMPin* InPin, const TArray<
 		for (int32 i=0; i<InconsistentLinks.Num(); ++i)
 		{
 			TSoftObjectPtr<URigVMGraph> GraphPtr(InconsistentLinksGraph[i]);
-			LinksToBreak.AddUnique(GraphPtr.Get()->FindLink(FString::Printf(TEXT("%s -> %s"), *InconsistentLinks[i].Key, *InconsistentLinks[i].Value)));
+			LinksToBreak.AddUnique(GraphPtr.Get()->FindLink(URigVMLink::GetPinPathRepresentation(InconsistentLinks[i].Key, InconsistentLinks[i].Value)));
 		}
 		
 		// Warn the user the links that would be broken
@@ -16397,7 +16504,7 @@ bool URigVMController::PrepareTemplatePinForType(URigVMPin* InPin, const TArray<
 
 			for (TPair<FString, FString>& Pair : InconsistentLinks)
 			{
-				LinksToBreak.AddUnique(GetGraph()->FindLink(FString::Printf(TEXT("%s -> %s"), *Pair.Key, *Pair.Value)));
+				LinksToBreak.AddUnique(GetGraph()->FindLink(URigVMLink::GetPinPathRepresentation(Pair.Key, Pair.Value)));
 			}
 		}
 
@@ -17258,12 +17365,6 @@ bool URigVMController::PropagateTemplateFilteredTypes(URigVMTemplateNode* InNode
 		const TArray<URigVMLink*>& Links = Pin->GetLinks();
 		for (URigVMLink* Link : Links)
 		{
-			// skip links which contain a cast
-			if(LinksWithCast.Contains(Link->GetPinPathRepresentation()))
-			{
-				continue;
-			}
-			
 			URigVMPin* OtherPin = Link->GetSourcePin() == Pin ? Link->GetTargetPin() : Link->GetSourcePin();
 			bool bPropagate = false;
 			bool bIsTemplate = false;
@@ -17296,12 +17397,9 @@ bool URigVMController::PropagateTemplateFilteredTypes(URigVMTemplateNode* InNode
 						}
 						else
 						{
-							if(!bEnableTypeCasting || !RigVMTypeUtils::CanCastTypes(Link->GetSourcePin()->GetTypeIndex(), Link->GetTargetPin()->GetTypeIndex()))
-							{
-								ensureMsgf(!ActionStack->BracketActions.IsEmpty(), TEXT("Unexpected link broken %s in package %s"), *Link->GetPinPathRepresentation(), *GetPackage()->GetPathName());
-								BreakLink(Link->GetSourcePin(), Link->GetTargetPin(), bSetupUndoRedo);
-								return false;
-							}
+							ensureMsgf(!ActionStack->BracketActions.IsEmpty(), TEXT("Unexpected link broken %s in package %s"), *Link->GetPinPathRepresentation(), *GetPackage()->GetPathName());
+							BreakLink(Link->GetSourcePin(), Link->GetTargetPin(), bSetupUndoRedo);
+							return false;
 						}
 					}
 				}								
@@ -17311,12 +17409,9 @@ bool URigVMController::PropagateTemplateFilteredTypes(URigVMTemplateNode* InNode
 			{
 				if (!InNode->FilteredSupportsType(Pin, OtherPin->GetTypeIndex()))
 				{
-					if(!bEnableTypeCasting || !RigVMTypeUtils::CanCastTypes(Link->GetSourcePin()->GetTypeIndex(), Link->GetTargetPin()->GetTypeIndex()))
-					{
-						ensureMsgf(!ActionStack->BracketActions.IsEmpty(), TEXT("Unexpected link broken %s in package %s"), *Link->GetPinPathRepresentation(), *GetPackage()->GetPathName());
-						BreakLink(Link->GetSourcePin(), Link->GetTargetPin(), bSetupUndoRedo);
-						return false;
-					}
+					ensureMsgf(!ActionStack->BracketActions.IsEmpty(), TEXT("Unexpected link broken %s in package %s"), *Link->GetPinPathRepresentation(), *GetPackage()->GetPathName());
+					BreakLink(Link->GetSourcePin(), Link->GetTargetPin(), bSetupUndoRedo);
+					return false;
 				}
 			}
 		}
@@ -17872,41 +17967,6 @@ bool URigVMController::RecomputeAllTemplateFilteredPermutations(bool bSetupUndoR
 	URigVMGraph* Graph = GetGraph();
 	check(Graph);
 
-	// update the transient list of links containing a cast
-	if(bEnableTypeCasting)
-	{
-		LinksWithCast.Reset();
-		for(URigVMLink* Link : Graph->GetLinks())
-		{
-			const URigVMPin* SourcePin = Link->GetSourcePin();
-			const URigVMPin* TargetPin = Link->GetTargetPin();
-
-			if(SourcePin && TargetPin)
-			{
-				const TRigVMTypeIndex SourceTypeIndex = SourcePin->GetTypeIndex();
-				const TRigVMTypeIndex TargetTypeIndex = TargetPin->GetTypeIndex();
-
-				if(SourceTypeIndex != TargetTypeIndex &&
-					RigVMTypeUtils::TypeIndex::WildCard != SourceTypeIndex &&
-					RigVMTypeUtils::TypeIndex::WildCardArray != SourceTypeIndex &&
-					RigVMTypeUtils::TypeIndex::WildCard != TargetTypeIndex &&
-					RigVMTypeUtils::TypeIndex::WildCardArray != TargetTypeIndex &&
-					!FRigVMRegistry::Get().CanMatchTypes(SourceTypeIndex, TargetTypeIndex, true))
-				{
-					ensureMsgf(RigVMTypeUtils::CanCastTypes(SourceTypeIndex, TargetTypeIndex),
-						TEXT("Incompatible types for link %s (%s) -> %s (%s) in %s.")
-						, *SourcePin->GetPinPath(true)
-						, *SourcePin->GetCPPType()
-						, *TargetPin->GetPinPath(true)
-						, *TargetPin->GetCPPType()
-						, *GetPackage()->GetPathName());
-
-					LinksWithCast.Add(Link->GetPinPathRepresentation());
-				}
-			}	
-		}
-	}
-
 	// Save all template pin types, in case we need them after initializing
 	TMap<URigVMPin*, TRigVMTypeIndex> TypesBeforeRecomputing;
 	for (URigVMNode* Node : Graph->GetNodes())
@@ -17948,18 +18008,7 @@ bool URigVMController::RecomputeAllTemplateFilteredPermutations(bool bSetupUndoR
 		TArray<URigVMLink*> SortedLinks = Graph->GetLinks().FilterByPredicate([this](URigVMLink* Link)
 		{
 			// Filter out detached links
-			if(!Link->GetSourcePin()->IsLinkedTo(Link->GetTargetPin()))
-			{
-				return false;
-			}
-
-			// filter out links with casts
-			if(LinksWithCast.Contains(Link->GetPinPathRepresentation()))
-			{
-				return false;
-			}
-
-			return true;
+			return Link->GetSourcePin()->IsLinkedTo(Link->GetTargetPin());
 		});
 
 		// Solve for unit nodes first, other templates are more expensive to solve. This way we are avoiding solving
@@ -18075,7 +18124,7 @@ bool URigVMController::RecomputeAllTemplateFilteredPermutations(bool bSetupUndoR
 					Notify(ERigVMGraphNotifType::LinkAdded, Link);
 				}
 			}
-			else if(!bEnableTypeCasting || !RigVMTypeUtils::CanCastTypes(OutputPin->GetTypeIndex(), InputPin->GetTypeIndex()))
+			else
 			{
 				LinksToRemove.Add(Link);
 			}
diff --git a/Engine/Source/Developer/RigVMDeveloper/Private/RigVMModel/RigVMGraph.cpp b/Engine/Source/Developer/RigVMDeveloper/Private/RigVMModel/RigVMGraph.cpp
--- a/Engine/Source/Developer/RigVMDeveloper/Private/RigVMModel/RigVMGraph.cpp
+++ b/Engine/Source/Developer/RigVMDeveloper/Private/RigVMModel/RigVMGraph.cpp
@@ -423,9 +423,9 @@ void URigVMGraph::PrepareCycleChecking(URigVMPin* InPin, bool bAsInput)
 	GetDiagnosticsAST(false, LinksToSkip)->PrepareCycleChecking(InPin);
 }
 
-bool URigVMGraph::CanLink(URigVMPin* InSourcePin, URigVMPin* InTargetPin, FString* OutFailureReason, const FRigVMByteCode* InByteCode, ERigVMPinDirection InUserLinkDirection)
+bool URigVMGraph::CanLink(URigVMPin* InSourcePin, URigVMPin* InTargetPin, FString* OutFailureReason, const FRigVMByteCode* InByteCode, ERigVMPinDirection InUserLinkDirection, bool bEnableTypeCasting)
 {
-	if (!URigVMPin::CanLink(InSourcePin, InTargetPin, OutFailureReason, InByteCode, InUserLinkDirection))
+	if (!URigVMPin::CanLink(InSourcePin, InTargetPin, OutFailureReason, InByteCode, InUserLinkDirection, false, bEnableTypeCasting))
 	{
 		return false;
 	}
diff --git a/Engine/Source/Developer/RigVMDeveloper/Private/RigVMModel/RigVMLink.cpp b/Engine/Source/Developer/RigVMDeveloper/Private/RigVMModel/RigVMLink.cpp
--- a/Engine/Source/Developer/RigVMDeveloper/Private/RigVMModel/RigVMLink.cpp
+++ b/Engine/Source/Developer/RigVMDeveloper/Private/RigVMModel/RigVMLink.cpp
@@ -85,12 +85,19 @@ URigVMPin* URigVMLink::GetOppositePin(const URigVMPin* InPin) const
 
 FString URigVMLink::GetPinPathRepresentation()
 {
-	return FString::Printf(TEXT("%s -> %s"), *GetSourcePin()->GetPinPath(), *GetTargetPin()->GetPinPath());
+	return GetPinPathRepresentation(GetSourcePin()->GetPinPath(), GetTargetPin()->GetPinPath());
+}
+
+FString URigVMLink::GetPinPathRepresentation(const FString& InSourcePinPath, const FString& InTargetPinPath)
+{
+	static constexpr TCHAR Format[] = TEXT("%s -> %s"); 
+	return FString::Printf(Format, *InSourcePinPath, *InTargetPinPath);
 }
 
 bool URigVMLink::SplitPinPathRepresentation(const FString& InString, FString& OutSource, FString& OutTarget)
 {
-	return InString.Split(TEXT(" -> "), &OutSource, &OutTarget);
+	static constexpr TCHAR Format[] = TEXT(" -> "); 
+	return InString.Split(Format, &OutSource, &OutTarget);
 }
 
 void URigVMLink::PrepareForCopy()
diff --git a/Engine/Source/Developer/RigVMDeveloper/Private/RigVMModel/RigVMPin.cpp b/Engine/Source/Developer/RigVMDeveloper/Private/RigVMModel/RigVMPin.cpp
--- a/Engine/Source/Developer/RigVMDeveloper/Private/RigVMModel/RigVMPin.cpp
+++ b/Engine/Source/Developer/RigVMDeveloper/Private/RigVMModel/RigVMPin.cpp
@@ -1656,7 +1656,7 @@ URigVMGraph* URigVMPin::GetGraph() const
 	return nullptr;
 }
 
-bool URigVMPin::CanLink(URigVMPin* InSourcePin, URigVMPin* InTargetPin, FString* OutFailureReason, const FRigVMByteCode* InByteCode, ERigVMPinDirection InUserLinkDirection, bool bInAllowNonArgumentPins)
+bool URigVMPin::CanLink(URigVMPin* InSourcePin, URigVMPin* InTargetPin, FString* OutFailureReason, const FRigVMByteCode* InByteCode, ERigVMPinDirection InUserLinkDirection, bool bInAllowNonArgumentPins, bool bEnableTypeCasting)
 {
 	if (InSourcePin == nullptr || InTargetPin == nullptr)
 	{
@@ -1738,7 +1738,7 @@ bool URigVMPin::CanLink(URigVMPin* InSourcePin, URigVMPin* InTargetPin, FString*
 
 		if (bCPPTypesDiffer)
 		{
-			if(RigVMTypeUtils::CanCastTypes(InSourcePin->GetTypeIndex(), InTargetPin->GetTypeIndex()))
+			if(bEnableTypeCasting && RigVMTypeUtils::CanCastTypes(InSourcePin->GetTypeIndex(), InTargetPin->GetTypeIndex()))
 			{
 				bCPPTypesDiffer = false;
 			}
diff --git a/Engine/Source/Developer/RigVMDeveloper/Public/RigVMDeveloperTypeUtils.h b/Engine/Source/Developer/RigVMDeveloper/Public/RigVMDeveloperTypeUtils.h
--- a/Engine/Source/Developer/RigVMDeveloper/Public/RigVMDeveloperTypeUtils.h
+++ b/Engine/Source/Developer/RigVMDeveloper/Public/RigVMDeveloperTypeUtils.h
@@ -62,5 +62,11 @@ namespace RigVMTypeUtils
 	
 	RIGVMDEVELOPER_API const FRigVMFunction* GetCastForTypeIndices(const TRigVMTypeIndex& InSourceTypeIndex, const TRigVMTypeIndex& InTargetTypeIndex);
 
+	RIGVMDEVELOPER_API const FName& GetCastTemplateValueName();
+
+	RIGVMDEVELOPER_API const FName& GetCastTemplateResultName();
+
+	RIGVMDEVELOPER_API const FName& GetCastTemplateNotation();
+
 #endif
 }
diff --git a/Engine/Source/Developer/RigVMDeveloper/Public/RigVMModel/RigVMController.h b/Engine/Source/Developer/RigVMDeveloper/Public/RigVMModel/RigVMController.h
--- a/Engine/Source/Developer/RigVMDeveloper/Public/RigVMModel/RigVMController.h
+++ b/Engine/Source/Developer/RigVMDeveloper/Public/RigVMModel/RigVMController.h
@@ -732,7 +732,7 @@ public:
 	// Adds a link to the graph.
 	// This causes a LinkAdded modified event.
 	UFUNCTION(BlueprintCallable, Category = RigVMController)
-	bool AddLink(const FString& InOutputPinPath, const FString& InInputPinPath, bool bSetupUndoRedo = true, bool bPrintPythonCommand = false, ERigVMPinDirection InUserDirection = ERigVMPinDirection::Output);
+	bool AddLink(const FString& InOutputPinPath, const FString& InInputPinPath, bool bSetupUndoRedo = true, bool bPrintPythonCommand = false, ERigVMPinDirection InUserDirection = ERigVMPinDirection::Output, bool bCreateCastNode = false);
 
 	// Removes a link from the graph.
 	// This causes a LinkRemoved modified event.
@@ -1015,7 +1015,7 @@ private:
 	void RelinkSourceAndTargetPins(URigVMNode* RigNode, bool bSetupUndoRedo = true);
 
 public:
-	bool AddLink(URigVMPin* OutputPin, URigVMPin* InputPin, bool bSetupUndoRedo = true, ERigVMPinDirection InUserDirection = ERigVMPinDirection::Invalid);
+	bool AddLink(URigVMPin* OutputPin, URigVMPin* InputPin, bool bSetupUndoRedo = true, ERigVMPinDirection InUserDirection = ERigVMPinDirection::Invalid, bool bCreateCastNode = false);
 	bool BreakLink(URigVMPin* OutputPin, URigVMPin* InputPin, bool bSetupUndoRedo = true);
 	bool BreakAllLinks(URigVMPin* Pin, bool bAsInput, bool bSetupUndoRedo = true);
 	void EnableTypeCasting(bool bEnabled = true) { bEnableTypeCasting = bEnabled; }
@@ -1178,7 +1178,6 @@ private:
 	bool bIgnoreRerouteCompactnessChanges;
 	ERigVMPinDirection UserLinkDirection;
 	bool bEnableTypeCasting;
-	mutable TArray<FString> LinksWithCast;
 
 	// temporary maps used for pin redirection
 	// only valid between Detach & ReattachLinksToPinObjects
diff --git a/Engine/Source/Developer/RigVMDeveloper/Public/RigVMModel/RigVMGraph.h b/Engine/Source/Developer/RigVMDeveloper/Public/RigVMModel/RigVMGraph.h
--- a/Engine/Source/Developer/RigVMDeveloper/Public/RigVMModel/RigVMGraph.h
+++ b/Engine/Source/Developer/RigVMDeveloper/Public/RigVMModel/RigVMGraph.h
@@ -140,7 +140,7 @@ public:
 	UScriptStruct* GetExecuteContextStruct() const;
 
 	void PrepareCycleChecking(URigVMPin* InPin, bool bAsInput);
-	virtual bool CanLink(URigVMPin* InSourcePin, URigVMPin* InTargetPin, FString* OutFailureReason, const FRigVMByteCode* InByteCode, ERigVMPinDirection InUserLinkDirection = ERigVMPinDirection::IO);
+	virtual bool CanLink(URigVMPin* InSourcePin, URigVMPin* InTargetPin, FString* OutFailureReason, const FRigVMByteCode* InByteCode, ERigVMPinDirection InUserLinkDirection = ERigVMPinDirection::IO, bool bEnableTypeCasting = true);
 	
 	TSharedPtr<FRigVMParserAST> GetDiagnosticsAST(bool bForceRefresh = false, TArray<URigVMLink*> InLinksToSkip = TArray<URigVMLink*>());
 	TSharedPtr<FRigVMParserAST> GetRuntimeAST(const FRigVMParserASTSettings& InSettings = FRigVMParserASTSettings::Optimized(), bool bForceRefresh = false);
diff --git a/Engine/Source/Developer/RigVMDeveloper/Public/RigVMModel/RigVMLink.h b/Engine/Source/Developer/RigVMDeveloper/Public/RigVMModel/RigVMLink.h
--- a/Engine/Source/Developer/RigVMDeveloper/Public/RigVMModel/RigVMLink.h
+++ b/Engine/Source/Developer/RigVMDeveloper/Public/RigVMModel/RigVMLink.h
@@ -57,6 +57,11 @@ public:
 	UFUNCTION(BlueprintCallable, Category = RigVMLink)
 	FString GetPinPathRepresentation();
 
+	// Returns a string representation of the Link given the two pin paths
+	// for example: "NodeA.Color.R -> NodeB.Translation.X"
+	// note: can be split again using SplitPinPathRepresentation
+	static FString GetPinPathRepresentation(const FString& InSourcePinPath, const FString& InTargetPinPath);
+
 	// Splits a pin path represenation of a link
 	// for example: "NodeA.Color.R -> NodeB.Translation.X"
 	// into its two pin paths
diff --git a/Engine/Source/Developer/RigVMDeveloper/Public/RigVMModel/RigVMPin.h b/Engine/Source/Developer/RigVMDeveloper/Public/RigVMModel/RigVMPin.h
--- a/Engine/Source/Developer/RigVMDeveloper/Public/RigVMModel/RigVMPin.h
+++ b/Engine/Source/Developer/RigVMDeveloper/Public/RigVMModel/RigVMPin.h
@@ -390,7 +390,7 @@ public:
 
 	// Returns true is the two provided source and target Pins
 	// can be linked to one another.
-	static bool CanLink(URigVMPin* InSourcePin, URigVMPin* InTargetPin, FString* OutFailureReason, const FRigVMByteCode* InByteCode, ERigVMPinDirection InUserLinkDirection = ERigVMPinDirection::IO, bool bInAllowWildcard = false);
+	static bool CanLink(URigVMPin* InSourcePin, URigVMPin* InTargetPin, FString* OutFailureReason, const FRigVMByteCode* InByteCode, ERigVMPinDirection InUserLinkDirection = ERigVMPinDirection::IO, bool bInAllowWildcard = false, bool bEnableTypeCasting = true);
 
 	// Returns true if this pin has injected nodes
 	bool HasInjectedNodes() const { return InjectionInfos.Num() > 0; }
