Control Rig: Reuse function compilation data

diff --git a/Engine/Plugins/Animation/ControlRig/Source/ControlRig/Private/ControlRigBlueprintGeneratedClass.cpp b/Engine/Plugins/Animation/ControlRig/Source/ControlRig/Private/ControlRigBlueprintGeneratedClass.cpp
--- a/Engine/Plugins/Animation/ControlRig/Source/ControlRig/Private/ControlRigBlueprintGeneratedClass.cpp
+++ b/Engine/Plugins/Animation/ControlRig/Source/ControlRig/Private/ControlRigBlueprintGeneratedClass.cpp
@@ -69,5 +69,12 @@ void UControlRigBlueprintGeneratedClass::Serialize(FArchive& Ar)
 			CDO->VM->CopyFrom(VM);
 		}
 	}
+
+	if (Ar.CustomVer(FControlRigObjectVersion::GUID) < FControlRigObjectVersion::StoreFunctionsInGeneratedClass)
+	{
+		return;
+	}
+	
+	Ar << GraphFunctionStore;
 }
 
diff --git a/Engine/Plugins/Animation/ControlRig/Source/ControlRig/Public/ControlRigBlueprintGeneratedClass.h b/Engine/Plugins/Animation/ControlRig/Source/ControlRig/Public/ControlRigBlueprintGeneratedClass.h
--- a/Engine/Plugins/Animation/ControlRig/Source/ControlRig/Public/ControlRigBlueprintGeneratedClass.h
+++ b/Engine/Plugins/Animation/ControlRig/Source/ControlRig/Public/ControlRigBlueprintGeneratedClass.h
@@ -7,10 +7,11 @@
 #include "Engine/BlueprintGeneratedClass.h"
 #include "ControlRigDefines.h"
 #include "RigVMCore/RigVM.h"
+#include "RigVMCore/RigVMGraphFunctionHost.h"
 #include "ControlRigBlueprintGeneratedClass.generated.h"
 
 UCLASS()
-class CONTROLRIG_API UControlRigBlueprintGeneratedClass : public UBlueprintGeneratedClass
+class CONTROLRIG_API UControlRigBlueprintGeneratedClass : public UBlueprintGeneratedClass, public IRigVMGraphFunctionHost
 {
 	GENERATED_UCLASS_BODY()
 
@@ -22,4 +23,12 @@ public:
 
 	// UObject interface
 	void Serialize(FArchive& Ar);
+
+	// IRigVMGraphFunctionHost interface
+	virtual const FRigVMGraphFunctionStore* GetRigVMGraphFunctionStore() const override { return &GraphFunctionStore; }
+	virtual FRigVMGraphFunctionStore* GetRigVMGraphFunctionStore() override { return &GraphFunctionStore; }
+
+	UPROPERTY()
+	FRigVMGraphFunctionStore GraphFunctionStore;
+	
 };
\ No newline at end of file
diff --git a/Engine/Plugins/Animation/ControlRig/Source/ControlRig/Public/ControlRigObjectVersion.h b/Engine/Plugins/Animation/ControlRig/Source/ControlRig/Public/ControlRigObjectVersion.h
--- a/Engine/Plugins/Animation/ControlRig/Source/ControlRig/Public/ControlRigObjectVersion.h
+++ b/Engine/Plugins/Animation/ControlRig/Source/ControlRig/Public/ControlRigObjectVersion.h
@@ -81,6 +81,9 @@ struct CONTROLRIG_API FControlRigObjectVersion
 		// Controls to be able specify which channels should be visible in sequencer
 		ControlTransformChannelFiltering,
 
+		// Store function information (and compilation data) in blueprint generated class
+		StoreFunctionsInGeneratedClass,
+
 		// -----<new versions can be added above this line>-------------------------------------------------
 		VersionPlusOne,
 		LatestVersion = VersionPlusOne - 1
diff --git a/Engine/Plugins/Animation/ControlRig/Source/ControlRigDeveloper/Private/ControlRigBlueprint.cpp b/Engine/Plugins/Animation/ControlRig/Source/ControlRigDeveloper/Private/ControlRigBlueprint.cpp
--- a/Engine/Plugins/Animation/ControlRig/Source/ControlRigDeveloper/Private/ControlRigBlueprint.cpp
+++ b/Engine/Plugins/Animation/ControlRig/Source/ControlRigDeveloper/Private/ControlRigBlueprint.cpp
@@ -128,8 +128,8 @@ UControlRigBlueprint::UControlRigBlueprint(const FObjectInitializer& ObjectIniti
 	RigVMClient.SetOuterClientHost(this, GET_MEMBER_NAME_CHECKED(UControlRigBlueprint, RigVMClient));
 	{
 		TGuardValue<bool> DisableClientNotifs(RigVMClient.bSuspendNotifications, true);
-		RigVMClient.AddModel(RigVMModelPrefix, false, &ObjectInitializer, false);
 		RigVMClient.GetOrCreateFunctionLibrary(false, &ObjectInitializer, false);
+		RigVMClient.AddModel(RigVMModelPrefix, false, &ObjectInitializer, false);
 	}
 	
 	FunctionLibraryEdGraph = ObjectInitializer.CreateDefaultSubobject<UControlRigGraph>(this, TEXT("RigVMFunctionLibraryEdGraph"));
@@ -316,6 +316,16 @@ const FRigVMClient* UControlRigBlueprint::GetRigVMClient() const
 	return &RigVMClient;
 }
 
+IRigVMGraphFunctionHost* UControlRigBlueprint::GetRigVMGraphFunctionHost() 
+{
+	return GetControlRigBlueprintGeneratedClass();
+}
+
+const IRigVMGraphFunctionHost* UControlRigBlueprint::GetRigVMGraphFunctionHost() const 
+{
+	return GetControlRigBlueprintGeneratedClass();
+}
+
 UObject* UControlRigBlueprint::GetEditorObjectForRigVMGraph(URigVMGraph* InVMGraph) const
 {
 	if(InVMGraph)
@@ -490,60 +500,6 @@ void UControlRigBlueprint::HandleConfigureRigVMController(const FRigVMClient* In
 
 	});
 
-	InControllerToConfigure->IsFunctionAvailableDelegate.BindLambda([WeakThis](URigVMLibraryNode* InFunction) -> bool
-	{
-		if(InFunction == nullptr)
-		{
-			return false;
-		}
-		
-		if(URigVMFunctionLibrary* Library = Cast<URigVMFunctionLibrary>(InFunction->GetOuter()))
-		{
-			if(UControlRigBlueprint* Blueprint = Cast<UControlRigBlueprint>(Library->GetOuter()))
-			{
-				if(Blueprint->IsFunctionPublic(InFunction->GetFName()))
-				{
-					return true;
-				}
-
-				// if it is private - we still see it as public if we are within the same blueprint 
-				if(WeakThis.IsValid())
-				{
-					if(WeakThis.Get() == Blueprint)
-					{
-						return true;
-					}
-				}
-			}
-		}
-
-		return false;
-	});
-
-	InControllerToConfigure->IsDependencyCyclicDelegate.BindLambda([WeakThis](UObject* InDependentObject, UObject* InDependencyObject) -> bool
-	{
-	    if(InDependentObject == nullptr || InDependencyObject == nullptr)
-	    {
-	        return false;
-	    }
-
-		UControlRigBlueprint* DependentBlueprint = InDependentObject->GetTypedOuter<UControlRigBlueprint>();
-		UControlRigBlueprint* DependencyBlueprint = InDependencyObject->GetTypedOuter<UControlRigBlueprint>();
-
-		if(DependentBlueprint == nullptr || DependencyBlueprint == nullptr)
-		{
-			return false;
-		}
-
-		if(DependentBlueprint == DependencyBlueprint)
-		{
-			return false;
-		}
-
-		const TArray<UControlRigBlueprint*> DependencyDependencies = DependencyBlueprint->GetDependencies(true);
-		return DependencyDependencies.Contains(DependentBlueprint);
-	});
-
 #if WITH_EDITOR
 
 	// this sets up three delegates:
@@ -658,12 +614,17 @@ bool UControlRigBlueprint::TryImportGraphFromText(const FString& InClipboardText
 	{
 		TGuardValue<FRigVMController_RequestLocalizeFunctionDelegate> RequestLocalizeDelegateGuard(
             FunctionLibraryController->RequestLocalizeFunctionDelegate,
-            FRigVMController_RequestLocalizeFunctionDelegate::CreateLambda([this](URigVMLibraryNode* InFunctionToLocalize)
+            FRigVMController_RequestLocalizeFunctionDelegate::CreateLambda([this](TSoftObjectPtr<URigVMLibraryNode> InFunctionToLocalize)
             {
-            	BroadcastRequestLocalizeFunctionDialog(InFunctionToLocalize);
-
-                const URigVMLibraryNode* LocalizedFunctionNode = GetLocalFunctionLibrary()->FindPreviouslyLocalizedFunction(InFunctionToLocalize);
-                return LocalizedFunctionNode != nullptr;
+            	if (InFunctionToLocalize.IsValid())
+            	{
+            		URigVMLibraryNode* LibraryNode = InFunctionToLocalize.Get();
+            		BroadcastRequestLocalizeFunctionDialog(LibraryNode);
+
+				   const URigVMLibraryNode* LocalizedFunctionNode = GetLocalFunctionLibrary()->FindPreviouslyLocalizedFunction(LibraryNode);
+				   return LocalizedFunctionNode != nullptr;
+            	}
+            	return false;
             })
         );
 		
@@ -774,6 +735,13 @@ void UControlRigBlueprint::PreSave(FObjectPreSaveContext ObjectSaveContext)
 		{
 			SupportedEventNames = CDO->GetSupportedEvents();
 		}
+
+		PublicGraphFunctions.Reset();
+		PublicGraphFunctions.SetNum(RigClass->GetRigVMGraphFunctionStore()->PublicFunctions.Num());
+		for (int32 i=0; i<PublicGraphFunctions.Num(); ++i)
+		{
+			PublicGraphFunctions[i] = RigClass->GetRigVMGraphFunctionStore()->PublicFunctions[i].Header;
+		}
 	}
 
 	bExposesAnimatableControls = false;
@@ -787,17 +755,6 @@ void UControlRigBlueprint::PreSave(FObjectPreSaveContext ObjectSaveContext)
 		return true;
 	});
 
-	for(FControlRigPublicFunctionData& FunctionData : PublicFunctions)
-	{
-		if(URigVMLibraryNode* FunctionNode = RigVMClient.GetFunctionLibrary()->FindFunction(FunctionData.Name))
-		{
-			if(UControlRigGraph* Graph = Cast<UControlRigGraph>(GetEdGraph(FunctionNode->GetContainedGraph())))
-			{
-				FunctionData = Graph->GetPublicFunctionData();
-			}
-		}
-	}
-
 	FunctionReferenceNodeData = GetReferenceNodeData();
 	IAssetRegistry::GetChecked().AssetTagsFinalized(*this);
 }
@@ -901,6 +858,7 @@ void UControlRigBlueprint::PostLoad()
 		PatchPropagateToChildren();
 		PatchParameterNodesOnLoad();
 		PatchLinksWithCast();
+		PatchFunctionsOnLoad();
 
 #if WITH_EDITOR
 
@@ -911,6 +869,8 @@ void UControlRigBlueprint::PostLoad()
 		// any pin that references the struct needs to be regenerated
 		RefreshAllModels();
 
+		GetControlRigBlueprintGeneratedClass()->GetRigVMGraphFunctionStore()->RemoveAllCompilationData();
+
 		// perform backwards compat value upgrades
 		TArray<URigVMGraph*> GraphsToValidate = GetAllModels();
 		for (int32 GraphIndex = 0; GraphIndex < GraphsToValidate.Num(); GraphIndex++)
@@ -955,7 +915,7 @@ void UControlRigBlueprint::PostLoad()
 				{
 					if(URigVMFunctionReferenceNode* FunctionReferenceNode = Cast<URigVMFunctionReferenceNode>(Node))
 					{
-						if(URigVMLibraryNode* DependencyNode = FunctionReferenceNode->GetReferencedNode())
+						if(const URigVMLibraryNode* DependencyNode = FunctionReferenceNode->LoadReferencedNode())
 						{
 							if(UControlRigBlueprint* DependencyBlueprint = DependencyNode->GetTypedOuter<UControlRigBlueprint>())
 							{
@@ -1114,7 +1074,7 @@ void UControlRigBlueprint::HandlePackageDone()
 					for(int32 ReferenceIndex = 0; ReferenceIndex < Pair.Value.Num(); ReferenceIndex++)
 					{
 						// update the build data
-						BuildData->RegisterFunctionReference(FunctionKey, Pair.Value[ReferenceIndex]);
+						BuildData->RegisterFunctionReference(FunctionKey->GetFunctionIdentifier(), Pair.Value[ReferenceIndex]);
 
 						// find all control rigs matching the reference node
 						FAssetData AssetData = AssetRegistry.GetAssetByObjectPath(
@@ -1352,6 +1312,7 @@ void UControlRigBlueprint::RefreshAllModels(EControlRigBlueprintLoadType InLoadT
 	const bool bIsPostLoad = InLoadType == EControlRigBlueprintLoadType::PostLoad;
 	
 	TGuardValue<bool> IsCompilingGuard(bIsCompiling, true);
+	TGuardValue<bool> ClientIgnoreModificationsGuard(RigVMClient.bIgnoreModelNotifications, true);
 	
 	TArray<URigVMGraph*> GraphsToDetach = RigVMClient.GetAllModels(true, false);
 
@@ -1407,18 +1368,6 @@ void UControlRigBlueprint::RefreshAllModels(EControlRigBlueprintLoadType InLoadT
 				{
 					if (URigVMFunctionReferenceNode* FunctionReferenceNode = Cast<URigVMFunctionReferenceNode>(LibraryNode))
 					{
-						if (URigVMLibraryNode* ReferencedNode = Cast<URigVMLibraryNode>(FunctionReferenceNode->GetReferencedNode()))
-						{
-							if (ReferencedNode->GetLibrary() != GetLocalFunctionLibrary())
-							{
-								ReferencedNode->GetPackage()->FullyLoad();
-							}
-							else
-							{
-								ContainedGraphs.Add(ReferencedNode->GetContainedGraph());
-							}
-						}
-						
 						continue;
 					}
 
@@ -1704,7 +1653,7 @@ TArray<UControlRigBlueprint*> UControlRigBlueprint::GetReferencedControlRigBluep
 			{
 				if(URigVMFunctionReferenceNode* FunctionRefNode = Cast<URigVMFunctionReferenceNode>(RigNode->GetModelNode()))
 				{
-					if(URigVMLibraryNode* ReferencedNode = FunctionRefNode->GetReferencedNode())
+					if(const URigVMLibraryNode* ReferencedNode = FunctionRefNode->LoadReferencedNode())
 					{
 						if(URigVMFunctionLibrary* ReferencedFunctionLibrary = ReferencedNode->GetLibrary())
 						{
@@ -1712,11 +1661,6 @@ TArray<UControlRigBlueprint*> UControlRigBlueprint::GetReferencedControlRigBluep
 							{
 								continue;
 							}
-
-							if(UControlRigBlueprint* ReferencedBlueprint = Cast<UControlRigBlueprint>(ReferencedFunctionLibrary->GetOuter()))
-							{
-								ReferencedBlueprints.AddUnique(ReferencedBlueprint);
-							}
 						}
 					}
 				}
@@ -2160,12 +2104,13 @@ TArray<FString> UControlRigBlueprint::GeneratePythonCommands(const FString InNew
 				{
 					if (URigVMFunctionReferenceNode* Reference = Cast<URigVMFunctionReferenceNode>(Node))
 					{
-						if (Reference->GetContainedGraph()->GetPackage() != GetPackage())
+						if (Reference->GetReferencedFunctionHeader().LibraryPointer.HostObject != GetControlRigBlueprintGeneratedClass())
 						{
 							continue;
 						}
-				
-						if (!ProcessedGraphs.Contains(Reference->GetContainedGraph()))
+
+						URigVMLibraryNode* LibraryNode = Cast<URigVMLibraryNode>(Reference->GetReferencedFunctionHeader().LibraryPointer.LibraryNode.ResolveObject());
+						if (!ProcessedGraphs.Contains(LibraryNode->GetContainedGraph()))
 						{
 							bFoundUnprocessedReference = true;
 							break;
@@ -2415,14 +2360,10 @@ UEdGraph* UControlRigBlueprint::GetEdGraph(const FString& InNodePath) const
 
 bool UControlRigBlueprint::IsFunctionPublic(const FName& InFunctionName) const
 {
-	for(const FControlRigPublicFunctionData& PublicFunction : PublicFunctions)
+	if (URigVMLibraryNode* LibraryNode = GetLocalFunctionLibrary()->FindFunction(InFunctionName))
 	{
-		if(PublicFunction.Name == InFunctionName)
-		{
-			return true;
-		}
+		return GetControlRigBlueprintGeneratedClass()->GetRigVMGraphFunctionStore()->IsFunctionPublic(LibraryNode->GetFunctionIdentifier());
 	}
-	
 	return false;
 }
 
@@ -2434,37 +2375,12 @@ void UControlRigBlueprint::MarkFunctionPublic(const FName& InFunctionName, bool
 	}
 	
 	Modify();
+	
+	UControlRigBlueprintGeneratedClass* Generated = GetControlRigBlueprintGeneratedClass();
+	Generated->Modify();
 
-	if(bIsPublic)
-	{
-		if(URigVMLibraryNode* FunctionNode = GetLocalFunctionLibrary()->FindFunction(InFunctionName))
-		{
-			if(UControlRigGraph* RigGraph = Cast<UControlRigGraph>(GetEdGraph(FunctionNode->GetContainedGraph())))
-			{
-				const FControlRigPublicFunctionData NewFunctionData = RigGraph->GetPublicFunctionData();
-				for(FControlRigPublicFunctionData& ExistingFunctionData : PublicFunctions)
-				{
-					if(ExistingFunctionData.Name == NewFunctionData.Name)
-					{
-						ExistingFunctionData = NewFunctionData;
-						return;
-					}
-				}
-				PublicFunctions.Add(NewFunctionData);
-			}
-		}
-	}
-	else
-	{
-		for(int32 Index = 0; Index < PublicFunctions.Num(); Index++)
-		{
-			if(PublicFunctions[Index].Name == InFunctionName)
-			{
-				PublicFunctions.RemoveAt(Index);
-				return;
-			}
-		}
-	}
+	const FRigVMGraphFunctionIdentifier& Identifier = GetLocalFunctionLibrary()->FindFunction(InFunctionName)->GetFunctionIdentifier();
+	Generated->GetRigVMGraphFunctionStore()->MarkFunctionAsPublic(Identifier, bIsPublic);
 }
 
 TArray<UControlRigBlueprint*> UControlRigBlueprint::GetDependencies(bool bRecursive) const
@@ -2478,7 +2394,7 @@ TArray<UControlRigBlueprint*> UControlRigBlueprint::GetDependencies(bool bRecurs
 		{
 			if(URigVMFunctionReferenceNode* FunctionReferenceNode = Cast<URigVMFunctionReferenceNode>(Node))
 			{
-				if(URigVMLibraryNode* LibraryNode = FunctionReferenceNode->GetReferencedNode())
+				if(const URigVMLibraryNode* LibraryNode = FunctionReferenceNode->LoadReferencedNode())
 				{
 					if(UControlRigBlueprint* DependencyBlueprint = LibraryNode->GetTypedOuter<UControlRigBlueprint>())
 					{
@@ -3424,6 +3340,8 @@ void UControlRigBlueprint::RebuildGraphFromModel()
 	DECLARE_SCOPE_HIERARCHICAL_COUNTER_FUNC()
 
 	TGuardValue<bool> SelfGuard(bSuspendModelNotificationsForSelf, true);
+	TGuardValue<bool> ClientIgnoreModificationsGuard(RigVMClient.bIgnoreModelNotifications, true);
+	
 	verify(GetOrCreateController());
 
 	TArray<UEdGraph*> EdGraphs;
@@ -3660,7 +3578,8 @@ void UControlRigBlueprint::HandleModifiedEvent(ERigVMGraphNotifType InNotifType,
 			case ERigVMGraphNotifType::NodeAdded:
 			case ERigVMGraphNotifType::NodeRemoved:
 			{
-				if (InNotifType == ERigVMGraphNotifType::NodeRemoved)
+				bool bAdded = InNotifType == ERigVMGraphNotifType::NodeAdded;
+				if (!bAdded)
 				{
 					if (URigVMNode* RigVMNode = Cast<URigVMNode>(InSubject))
 					{
@@ -3670,7 +3589,7 @@ void UControlRigBlueprint::HandleModifiedEvent(ERigVMGraphNotifType InNotifType,
 
 				if (URigVMCollapseNode* CollapseNode = Cast<URigVMCollapseNode>(InSubject))
 				{
-					if (InNotifType == ERigVMGraphNotifType::NodeAdded)
+					if (bAdded)
 					{
 						// If the controller for this graph already exist, make sure it is referencing the correct graph
 						if (URigVMController* Controller = RigVMClient.GetController(CollapseNode->GetContainedGraph()))
@@ -3691,18 +3610,6 @@ void UControlRigBlueprint::HandleModifiedEvent(ERigVMGraphNotifType InNotifType,
 					ClearTransientControls();
 					RequestAutoVMRecompilation();
 
-					if(CollapseNode->GetOuter()->IsA<URigVMFunctionLibrary>())
-					{
-						for(int32 Index = 0; Index < PublicFunctions.Num(); Index++)
-						{
-							if(PublicFunctions[Index].Name == CollapseNode->GetFName())
-							{
-								Modify();
-								PublicFunctions.RemoveAt(Index);
-							}
-						}
-					}
-
 					MarkPackageDirty();
 					FBlueprintEditorUtils::MarkBlueprintAsStructurallyModified(this);
 					break;
@@ -3858,18 +3765,6 @@ void UControlRigBlueprint::HandleModifiedEvent(ERigVMGraphNotifType InNotifType,
 						ContainedEdGraph->Rename(*CollapseNode->GetEditorSubGraphName(), nullptr);
 					}
 
-					if(CollapseNode->GetOuter()->IsA<URigVMFunctionLibrary>())
-					{
-						for(int32 Index = 0; Index < PublicFunctions.Num(); Index++)
-						{
-							if(PublicFunctions[Index].Name == CollapseNode->GetPreviousFName())
-							{
-								Modify();
-								PublicFunctions[Index].Name = CollapseNode->GetFName();
-							}
-						}
-					}
-
 					FBlueprintEditorUtils::MarkBlueprintAsStructurallyModified(this);
 				}
 				break;
@@ -4136,11 +4031,15 @@ void UControlRigBlueprint::PatchFunctionReferencesOnLoad()
 		{
 			if (URigVMFunctionReferenceNode* FunctionReferenceNode = Cast<URigVMFunctionReferenceNode>(Node))
 			{
-				if (!FunctionReferenceNode->GetReferencedNode())
+				if (!FunctionReferenceNode->ReferencedNodePtr_DEPRECATED.IsValid())
+				{
+					FunctionReferenceNode->ReferencedNodePtr_DEPRECATED.LoadSynchronous();
+				}
+				if (!FunctionReferenceNode->ReferencedNodePtr_DEPRECATED)
 				{
 					if(URigVMFunctionLibrary* FunctionLibrary = RigVMClient.GetFunctionLibrary())
 					{
-						FString FunctionPath = FunctionReferenceNode->ReferencedNodePtr.ToSoftObjectPath().GetSubPathString();
+						FString FunctionPath = FunctionReferenceNode->ReferencedNodePtr_DEPRECATED.ToSoftObjectPath().GetSubPathString();
 						
 						FString Left, Right;
 						if(FunctionPath.Split(TEXT("."), &Left, &Right))
@@ -4150,11 +4049,11 @@ void UControlRigBlueprint::PatchFunctionReferencesOnLoad()
 							{
 								if (URigVMLibraryNode* LibraryNode = Cast<URigVMLibraryNode>(FunctionLibrary->FindNode(Right)))
 								{
-									FunctionReferenceNode->SetReferencedNode(LibraryNode);
+									FunctionReferenceNode->ReferencedNodePtr_DEPRECATED = LibraryNode;
 								}
 							}
 						}
-					}				
+					}
 				}
 			}
 		}
@@ -4609,6 +4508,70 @@ void UControlRigBlueprint::PatchLinksWithCast()
 #endif
 }
 
+void UControlRigBlueprint::PatchFunctionsOnLoad()
+{
+	UControlRigBlueprintGeneratedClass* CRGeneratedClass = GetControlRigBlueprintGeneratedClass();
+	FRigVMGraphFunctionStore& Store = CRGeneratedClass->GraphFunctionStore;
+	const URigVMFunctionLibrary* Library = GetLocalFunctionLibrary();
+	if (GetLinkerCustomVersion(FControlRigObjectVersion::GUID) < FControlRigObjectVersion::StoreFunctionsInGeneratedClass)
+	{
+		// Move functions to the generated class
+		for (const FControlRigPublicFunctionData& Data : PublicFunctions_DEPRECATED)
+		{
+			const URigVMLibraryNode* LibraryNode = Library->FindFunction(Data.Name);
+			Store.AddFunction(LibraryNode->GetFunctionHeader(), true);
+		}
+	}
+
+	// Add functions
+	for (const URigVMLibraryNode* LibraryNode : Library->GetFunctions())
+	{
+		FRigVMGraphFunctionIdentifier Identifier = LibraryNode->GetFunctionIdentifier();
+		bool bIsPublic = PublicGraphFunctions.ContainsByPredicate([Identifier](const FRigVMGraphFunctionHeader& Header)
+		{
+			return Header.LibraryPointer == Identifier;
+		});
+		Store.AddFunction(LibraryNode->GetFunctionHeader(), bIsPublic);		
+	}
+
+	if (GetLinkerCustomVersion(FControlRigObjectVersion::GUID) < FControlRigObjectVersion::StoreFunctionsInGeneratedClass)
+	{
+		TArray<URigVMGraph*> Graphs = GetAllModels();
+		for (URigVMGraph* Graph : Graphs)
+		{
+			TArray<URigVMNode*> Nodes = Graph->GetNodes();
+			for (int32 i=0; i<Nodes.Num(); ++i)
+			{
+				if (URigVMFunctionReferenceNode* FuncRefNode = Cast<URigVMFunctionReferenceNode>(Nodes[i]))
+				{
+					TSoftObjectPtr<URigVMLibraryNode> ReferencedLibraryNode = FuncRefNode->ReferencedNodePtr_DEPRECATED;
+					if (ReferencedLibraryNode.IsValid())
+					{
+						if (URigVMLibraryNode* LibraryNode = ReferencedLibraryNode.LoadSynchronous())
+						{
+							if (URigVMFunctionLibrary* FunctionLibrary = LibraryNode->GetGraph()->GetDefaultFunctionLibrary())
+							{
+								FuncRefNode->ReferencedFunctionHeader = LibraryNode->GetFunctionHeader();
+							}
+						}
+					}
+				}
+
+				if (URigVMCollapseNode* CollapseNode = Cast<URigVMCollapseNode>(Nodes[i]))
+				{
+					Nodes.Append(CollapseNode->GetContainedNodes());
+				}
+			}
+		}
+	}
+
+	// Update dependencies
+	for (const URigVMLibraryNode* LibraryNode : Library->GetFunctions())
+	{
+		GetRigVMClient()->UpdateDependenciesForFunction(LibraryNode);		
+	}
+}
+
 void UControlRigBlueprint::PropagatePoseFromInstanceToBP(UControlRig* InControlRig)
 {
 	check(InControlRig);
diff --git a/Engine/Plugins/Animation/ControlRig/Source/ControlRigDeveloper/Private/Graph/ControlRigGraph.cpp b/Engine/Plugins/Animation/ControlRig/Source/ControlRigDeveloper/Private/Graph/ControlRigGraph.cpp
--- a/Engine/Plugins/Animation/ControlRig/Source/ControlRigDeveloper/Private/Graph/ControlRigGraph.cpp
+++ b/Engine/Plugins/Animation/ControlRig/Source/ControlRigDeveloper/Private/Graph/ControlRigGraph.cpp
@@ -1194,42 +1194,5 @@ void UControlRigGraph::HandleVMCompiledEvent(UObject* InCompiledObject, URigVM*
 
 #endif
 
-FControlRigPublicFunctionData UControlRigGraph::GetPublicFunctionData() const
-{
-	FControlRigPublicFunctionData Data;
-
-	FString Prefix, ModelNodeName;
-	if(!URigVMNode::SplitNodePathAtEnd(ModelNodePath, Prefix, ModelNodeName))
-	{
-		ModelNodeName = ModelNodePath;
-	}
-	Data.Name = *ModelNodeName;
-
-	if(URigVMGraph* RigGraph = GetModel())
-	{
-		if(URigVMCollapseNode* FunctionNode = Cast<URigVMCollapseNode>(RigGraph->GetOuter()))
-		{
-			Data.Category = FunctionNode->GetNodeCategory();
-			Data.Keywords = FunctionNode->GetNodeKeywords();
-			
-			for(URigVMPin* Pin : FunctionNode->GetPins())
-			{
-				FControlRigPublicFunctionArg Arg;
-				Arg.Name = Pin->GetFName();
-				Arg.bIsArray = Pin->IsArray();
-				Arg.Direction = Pin->GetDirection();
-				Arg.CPPType = *Pin->GetCPPType();
-				if(Pin->GetCPPTypeObject())
-				{
-					Arg.CPPTypeObjectPath = *Pin->GetCPPTypeObject()->GetPathName();
-				}
-				Data.Arguments.Add(Arg);
-			}
-		}
-	}
-
-	return Data;
-}
-
 #undef LOCTEXT_NAMESPACE
 
diff --git a/Engine/Plugins/Animation/ControlRig/Source/ControlRigDeveloper/Private/Graph/ControlRigGraphNode.cpp b/Engine/Plugins/Animation/ControlRig/Source/ControlRigDeveloper/Private/Graph/ControlRigGraphNode.cpp
--- a/Engine/Plugins/Animation/ControlRig/Source/ControlRigDeveloper/Private/Graph/ControlRigGraphNode.cpp
+++ b/Engine/Plugins/Animation/ControlRig/Source/ControlRigDeveloper/Private/Graph/ControlRigGraphNode.cpp
@@ -81,12 +81,13 @@ FText UControlRigGraphNode::GetNodeTitle(ENodeTitleType::Type TitleType) const
 
 			else if(URigVMFunctionReferenceNode* FunctionReferenceNode = Cast<URigVMFunctionReferenceNode>(ModelNode))
 			{
-				if(URigVMLibraryNode* ReferencedNode = FunctionReferenceNode->GetReferencedNode())
+				const FRigVMGraphFunctionHeader& ReferencedHeader = FunctionReferenceNode->GetReferencedFunctionHeader();
 				{
-					UPackage* ReferencedPackage = ReferencedNode->GetOutermost();
-					if(ReferencedPackage != ModelNode->GetOutermost())
+					TSoftObjectPtr<URigVMFunctionReferenceNode> RefNodePtr(FunctionReferenceNode);
+					const FString& PackagePath = ReferencedHeader.LibraryPointer.LibraryNode.GetLongPackageName();
+					if(PackagePath != RefNodePtr.GetLongPackageName())
 					{
-						SubTitle = FString::Printf(TEXT("From %s"), *ReferencedPackage->GetName());
+						SubTitle = FString::Printf(TEXT("From %s"), *PackagePath);
 					}
 					else
 					{
@@ -1142,7 +1143,7 @@ void UControlRigGraphNode::AllocateDefaultPins()
 	{
 		if(FunctionReferenceNode->RequiresVariableRemapping())
 		{
-			TArray<FRigVMExternalVariable> CurrentExternalVariables = FunctionReferenceNode->GetContainedGraph()->GetExternalVariables();
+			TArray<FRigVMExternalVariable> CurrentExternalVariables = FunctionReferenceNode->GetReferencedFunctionHeader().ExternalVariables;
 			for(const FRigVMExternalVariable& CurrentExternalVariable : CurrentExternalVariables)
 			{
 				ExternalVariables.Add(MakeShared<FRigVMExternalVariable>(CurrentExternalVariable));
diff --git a/Engine/Plugins/Animation/ControlRig/Source/ControlRigDeveloper/Public/ControlRigBlueprint.h b/Engine/Plugins/Animation/ControlRig/Source/ControlRigDeveloper/Public/ControlRigBlueprint.h
--- a/Engine/Plugins/Animation/ControlRig/Source/ControlRigDeveloper/Public/ControlRigBlueprint.h
+++ b/Engine/Plugins/Animation/ControlRig/Source/ControlRigDeveloper/Public/ControlRigBlueprint.h
@@ -19,6 +19,7 @@
 #include "RigVMCompiler/RigVMCompiler.h"
 #include "ControlRigValidationPass.h"
 #include "Drawing/ControlRigDrawContainer.h"
+#include "RigVMCore/RigVMGraphFunctionDefinition.h"
 
 #if WITH_EDITOR
 #include "Kismet2/Kismet2NameValidators.h"
@@ -46,7 +47,7 @@ DECLARE_DELEGATE_RetVal_OneParam(bool, FControlRigOnBreakLinksDialogRequestedDel
 DECLARE_EVENT(UControlRigBlueprint, FOnBreakpointAdded);
 DECLARE_EVENT_OneParam(UControlRigBlueprint, FOnRequestInspectObject, const TArray<UObject*>& );
 
-USTRUCT()
+USTRUCT(meta = (Deprecated = "5.2"))
 struct CONTROLRIGDEVELOPER_API FControlRigPublicFunctionArg
 {
 	GENERATED_BODY();
@@ -77,7 +78,7 @@ struct CONTROLRIGDEVELOPER_API FControlRigPublicFunctionArg
 	FEdGraphPinType GetPinType() const;
 };
 
-USTRUCT()
+USTRUCT(meta = (Deprecated = "5.2"))
 struct CONTROLRIGDEVELOPER_API FControlRigPublicFunctionData
 {
 	GENERATED_BODY();
@@ -264,6 +265,8 @@ public:
 	// IRigVMClientHost interface
 	virtual FRigVMClient* GetRigVMClient() override;
 	virtual const FRigVMClient* GetRigVMClient() const override;
+	virtual IRigVMGraphFunctionHost* GetRigVMGraphFunctionHost() override;
+	virtual const IRigVMGraphFunctionHost* GetRigVMGraphFunctionHost() const override;
 	virtual UObject* GetEditorObjectForRigVMGraph(URigVMGraph* InVMGraph) const override;
 	virtual void HandleRigVMGraphAdded(const FRigVMClient* InClient, const FString& InNodePath) override;
 	virtual void HandleRigVMGraphRemoved(const FRigVMClient* InClient, const FString& InNodePath) override;
@@ -408,7 +411,11 @@ protected:
 
 	/** Asset searchable information about exposed public functions on this rig */
 	UPROPERTY(AssetRegistrySearchable)
-	TArray<FControlRigPublicFunctionData> PublicFunctions;
+	TArray<FControlRigPublicFunctionData> PublicFunctions_DEPRECATED;
+
+	/** Asset searchable information about exposed public functions on this rig */
+	UPROPERTY(AssetRegistrySearchable)
+	TArray<FRigVMGraphFunctionHeader> PublicGraphFunctions;
 
 	/** Asset searchable information function references in this rig */
 	UPROPERTY(AssetRegistrySearchable)
@@ -606,6 +613,7 @@ private:
 	void PatchParameterNodesOnLoad();
 	void PatchTemplateNodesWithPreferredPermutation();
 	void PatchLinksWithCast();
+	void PatchFunctionsOnLoad();
 
 	TMap<FName, int32> AddedMemberVariableMap;
 	TArray<FBPVariableDescription> LastNewVariables;
diff --git a/Engine/Plugins/Animation/ControlRig/Source/ControlRigDeveloper/Public/Graph/ControlRigGraph.h b/Engine/Plugins/Animation/ControlRig/Source/ControlRigDeveloper/Public/Graph/ControlRigGraph.h
--- a/Engine/Plugins/Animation/ControlRig/Source/ControlRigDeveloper/Public/Graph/ControlRigGraph.h
+++ b/Engine/Plugins/Animation/ControlRig/Source/ControlRigDeveloper/Public/Graph/ControlRigGraph.h
@@ -96,8 +96,6 @@ public:
 	UPROPERTY()
 	bool bIsFunctionDefinition;
 
-	FControlRigPublicFunctionData GetPublicFunctionData() const;
-
 private:
 
 	template<class T>
diff --git a/Engine/Plugins/Animation/ControlRig/Source/ControlRigEditor/Private/ControlRigEditorModule.cpp b/Engine/Plugins/Animation/ControlRig/Source/ControlRigEditor/Private/ControlRigEditorModule.cpp
--- a/Engine/Plugins/Animation/ControlRig/Source/ControlRigEditor/Private/ControlRigEditorModule.cpp
+++ b/Engine/Plugins/Animation/ControlRig/Source/ControlRigEditor/Private/ControlRigEditorModule.cpp
@@ -1175,28 +1175,37 @@ void FControlRigEditorModule::GetTypeActions(UControlRigBlueprint* CRB, FBluepri
 		ActionRegistrar.AddBlueprintAction(ActionKey, NodeSpawner);
 	}
 
-	FArrayProperty* PublicFunctionsProperty = CastField<FArrayProperty>(UControlRigBlueprint::StaticClass()->FindPropertyByName(TEXT("PublicFunctions")));
+	FArrayProperty* PublicFunctionsProperty = CastField<FArrayProperty>(UControlRigBlueprint::StaticClass()->FindPropertyByName(TEXT("PublicGraphFunctions")));
 	if(PublicFunctionsProperty)
 	{
 		// find all control rigs in the project
 		TArray<FAssetData> ControlRigAssetDatas;
 		FARFilter ControlRigAssetFilter;
 		ControlRigAssetFilter.ClassPaths.Add(UControlRigBlueprint::StaticClass()->GetClassPathName());
+		ControlRigAssetFilter.ClassPaths.Add(UControlRigBlueprintGeneratedClass::StaticClass()->GetClassPathName());
 		AssetRegistryModule.Get().GetAssets(ControlRigAssetFilter, ControlRigAssetDatas);
 
 		// loop over all control rigs in the project
+		TSet<FName> PackagesProcessed;
 		for(const FAssetData& ControlRigAssetData : ControlRigAssetDatas)
 		{
+			// Avoid duplication of spawners
+			if (PackagesProcessed.Contains(ControlRigAssetData.PackageName))
+			{
+				continue;
+			}
+			PackagesProcessed.Add(ControlRigAssetData.PackageName);
+			
 			const FString PublicFunctionsString = ControlRigAssetData.GetTagValueRef<FString>(PublicFunctionsProperty->GetFName());
 			if(PublicFunctionsString.IsEmpty())
 			{
 				continue;
 			}
 
-			TArray<FControlRigPublicFunctionData> PublicFunctions;
+			TArray<FRigVMGraphFunctionHeader> PublicFunctions;
 			PublicFunctionsProperty->ImportText_Direct(*PublicFunctionsString, &PublicFunctions, nullptr, EPropertyPortFlags::PPF_None);
 
-			for(const FControlRigPublicFunctionData& PublicFunction : PublicFunctions)
+			for(const FRigVMGraphFunctionHeader& PublicFunction : PublicFunctions)
 			{
 				UBlueprintNodeSpawner* NodeSpawner = UControlRigFunctionRefNodeSpawner::CreateFromAssetData(ControlRigAssetData, PublicFunction);
 				check(NodeSpawner != nullptr);
@@ -1208,7 +1217,7 @@ void FControlRigEditorModule::GetTypeActions(UControlRigBlueprint* CRB, FBluepri
 
 void FControlRigEditorModule::GetInstanceActions(UControlRigBlueprint* CRB, FBlueprintActionDatabaseRegistrar& ActionRegistrar)
 {
-	if (UClass* GeneratedClass = CRB->GetControlRigBlueprintGeneratedClass())
+	if (UControlRigBlueprintGeneratedClass* GeneratedClass = CRB->GetControlRigBlueprintGeneratedClass())
 	{
 		if (UControlRig* CDO = Cast<UControlRig>(GeneratedClass->GetDefaultObject()))
 		{
@@ -2371,19 +2380,23 @@ void FControlRigEditorModule::GetContextMenuActions(const UControlRigGraphSchema
 
 					if (URigVMFunctionReferenceNode* FunctionReferenceNode = Cast<URigVMFunctionReferenceNode>(RigNode->GetModelNode()))
 					{
-						if(FunctionReferenceNode->GetLibrary() != RigBlueprint->GetLocalFunctionLibrary())
+						TSoftObjectPtr<URigVMFunctionReferenceNode> RefPtr(FunctionReferenceNode);
+						if(RefPtr.GetLongPackageName() != FunctionReferenceNode->GetReferencedFunctionHeader().LibraryPointer.LibraryNode.GetLongPackageName())
 						{
-							OrganizationSection.AddMenuEntry(
-                                "Localize Function",
-                                LOCTEXT("LocalizeFunction", "Localize Function"),
-                                LOCTEXT("LocalizeFunction_Tooltip", "Creates a local copy of the function backing the node."),
-                                FSlateIcon(),
-                                FUIAction(FExecuteAction::CreateLambda([RigBlueprint, FunctionReferenceNode]() {
-                                    RigBlueprint->BroadcastRequestLocalizeFunctionDialog(FunctionReferenceNode->GetReferencedNode(), true);
-                                })
-                            ));
-
-							if(!FunctionReferenceNode->IsFullyRemapped())
+							if (URigVMLibraryNode* LibraryNode = FunctionReferenceNode->LoadReferencedNode())
+							{
+								OrganizationSection.AddMenuEntry(
+								   "Localize Function",
+								   LOCTEXT("LocalizeFunction", "Localize Function"),
+								   LOCTEXT("LocalizeFunction_Tooltip", "Creates a local copy of the function backing the node."),
+								   FSlateIcon(),
+								   FUIAction(FExecuteAction::CreateLambda([RigBlueprint, LibraryNode]() {
+									   RigBlueprint->BroadcastRequestLocalizeFunctionDialog(LibraryNode, true);
+								   })
+							   ));
+							}
+
+							if(!FunctionReferenceNode->IsFullyRemapped() && FunctionReferenceNode->GetReferencedFunctionHeader().LibraryPointer.LibraryNode.ResolveObject())
 							{
 								FToolMenuSection& VariablesSection = Menu->AddSection("EdGraphSchemaVariables", LOCTEXT("Variables", "Variables"));
 								VariablesSection.AddMenuEntry(
@@ -2399,29 +2412,32 @@ void FControlRigEditorModule::GetContextMenuActions(const UControlRigGraphSchema
 											FScopedTransaction Transaction(LOCTEXT("MakeVariablesFromFunctionReferenceNode", "Create required variables"));
                                     		RigBlueprint->Modify();
 
-                                    		UControlRigBlueprint* ReferencedBlueprint = FunctionReferenceNode->GetReferencedNode()->GetTypedOuter<UControlRigBlueprint>();
-                                    		// ReferencedBlueprint != RigBlueprint - since only FunctionReferenceNodes from other assets have the potential to be unmapped
-                                    		
-                                    		for(const FRigVMExternalVariable& ExternalVariable : ExternalVariables)
+                                    		if (URigVMLibraryNode* LibraryNode = FunctionReferenceNode->LoadReferencedNode())
                                     		{
-                                    			FString DefaultValue;
-                                    			if(ReferencedBlueprint)
-                                    			{
-                                    				for(const FBPVariableDescription& NewVariable : ReferencedBlueprint->NewVariables)
-                                    				{
-                                    					if(NewVariable.VarName == ExternalVariable.Name)
-                                    					{
-                                    						DefaultValue = NewVariable.DefaultValue;
-                                    						break;
-                                    					}
-                                    				}
-                                    			}
+                                    			UControlRigBlueprint* ReferencedBlueprint = LibraryNode->GetTypedOuter<UControlRigBlueprint>();
+											   // ReferencedBlueprint != RigBlueprint - since only FunctionReferenceNodes from other assets have the potential to be unmapped
+                                    		
+											   for(const FRigVMExternalVariable& ExternalVariable : ExternalVariables)
+											   {
+												   FString DefaultValue;
+												   if(ReferencedBlueprint)
+												   {
+													   for(const FBPVariableDescription& NewVariable : ReferencedBlueprint->NewVariables)
+													   {
+														   if(NewVariable.VarName == ExternalVariable.Name)
+														   {
+															   DefaultValue = NewVariable.DefaultValue;
+															   break;
+														   }
+													   }
+												   }
                                     			
-                                                FName NewVariableName = RigBlueprint->AddCRMemberVariableFromExternal(ExternalVariable, DefaultValue);
-                                    			if(!NewVariableName.IsNone())
-                                    			{
-                                    				Controller->SetRemappedVariable(FunctionReferenceNode, ExternalVariable.Name, NewVariableName);
-                                    			}
+												   FName NewVariableName = RigBlueprint->AddCRMemberVariableFromExternal(ExternalVariable, DefaultValue);
+												   if(!NewVariableName.IsNone())
+												   {
+													   Controller->SetRemappedVariable(FunctionReferenceNode, ExternalVariable.Name, NewVariableName);
+												   }
+											   }
                                     		}
 
                                     		FBlueprintEditorUtils::MarkBlueprintAsModified(RigBlueprint);
diff --git a/Engine/Plugins/Animation/ControlRig/Source/ControlRigEditor/Private/ControlRigGraphDetails.cpp b/Engine/Plugins/Animation/ControlRig/Source/ControlRigEditor/Private/ControlRigGraphDetails.cpp
--- a/Engine/Plugins/Animation/ControlRig/Source/ControlRigEditor/Private/ControlRigGraphDetails.cpp
+++ b/Engine/Plugins/Animation/ControlRig/Source/ControlRigEditor/Private/ControlRigGraphDetails.cpp
@@ -1236,10 +1236,15 @@ FText FControlRigGraphDetails::GetCurrentAccessSpecifierName() const
 		UControlRigGraph* Graph = GraphPtr.Get();
 		UControlRigBlueprint* ControlRigBlueprint = ControlRigBlueprintPtr.Get();
 
-		const FControlRigPublicFunctionData ExpectedFunctionData = Graph->GetPublicFunctionData();
-		if(ControlRigBlueprint->IsFunctionPublic(ExpectedFunctionData.Name))
+		if (URigVMGraph* RigVMGraph = Graph->GetModel())
 		{
-			return FText::FromString(*AccessSpecifierStrings[0].Get()); // public
+			if (URigVMLibraryNode* LibraryNode = RigVMGraph->GetTypedOuter<URigVMLibraryNode>())
+			{
+				if(ControlRigBlueprint->IsFunctionPublic(LibraryNode->GetFName()))
+				{
+					return FText::FromString(*AccessSpecifierStrings[0].Get()); // public
+				}
+			}
 		}
 	}
 
@@ -1252,15 +1257,20 @@ void FControlRigGraphDetails::OnAccessSpecifierSelected( TSharedPtr<FString> Spe
 	{
 		UControlRigGraph* Graph = GraphPtr.Get();
 		UControlRigBlueprint* ControlRigBlueprint = ControlRigBlueprintPtr.Get();
-		const FControlRigPublicFunctionData ExpectedFunctionData = Graph->GetPublicFunctionData();
-		
-		if(SpecifierName->Equals(TEXT("Private")))
-		{
-			ControlRigBlueprint->MarkFunctionPublic(ExpectedFunctionData.Name, false);
-		}
-		else
+
+		if (URigVMGraph* RigVMGraph = Graph->GetModel())
 		{
-			ControlRigBlueprint->MarkFunctionPublic(ExpectedFunctionData.Name, true);
+			if (URigVMLibraryNode* LibraryNode = RigVMGraph->GetTypedOuter<URigVMLibraryNode>())
+			{
+				if(SpecifierName->Equals(TEXT("Private")))
+				{
+					ControlRigBlueprint->MarkFunctionPublic(LibraryNode->GetFName(), false);
+				}
+				else
+				{
+					ControlRigBlueprint->MarkFunctionPublic(LibraryNode->GetFName(), true);
+				}
+			}
 		}
 	}
 }
diff --git a/Engine/Plugins/Animation/ControlRig/Source/ControlRigEditor/Private/Editor/ControlRigEditor.cpp b/Engine/Plugins/Animation/ControlRig/Source/ControlRigEditor/Private/Editor/ControlRigEditor.cpp
--- a/Engine/Plugins/Animation/ControlRig/Source/ControlRigEditor/Private/Editor/ControlRigEditor.cpp
+++ b/Engine/Plugins/Animation/ControlRig/Source/ControlRigEditor/Private/Editor/ControlRigEditor.cpp
@@ -612,8 +612,8 @@ void FControlRigEditor::InitFunctionReferences()
 		return;
 	}
 	
-	FFunctionGraphTask::CreateAndDispatchWhenReady([]()
-	{
+	//FFunctionGraphTask::CreateAndDispatchWhenReady([]()
+	//{
 		if(URigVMBuildData* BuildData = URigVMController::GetBuildData())
 		{
 			const FArrayProperty* ReferenceNodeDataProperty =
@@ -641,14 +641,21 @@ void FControlRigEditor::InitFunctionReferences()
 					TArray<FRigVMReferenceNodeData> ReferenceNodeDatas;
 					ReferenceNodeDataProperty->ImportText_Direct(*ReferenceNodeDataString, &ReferenceNodeDatas, nullptr, EPropertyPortFlags::PPF_None);
 
-					for(const FRigVMReferenceNodeData& ReferenceNodeData : ReferenceNodeDatas)
+					for(FRigVMReferenceNodeData& ReferenceNodeData : ReferenceNodeDatas)
 					{
-						BuildData->RegisterFunctionReference(ReferenceNodeData);
+						if (ReferenceNodeData.ReferencedHeader.IsValid())
+						{
+							BuildData->RegisterFunctionReference(ReferenceNodeData.ReferencedHeader.LibraryPointer, ReferenceNodeData.GetReferenceNodeObjectPath());
+						}
+						else if (!ReferenceNodeData.ReferencedFunctionPath_DEPRECATED.IsEmpty())
+						{
+							BuildData->RegisterFunctionReference(ReferenceNodeData);							
+						}
 					}
 				}
 			}
 		}
-	}, TStatId(), NULL, ENamedThreads::GameThread);
+	//}, TStatId(), NULL, ENamedThreads::GameThread);
 	
 	bAreFunctionReferencesInitialized = true;
 }
@@ -2773,12 +2780,17 @@ void FControlRigEditor::PasteNodes()
 
 	TGuardValue<FRigVMController_RequestLocalizeFunctionDelegate> RequestLocalizeDelegateGuard(
 		GetFocusedController()->RequestLocalizeFunctionDelegate,
-		FRigVMController_RequestLocalizeFunctionDelegate::CreateLambda([this](URigVMLibraryNode* InFunctionToLocalize)
+		FRigVMController_RequestLocalizeFunctionDelegate::CreateLambda([this](TSoftObjectPtr<URigVMLibraryNode> InFunctionToLocalize)
 		{
-			OnRequestLocalizeFunctionDialog(InFunctionToLocalize, GetControlRigBlueprint(), true);
+			if (InFunctionToLocalize.IsValid())
+			{
+				URigVMLibraryNode* LibraryNode = InFunctionToLocalize.Get();
+				OnRequestLocalizeFunctionDialog(LibraryNode, GetControlRigBlueprint(), true);
 
-			const URigVMLibraryNode* LocalizedFunctionNode = GetControlRigBlueprint()->GetLocalFunctionLibrary()->FindPreviouslyLocalizedFunction(InFunctionToLocalize);
-			return LocalizedFunctionNode != nullptr;
+			   const URigVMLibraryNode* LocalizedFunctionNode = GetControlRigBlueprint()->GetLocalFunctionLibrary()->FindPreviouslyLocalizedFunction(LibraryNode);
+			   return LocalizedFunctionNode != nullptr;
+			}
+			return false;
 		})
 	);
 	
@@ -6809,7 +6821,15 @@ void FControlRigEditor::OnNodeDoubleClicked(UControlRigBlueprint* InBlueprint, U
 
 	if (URigVMLibraryNode* LibraryNode = Cast<URigVMLibraryNode>(InNode))
 	{
-		if(URigVMGraph* ContainedGraph = LibraryNode->GetContainedGraph())
+		URigVMGraph* ContainedGraph = LibraryNode->GetContainedGraph();
+		if (URigVMFunctionReferenceNode* FunctionReferenceNode = Cast<URigVMFunctionReferenceNode>(LibraryNode))
+		{
+			if (URigVMLibraryNode* ReferencedNode = FunctionReferenceNode->LoadReferencedNode())
+			{
+				ContainedGraph = ReferencedNode->GetContainedGraph();
+			}
+		}
+		if(ContainedGraph)
 		{
 			if (UEdGraph* EdGraph = InBlueprint->GetEdGraph(ContainedGraph))
 			{
diff --git a/Engine/Plugins/Animation/ControlRig/Source/ControlRigEditor/Private/Editor/SControlRigFunctionLocalizationWidget.cpp b/Engine/Plugins/Animation/ControlRig/Source/ControlRigEditor/Private/Editor/SControlRigFunctionLocalizationWidget.cpp
--- a/Engine/Plugins/Animation/ControlRig/Source/ControlRigEditor/Private/Editor/SControlRigFunctionLocalizationWidget.cpp
+++ b/Engine/Plugins/Animation/ControlRig/Source/ControlRigEditor/Private/Editor/SControlRigFunctionLocalizationWidget.cpp
@@ -102,7 +102,15 @@ void SControlRigFunctionLocalizationWidget::Construct(const FArguments& InArgs,
 	for(int32 NodeToVisitIndex=0; NodeToVisitIndex<NodesToVisit.Num(); NodeToVisitIndex++)
 	{
 		URigVMLibraryNode* NodeToVisit = NodesToVisit[NodeToVisitIndex];
-		const TArray<URigVMNode*>& ContainedNodes = NodeToVisit->GetContainedNodes();
+		TArray<URigVMNode*> ContainedNodes = NodeToVisit->GetContainedNodes();
+		if (URigVMFunctionReferenceNode* ReferenceNode = Cast<URigVMFunctionReferenceNode>(NodeToVisit))
+		{
+			if (ReferenceNode->GetReferencedFunctionHeader().LibraryPointer.LibraryNode.ResolveObject())
+			{
+				URigVMLibraryNode* LibraryNode = Cast<URigVMLibraryNode>(ReferenceNode->GetReferencedFunctionHeader().LibraryPointer.LibraryNode.ResolveObject());
+				ContainedNodes = LibraryNode->GetContainedNodes();
+			}
+		}
 		for(URigVMNode* ContainedNode : ContainedNodes)
 		{
 			if(URigVMLibraryNode* ContainedLibraryNode = Cast<URigVMLibraryNode>(ContainedNode))
@@ -113,8 +121,10 @@ void SControlRigFunctionLocalizationWidget::Construct(const FArguments& InArgs,
 
 		if(URigVMFunctionReferenceNode* FunctionReferenceNode = Cast<URigVMFunctionReferenceNode>(NodeToVisit))
 		{
-			if(URigVMLibraryNode* ReferencedNode = FunctionReferenceNode->GetReferencedNode())
+			FSoftObjectPath ReferencedNodePath = FunctionReferenceNode->GetReferencedFunctionHeader().LibraryPointer.LibraryNode;
+			if(ReferencedNodePath.ResolveObject())
 			{
+				URigVMLibraryNode* ReferencedNode = Cast<URigVMLibraryNode>(ReferencedNodePath.ResolveObject());
 				if(URigVMFunctionLibrary* ReferencedLibrary = ReferencedNode->GetLibrary())
 				{
 					if(UControlRigBlueprint* ReferencedBlueprint = Cast<UControlRigBlueprint>(ReferencedLibrary->GetOuter()))
diff --git a/Engine/Plugins/Animation/ControlRig/Source/ControlRigEditor/Private/Graph/NodeSpawners/ControlRigFunctionRefNodeSpawner.cpp b/Engine/Plugins/Animation/ControlRig/Source/ControlRigEditor/Private/Graph/NodeSpawners/ControlRigFunctionRefNodeSpawner.cpp
--- a/Engine/Plugins/Animation/ControlRig/Source/ControlRigEditor/Private/Graph/NodeSpawners/ControlRigFunctionRefNodeSpawner.cpp
+++ b/Engine/Plugins/Animation/ControlRig/Source/ControlRigEditor/Private/Graph/NodeSpawners/ControlRigFunctionRefNodeSpawner.cpp
@@ -83,44 +83,29 @@ UControlRigFunctionRefNodeSpawner* UControlRigFunctionRefNodeSpawner::CreateFrom
 	return NodeSpawner;
 }
 
-UControlRigFunctionRefNodeSpawner* UControlRigFunctionRefNodeSpawner::CreateFromAssetData(const FAssetData& InAssetData, const FControlRigPublicFunctionData& InPublicFunction)
+UControlRigFunctionRefNodeSpawner* UControlRigFunctionRefNodeSpawner::CreateFromAssetData(const FAssetData& InAssetData, const FRigVMGraphFunctionHeader& InPublicFunction)
 {
 	UControlRigFunctionRefNodeSpawner* NodeSpawner = NewObject<UControlRigFunctionRefNodeSpawner>(GetTransientPackage());
-PRAGMA_DISABLE_DEPRECATION_WARNINGS
-	NodeSpawner->ReferencedAssetObjectPath = InAssetData.ToSoftObjectPath().ToFName();
-PRAGMA_ENABLE_DEPRECATION_WARNINGS
-	NodeSpawner->ReferencedPublicFunctionData = InPublicFunction;
+	NodeSpawner->ReferencedPublicFunctionHeader = InPublicFunction;
 	NodeSpawner->NodeClass = UControlRigGraphNode::StaticClass();
 	NodeSpawner->bIsLocalFunction = false;
 
 	FBlueprintActionUiSpec& MenuSignature = NodeSpawner->DefaultMenuSignature;
 
-	const FString Category = InPublicFunction.Category;
-
 	MenuSignature.MenuName = FText::FromName(InPublicFunction.Name);
-	MenuSignature.Category = FText::FromString(Category);
+	MenuSignature.Category = FText::FromString(InPublicFunction.Category);
 	MenuSignature.Keywords = FText::FromString(InPublicFunction.Keywords);
+	MenuSignature.Tooltip = InPublicFunction.Tooltip;
 
-	if(const UControlRigBlueprint* ReferencedBlueprint = Cast<UControlRigBlueprint>(InAssetData.FastGetAsset(false)))
-	{
-		if(const URigVMFunctionLibrary* FunctionLibrary = ReferencedBlueprint->GetLocalFunctionLibrary())
-		{
-			if(const URigVMLibraryNode* FunctionNode = FunctionLibrary->FindFunction(InPublicFunction.Name))
-			{
-				MenuSignature.Tooltip = FunctionNode->GetToolTipText();
-			}
-		}
-	}
-
-	const FString ObjectPathString = InAssetData.GetObjectPathString();
+	const FString PackagePathString = InAssetData.PackageName.ToString();
 	if(MenuSignature.Tooltip.IsEmpty())
 	{
-		MenuSignature.Tooltip = FText::FromString(ObjectPathString);
+		MenuSignature.Tooltip = FText::FromString(PackagePathString);
 	}
 	else
 	{
 		static constexpr TCHAR Format[] = TEXT("%s\n\n%s");
-		MenuSignature.Tooltip = FText::FromString(FString::Printf(Format, *MenuSignature.Tooltip.ToString(), *ObjectPathString));
+		MenuSignature.Tooltip = FText::FromString(FString::Printf(Format, *MenuSignature.Tooltip.ToString(), *PackagePathString));
 	}
 
 	// add at least one character, so that PrimeDefaultUiSpec() doesn't 
@@ -156,25 +141,11 @@ void UControlRigFunctionRefNodeSpawner::Prime()
 FBlueprintNodeSignature UControlRigFunctionRefNodeSpawner::GetSpawnerSignature() const
 {
 	FString SignatureString = TEXT("Invalid RigFunction");
-	if(ReferencedFunctionPtr.IsValid())
-	{
-		if(UObject* FunctionLibrary = ReferencedFunctionPtr->GetOuter())
-		{
-			if(UObject* Blueprint = FunctionLibrary->GetOuter())
-			{
-				SignatureString = 
-					FString::Printf(TEXT("RigFunction=%s::%s"),
-					*Blueprint->GetPathName(),
-					*ReferencedFunctionPtr->GetName());
-			}
-		}
-	}
-	else if(ReferencedAssetObjectPath.IsValid() && ReferencedPublicFunctionData.Name.IsValid())
+	if (!ReferencedPublicFunctionHeader.IsValid())
 	{
 		SignatureString = 
-			FString::Printf(TEXT("RigFunction=%s::%s"),
-			*ReferencedAssetObjectPath.ToString(),
-			*ReferencedPublicFunctionData.Name.ToString());
+			FString::Printf(TEXT("RigFunction=%s"),
+			*ReferencedPublicFunctionHeader.GetHash());
 	}
 
 	if(bIsLocalFunction)
@@ -198,34 +169,18 @@ UEdGraphNode* UControlRigFunctionRefNodeSpawner::Invoke(UEdGraph* ParentGraph, F
 {
 	UControlRigGraphNode* NewNode = nullptr;
 
-	// if we are trying to build the real function ref - but we haven't loaded the asset yet...
-	if(!FBlueprintNodeTemplateCache::IsTemplateOuter(ParentGraph))
+	if (!ReferencedFunctionPtr.ToString().IsEmpty())
 	{
-		if(!ReferencedFunctionPtr.IsValid() && !ReferencedAssetObjectPath.IsNone() && !ReferencedPublicFunctionData.Name.IsNone())
+		if (URigVMLibraryNode* LibraryNode = ReferencedFunctionPtr.LoadSynchronous())
 		{
-			UControlRigBlueprint* ReferencedBlueprint = LoadObject<UControlRigBlueprint>(nullptr, *ReferencedAssetObjectPath.ToString());
-			if(ReferencedBlueprint == nullptr)
+			if (URigVMFunctionLibrary* FunctionLibrary = LibraryNode->GetGraph()->GetDefaultFunctionLibrary())
 			{
-				return nullptr;
+				ReferencedPublicFunctionHeader = LibraryNode->GetFunctionHeader();
 			}
-
-			URigVMFunctionLibrary* FunctionLibrary = ReferencedBlueprint->GetLocalFunctionLibrary();
-			if(FunctionLibrary == nullptr)
-			{
-				return nullptr;
-			}
-
-			URigVMLibraryNode* FunctionNode = FunctionLibrary->FindFunction(ReferencedPublicFunctionData.Name);
-			if(FunctionNode == nullptr)
-			{
-				return nullptr;
-			}
-
-			ReferencedFunctionPtr = FunctionNode;
 		}
 	}
 
-	if(ReferencedFunctionPtr.IsValid())
+	if(ReferencedPublicFunctionHeader.IsValid())
 	{
 #if WITH_EDITOR
 		if (GEditor)
@@ -235,47 +190,13 @@ UEdGraphNode* UControlRigFunctionRefNodeSpawner::Invoke(UEdGraph* ParentGraph, F
 #endif
 
 		UBlueprint* Blueprint = FBlueprintEditorUtils::FindBlueprintForGraph(ParentGraph);
-		NewNode = SpawnNode(ParentGraph, Blueprint, ReferencedFunctionPtr.Get(), Location);
-	}
-	else
-	{
-		// we are only going to get here if we are spawning a template node
-		NewNode = NewObject<UControlRigGraphNode>(ParentGraph);
-		ParentGraph->AddNode(NewNode, false);
-
-		NewNode->CreateNewGuid();
-		NewNode->PostPlacedNewNode();
-
-		for(const FControlRigPublicFunctionArg& Arg : ReferencedPublicFunctionData.Arguments)
-		{
-			if(Arg.Direction ==  ERigVMPinDirection::Input ||
-				Arg.Direction ==  ERigVMPinDirection::IO)
-			{
-				UEdGraphPin* InputPin = UEdGraphPin::CreatePin(NewNode);
-				NewNode->Pins.Add(InputPin);
-
-				InputPin->Direction = EGPD_Input;
-				InputPin->PinType = Arg.GetPinType();
-			}
-
-			if(Arg.Direction ==  ERigVMPinDirection::Output ||
-				Arg.Direction ==  ERigVMPinDirection::IO)
-			{
-				UEdGraphPin* OutputPin = UEdGraphPin::CreatePin(NewNode);
-				NewNode->Pins.Add(OutputPin);
-
-				OutputPin->Direction = EGPD_Output;
-				OutputPin->PinType = Arg.GetPinType();
-			}
-		}
-
-		NewNode->SetFlags(RF_Transactional);
+		NewNode = SpawnNode(ParentGraph, Blueprint, ReferencedPublicFunctionHeader, Location);
 	}
 
 	return NewNode;
 }
 
-UControlRigGraphNode* UControlRigFunctionRefNodeSpawner::SpawnNode(UEdGraph* ParentGraph, UBlueprint* Blueprint, URigVMLibraryNode* InFunction, FVector2D const Location)
+UControlRigGraphNode* UControlRigFunctionRefNodeSpawner::SpawnNode(UEdGraph* ParentGraph, UBlueprint* Blueprint, FRigVMGraphFunctionHeader& InFunction, FVector2D const Location)
 {
 	UControlRigGraphNode* NewNode = nullptr;
 	UControlRigBlueprint* RigBlueprint = Cast<UControlRigBlueprint>(Blueprint);
@@ -286,7 +207,7 @@ UControlRigGraphNode* UControlRigFunctionRefNodeSpawner::SpawnNode(UEdGraph* Par
 		bool const bIsTemplateNode = FBlueprintNodeTemplateCache::IsTemplateOuter(ParentGraph);
 		bool const bIsUserFacingNode = !bIsTemplateNode;
 
-		FName Name = bIsTemplateNode ? *InFunction->GetName() : FControlRigBlueprintUtils::ValidateName(RigBlueprint, InFunction->GetName());
+		FName Name = bIsTemplateNode ? InFunction.Name : FControlRigBlueprintUtils::ValidateName(RigBlueprint, InFunction.Name.ToString());
 		URigVMController* Controller = bIsTemplateNode ? RigGraph->GetTemplateController() : RigBlueprint->GetController(ParentGraph);
 
 		if (!bIsTemplateNode)
@@ -294,7 +215,7 @@ UControlRigGraphNode* UControlRigFunctionRefNodeSpawner::SpawnNode(UEdGraph* Par
 			Controller->OpenUndoBracket(FString::Printf(TEXT("Add '%s' Node"), *Name.ToString()));
 		}
 
-		if (URigVMFunctionReferenceNode* ModelNode = Controller->AddFunctionReferenceNode(InFunction, Location, Name.ToString(), bIsUserFacingNode, !bIsTemplateNode))
+		if (URigVMFunctionReferenceNode* ModelNode = Controller->AddFunctionReferenceNodeFromDescription(InFunction, Location, Name.ToString(), bIsUserFacingNode, !bIsTemplateNode, bIsTemplateNode))
 		{
 			NewNode = Cast<UControlRigGraphNode>(RigGraph->FindNodeForModelNodeName(ModelNode->GetFName()));
 			check(NewNode);
@@ -322,7 +243,7 @@ UControlRigGraphNode* UControlRigFunctionRefNodeSpawner::SpawnNode(UEdGraph* Par
 						continue;
 					}
 
-					if(ExistingFunctionReferenceNode->GetReferencedNode() != InFunction)
+					if(!(ExistingFunctionReferenceNode->GetReferencedFunctionHeader().LibraryPointer == InFunction.LibraryPointer))
 					{
 						continue;
 					}
@@ -380,15 +301,12 @@ bool UControlRigFunctionRefNodeSpawner::IsTemplateNodeFilteredOut(FBlueprintActi
 			}
 		}
 	}
-	else if(ReferencedAssetObjectPath.IsValid())
+	const FString ReferencedAssetObjectPathString = ReferencedFunctionPtr.GetAssetName();
+	for (UBlueprint* Blueprint : Filter.Context.Blueprints)
 	{
-		const FString ReferencedAssetObjectPathString = ReferencedAssetObjectPath.ToString();
-		for (UBlueprint* Blueprint : Filter.Context.Blueprints)
+		if(Blueprint->GetPathName() == ReferencedAssetObjectPathString)
 		{
-			if(Blueprint->GetPathName() == ReferencedAssetObjectPathString)
-			{
-				return true;
-			}
+			return true;
 		}
 	}
 	return false;
diff --git a/Engine/Plugins/Animation/ControlRig/Source/ControlRigEditor/Private/Graph/NodeSpawners/ControlRigFunctionRefNodeSpawner.h b/Engine/Plugins/Animation/ControlRig/Source/ControlRigEditor/Private/Graph/NodeSpawners/ControlRigFunctionRefNodeSpawner.h
--- a/Engine/Plugins/Animation/ControlRig/Source/ControlRigEditor/Private/Graph/NodeSpawners/ControlRigFunctionRefNodeSpawner.h
+++ b/Engine/Plugins/Animation/ControlRig/Source/ControlRigEditor/Private/Graph/NodeSpawners/ControlRigFunctionRefNodeSpawner.h
@@ -41,7 +41,7 @@ public:
 	 * 
 	 * @return A newly allocated instance of this class.
 	 */
-	static UControlRigFunctionRefNodeSpawner* CreateFromAssetData(const FAssetData& InAssetData, const FControlRigPublicFunctionData& InPublicFunction);
+	static UControlRigFunctionRefNodeSpawner* CreateFromAssetData(const FAssetData& InAssetData, const FRigVMGraphFunctionHeader& InPublicFunction);
 
 	// UBlueprintNodeSpawner interface
 	virtual void Prime() override;
@@ -55,21 +55,17 @@ private:
 
 	/** The unit type we will spawn [optional] */
 	UPROPERTY(Transient)
-	mutable TWeakObjectPtr<URigVMLibraryNode> ReferencedFunctionPtr;
-
-	/** The asset object path we'll spawn from [optional] */
-	UPROPERTY(Transient)
-	FName ReferencedAssetObjectPath;
+	mutable TSoftObjectPtr<URigVMLibraryNode> ReferencedFunctionPtr;
 
 	/** The public function definition we will spawn from [optional] */
 	UPROPERTY(Transient)
-	FControlRigPublicFunctionData ReferencedPublicFunctionData;
+	mutable FRigVMGraphFunctionHeader ReferencedPublicFunctionHeader;
 
 	/** Marked as true for local function definitions */
 	UPROPERTY(Transient)
 	bool bIsLocalFunction;
 
-	static UControlRigGraphNode* SpawnNode(UEdGraph* ParentGraph, UBlueprint* Blueprint, URigVMLibraryNode* InFunction, FVector2D const Location);
+	static UControlRigGraphNode* SpawnNode(UEdGraph* ParentGraph, UBlueprint* Blueprint, FRigVMGraphFunctionHeader& InFunction, FVector2D const Location);
 
 	friend class UEngineTestControlRig;
 	friend class FControlRigEditor;
diff --git a/Engine/Plugins/Animation/ControlRig/Source/ControlRigEditor/Private/Graph/SControlRigGraphPinVariableBinding.cpp b/Engine/Plugins/Animation/ControlRig/Source/ControlRigEditor/Private/Graph/SControlRigGraphPinVariableBinding.cpp
--- a/Engine/Plugins/Animation/ControlRig/Source/ControlRigEditor/Private/Graph/SControlRigGraphPinVariableBinding.cpp
+++ b/Engine/Plugins/Animation/ControlRig/Source/ControlRigEditor/Private/Graph/SControlRigGraphPinVariableBinding.cpp
@@ -181,7 +181,7 @@ bool SControlRigVariableBinding::OnCanBindProperty(FProperty* InProperty) const
 		}
 		else if(FunctionReferenceNode && !InnerVariableName.IsNone())
 		{
-			TArray<FRigVMExternalVariable> InnerVariables = FunctionReferenceNode->GetContainedGraph()->GetExternalVariables();
+			TArray<FRigVMExternalVariable> InnerVariables = FunctionReferenceNode->GetReferencedFunctionHeader().ExternalVariables;
 			for(const FRigVMExternalVariable& InnerVariable : InnerVariables)
 			{
 				if(InnerVariable.Name == InnerVariableName)
diff --git a/Engine/Plugins/Experimental/DataInterfaceGraph/Source/DataInterfaceGraphUncookedOnly/Internal/DataInterfaceGraph_EditorData.h b/Engine/Plugins/Experimental/DataInterfaceGraph/Source/DataInterfaceGraphUncookedOnly/Internal/DataInterfaceGraph_EditorData.h
--- a/Engine/Plugins/Experimental/DataInterfaceGraph/Source/DataInterfaceGraphUncookedOnly/Internal/DataInterfaceGraph_EditorData.h
+++ b/Engine/Plugins/Experimental/DataInterfaceGraph/Source/DataInterfaceGraphUncookedOnly/Internal/DataInterfaceGraph_EditorData.h
@@ -40,6 +40,8 @@ class UDataInterfaceGraph_EditorData : public UObject, public IRigVMClientHost
 	// IRigVMClientHost interface
 	virtual FRigVMClient* GetRigVMClient() override;
 	virtual const FRigVMClient* GetRigVMClient() const override;
+	virtual IRigVMGraphFunctionHost* GetRigVMGraphFunctionHost() override {return nullptr; }
+	virtual const IRigVMGraphFunctionHost* GetRigVMGraphFunctionHost() const override {return nullptr; }
 	virtual void HandleRigVMGraphAdded(const FRigVMClient* InClient, const FString& InNodePath) override;
 	virtual void HandleRigVMGraphRemoved(const FRigVMClient* InClient, const FString& InNodePath) override {}
 	virtual void HandleRigVMGraphRenamed(const FRigVMClient* InClient, const FString& InOldNodePath, const FString& InNewNodePath) override {}
diff --git a/Engine/Plugins/Experimental/DataInterfaceGraph/Source/DataInterfaceGraphUncookedOnly/Private/DataInterfaceGraph_EditorData.cpp b/Engine/Plugins/Experimental/DataInterfaceGraph/Source/DataInterfaceGraphUncookedOnly/Private/DataInterfaceGraph_EditorData.cpp
--- a/Engine/Plugins/Experimental/DataInterfaceGraph/Source/DataInterfaceGraphUncookedOnly/Private/DataInterfaceGraph_EditorData.cpp
+++ b/Engine/Plugins/Experimental/DataInterfaceGraph/Source/DataInterfaceGraphUncookedOnly/Private/DataInterfaceGraph_EditorData.cpp
@@ -172,16 +172,6 @@ void UDataInterfaceGraph_EditorData::HandleConfigureRigVMController(const FRigVM
 
 	});
 
-	InControllerToConfigure->IsFunctionAvailableDelegate.BindLambda([](URigVMLibraryNode* InFunction) -> bool
-	{
-		return true;	// @TODO: should only allow main entry point function here
-	});
-
-	InControllerToConfigure->IsDependencyCyclicDelegate.BindLambda([](UObject* InDependentObject, UObject* InDependencyObject) -> bool
-	{
-		return false;
-	});
-
 #if WITH_EDITOR
 	InControllerToConfigure->SetupDefaultUnitNodeDelegates(TDelegate<FName(FRigVMExternalVariable, FString)>::CreateLambda(
 		[](FRigVMExternalVariable InVariableToCreate, FString InDefaultValue) -> FName
diff --git a/Engine/Source/Developer/RigVMDeveloper/Private/RigVMCompiler/RigVMAST.cpp b/Engine/Source/Developer/RigVMDeveloper/Private/RigVMCompiler/RigVMAST.cpp
--- a/Engine/Source/Developer/RigVMDeveloper/Private/RigVMCompiler/RigVMAST.cpp
+++ b/Engine/Source/Developer/RigVMDeveloper/Private/RigVMCompiler/RigVMAST.cpp
@@ -1845,7 +1845,8 @@ bool FRigVMParserAST::FoldConstantValuesToLiterals(TArray<URigVMGraph*> InGraphs
 	for (const FRigVMASTProxy& NodeProxy : NodeProxies)
 	{
 		if (NodeProxy.IsA<URigVMVariableNode>() ||
-			NodeProxy.IsA<URigVMEnumNode>())
+			NodeProxy.IsA<URigVMEnumNode>() ||
+			NodeProxy.IsA<URigVMFunctionReferenceNode>())
 		{
 			continue;
 		}
@@ -1918,6 +1919,7 @@ bool FRigVMParserAST::FoldConstantValuesToLiterals(TArray<URigVMGraph*> InGraphs
 
 				if (SourceNodeProxy.IsA<URigVMVariableNode>() ||
 					SourceNodeProxy.IsA<URigVMRerouteNode>() ||
+					SourceNodeProxy.IsA<URigVMFunctionReferenceNode>() ||
 					SourceNodeProxy.IsA<URigVMEnumNode>())
 				{
 					continue;
@@ -1955,19 +1957,19 @@ bool FRigVMParserAST::FoldConstantValuesToLiterals(TArray<URigVMGraph*> InGraphs
 				NodesToCompute.AddUnique(SourceProxy);
 			}
 
-			// if the node is a library node - also add the return node
-			if(URigVMLibraryNode* LibraryToCompute = Cast<URigVMLibraryNode>(NodeToCompute))
+			// if the node is a collapse node - also add the return node
+			if(URigVMCollapseNode* CollapseToCompute = Cast<URigVMCollapseNode>(NodeToCompute))
 			{
-				if(URigVMNode* ReturnNodeToCompute = LibraryToCompute->GetReturnNode())
+				if(URigVMNode* ReturnNodeToCompute = CollapseToCompute->GetReturnNode())
 				{
-					const FRigVMASTProxy LibraryProxy = ProxyToCompute.GetSibling(LibraryToCompute);
-					const FRigVMASTProxy ReturnProxy = LibraryProxy.GetChild(ReturnNodeToCompute);
+					const FRigVMASTProxy CollapseProxy = ProxyToCompute.GetSibling(CollapseToCompute);
+					const FRigVMASTProxy ReturnProxy = CollapseProxy.GetChild(ReturnNodeToCompute);
 					NodesToCompute.AddUnique(ReturnProxy);
 				}
 			}
       
-			// if the node is an entry node - also add host library node
-			if(NodeToCompute->IsA<URigVMFunctionEntryNode>())
+			// if the node is an entry node - also add host collapse node
+			if(NodeToCompute->IsA<URigVMFunctionEntryNode>() && NodeToCompute->GetGraph()->GetOuter()->IsA<URigVMCollapseNode>())
 			{
 				check(ProxyToCompute.GetCallstack().Num() > 1);
 				const FRigVMASTProxy LibraryProxy = ProxyToCompute.GetParent();
@@ -1976,6 +1978,18 @@ bool FRigVMParserAST::FoldConstantValuesToLiterals(TArray<URigVMGraph*> InGraphs
 		}
 	}
 
+	// validate nodes to compute
+	for(int32 NodeToComputeIndex = 0; NodeToComputeIndex < NodesToCompute.Num(); NodeToComputeIndex++)
+	{
+		FRigVMASTProxy ProxyToCompute = NodesToCompute[NodeToComputeIndex];
+		if(URigVMNode* NodeToCompute = ProxyToCompute.GetSubject<URigVMNode>())
+		{
+			// All mutable nodes should be considered varying, so should never arrive to this point
+			check(!NodeToCompute->IsMutable());			
+		}
+	}
+	
+
 	// we now know the node we need to run.
 	// let's build an temporary AST which has only those nodes
 	TSharedPtr<FRigVMParserAST> TempAST = MakeShareable(new FRigVMParserAST(InGraphs, NodesToCompute));
@@ -2734,7 +2748,7 @@ FString FRigVMParserAST::DumpDot() const
 				{
 					if (URigVMFunctionReferenceNode* Node = Cast<URigVMFunctionReferenceNode>(InExpr->To<FRigVMInlineFunctionExprAST>()->GetNode()))
 					{
-						Label = FString::Printf(TEXT("Inline %s"), *Node->GetReferencedNode()->GetName());
+						Label = FString::Printf(TEXT("Inline %s"), *Node->GetReferencedFunctionHeader().Name.ToString());
 					}
 					break;
 				}
@@ -3465,52 +3479,47 @@ void FRigVMParserAST::Inline(TArray<URigVMGraph*> InGraphs, const TArray<FRigVMA
 				}
 			}
 
-			if (URigVMLibraryNode* LibraryNode = InNodeProxy.GetSubject<URigVMLibraryNode>())
+			if (URigVMCollapseNode* LibraryNode = InNodeProxy.GetSubject<URigVMCollapseNode>())
 			{
 				if (LibraryNode->GetContainedGraph() == nullptr)
 				{
-					if (URigVMFunctionReferenceNode* FunctionRef = Cast<URigVMFunctionReferenceNode>(LibraryNode))
-					{
-						FString FunctionPath = FunctionRef->ReferencedNodePtr.ToString();
-
-						OutTraversalInfo.Settings->Reportf(
-							EMessageSeverity::Error, 
-							LibraryNode, 
-							TEXT("Function Reference '%s' references a missing function (%s)."), 
-							*LibraryNode->GetName(),
-							*FunctionPath);
-					}
-					else
-					{
-						OutTraversalInfo.Settings->Reportf(
-							EMessageSeverity::Error,
-							LibraryNode,
-							TEXT("Library Node '%s' doesn't contain a subgraph."),
-							*LibraryNode->GetName());
-					}
+					OutTraversalInfo.Settings->Reportf(
+						EMessageSeverity::Error,
+						LibraryNode,
+						TEXT("Library Node '%s' doesn't contain a subgraph."),
+						*LibraryNode->GetName());
 					return;
 				}
 
-				if (URigVMFunctionReferenceNode* FunctionRef = Cast<URigVMFunctionReferenceNode>(LibraryNode))
+				OutTraversalInfo.LibraryNodeCallstack.Push(InNodeProxy);
+				TArray<URigVMNode*> ContainedNodes = LibraryNode->GetContainedNodes();
+				for (URigVMNode* ContainedNode : ContainedNodes)
 				{
-					for (URigVMPin* Pin : FunctionRef->GetPins())
-					{
-						FRigVMASTProxy PinProxy = InNodeProxy.GetSibling(Pin);
-						LocalPinTraversalInfo::VisitPin(PinProxy, OutTraversalInfo);
-					}
+					// create a proxy which uses the previous node as a callstack
+					FRigVMASTProxy ContainedNodeProxy = InNodeProxy.GetChild(ContainedNode);
+					VisitNode(ContainedNodeProxy, OutTraversalInfo);
 				}
-				else
+				OutTraversalInfo.LibraryNodeCallstack.Pop();
+			}
+			else if (URigVMFunctionReferenceNode* FuncRefNode = InNodeProxy.GetSubject<URigVMFunctionReferenceNode>())
+			{
+				if (FuncRefNode->GetReferencedFunctionData() == nullptr)
 				{
-					OutTraversalInfo.LibraryNodeCallstack.Push(InNodeProxy);
-					TArray<URigVMNode*> ContainedNodes = LibraryNode->GetContainedNodes();
-					for (URigVMNode* ContainedNode : ContainedNodes)
-					{
-						// create a proxy which uses the previous node as a callstack
-						FRigVMASTProxy ContainedNodeProxy = InNodeProxy.GetChild(ContainedNode);
-						VisitNode(ContainedNodeProxy, OutTraversalInfo);
-					}
-					OutTraversalInfo.LibraryNodeCallstack.Pop();
-				}				
+					FString FunctionPath = FuncRefNode->GetReferencedFunctionHeader().GetHash();
+
+					OutTraversalInfo.Settings->Reportf(
+						EMessageSeverity::Error, 
+						FuncRefNode, 
+						TEXT("Function Reference '%s' references a missing function (%s)."), 
+						*FuncRefNode->GetName(),
+						*FunctionPath);
+				}
+					
+				for (URigVMPin* Pin : FuncRefNode->GetPins())
+				{
+					FRigVMASTProxy PinProxy = InNodeProxy.GetSibling(Pin);
+					LocalPinTraversalInfo::VisitPin(PinProxy, OutTraversalInfo);
+				}
 			}
 			else
 			{
diff --git a/Engine/Source/Developer/RigVMDeveloper/Private/RigVMCompiler/RigVMASTProxy.cpp b/Engine/Source/Developer/RigVMDeveloper/Private/RigVMCompiler/RigVMASTProxy.cpp
--- a/Engine/Source/Developer/RigVMDeveloper/Private/RigVMCompiler/RigVMASTProxy.cpp
+++ b/Engine/Source/Developer/RigVMDeveloper/Private/RigVMCompiler/RigVMASTProxy.cpp
@@ -4,6 +4,8 @@
 #include "RigVMModel/RigVMGraph.h"
 #include "RigVMModel/RigVMNode.h"
 #include "RigVMModel/RigVMPin.h"
+#include "RigVMModel/Nodes/RigVMCollapseNode.h"
+#include "RigVMModel/Nodes/RigVMFunctionReferenceNode.h"
 #include "RigVMModel/Nodes/RigVMLibraryNode.h"
 
 FString FRigVMCallstack::GetCallPath(bool bIncludeLast) const
@@ -147,9 +149,16 @@ FRigVMASTProxy FRigVMASTProxy::MakeFromCallPath(const FString& InCallPath, UObje
 			Right.Empty();
 		}
 
-		if(URigVMLibraryNode* LibraryNode = Cast<URigVMLibraryNode>(ParentObject))
+		if(URigVMCollapseNode* CollapseNode = Cast<URigVMCollapseNode>(ParentObject))
 		{
-			ParentObject = LibraryNode->GetContainedGraph();
+			ParentObject = CollapseNode->GetContainedGraph();
+		}
+		else if(URigVMFunctionReferenceNode* ReferenceNode = Cast<URigVMFunctionReferenceNode>(ParentObject))
+		{
+			if (URigVMLibraryNode* LibraryNode = ReferenceNode->LoadReferencedNode())
+			{
+				ParentObject = LibraryNode->GetContainedGraph();
+			}
 		}
 
 		if(URigVMGraph* Graph = Cast<URigVMGraph>(ParentObject))
diff --git a/Engine/Source/Developer/RigVMDeveloper/Private/RigVMCompiler/RigVMCompiler.cpp b/Engine/Source/Developer/RigVMDeveloper/Private/RigVMCompiler/RigVMCompiler.cpp
--- a/Engine/Source/Developer/RigVMDeveloper/Private/RigVMCompiler/RigVMCompiler.cpp
+++ b/Engine/Source/Developer/RigVMDeveloper/Private/RigVMCompiler/RigVMCompiler.cpp
@@ -10,6 +10,9 @@
 #include "UObject/Interface.h"
 #include "Stats/StatsHierarchical.h"
 #include "RigVMTypeUtils.h"
+#include "RigVMCore/RigVMGraphFunctionDefinition.h"
+#include "ProfilingDebugging/ScopedTimers.h"
+#include "RigVMModel/RigVMClient.h"
 #include "Algo/Count.h"
 
 #include UE_INLINE_GENERATED_CPP_BY_NAME(RigVMCompiler)
@@ -242,8 +245,11 @@ URigVMCompiler::URigVMCompiler()
 {
 }
 
-bool URigVMCompiler::Compile(TArray<URigVMGraph*> InGraphs, URigVMController* InController, URigVM* OutVM, const TArray<FRigVMExternalVariable>& InExternalVariables, const TArray<FRigVMUserDataArray>& InRigVMUserData, TMap<FString, FRigVMOperand>* OutOperands, TSharedPtr<FRigVMParserAST> InAST)
+bool URigVMCompiler::Compile(TArray<URigVMGraph*> InGraphs, URigVMController* InController, URigVM* OutVM, const TArray<FRigVMExternalVariable>& InExternalVariables, const TArray<FRigVMUserDataArray>& InRigVMUserData, TMap<FString, FRigVMOperand>* OutOperands, TSharedPtr<FRigVMParserAST> InAST, FRigVMFunctionCompilationData* OutFunctionCompilationData)
 {
+	double CompilationTime = 0;
+	FDurationTimer CompileTimer(CompilationTime);
+	
 	if (InGraphs.IsEmpty() || InGraphs.Contains(nullptr))
 	{
 		ReportError(TEXT("Provided graph is nullptr."));
@@ -290,115 +296,107 @@ bool URigVMCompiler::Compile(TArray<URigVMGraph*> InGraphs, URigVMController* In
 
 	URigVMFunctionLibrary* FunctionLibrary = InGraphs[0]->GetDefaultFunctionLibrary();
 	bool bEncounteredGraphError = false;
-	if (CurrentCompilationFunction == nullptr)
-	{
-		ClearFunctionCompilationData();
 
-		// We need to compile the functions in order of dependency (leaf functions with no dependency first)
-		TArray<URigVMLibraryNode*> FunctionsToCompile;
-		if (FunctionLibrary)
-		{
-			FunctionsToCompile = FunctionLibrary->GetFunctions();
-		}
+	TMap<FString, const FRigVMFunctionCompilationData*> CurrentCompiledFunctions;
 
-		// Add public functions from other assets
-		// todo: this should be a check if the compilation data exists
+	// Gather function compilation data
+	{
+		TArray<URigVMNode*> Nodes = InGraphs[0]->GetNodes();
+		for (int32 i=0; i<Nodes.Num(); ++i)
 		{
-			TArray<URigVMNode*> Nodes = InGraphs[0]->GetNodes();
-			for (int32 i=0; i<Nodes.Num(); ++i)
+			if (URigVMFunctionReferenceNode* ReferenceNode = Cast<URigVMFunctionReferenceNode>(Nodes[i]))
 			{
-				if (URigVMFunctionReferenceNode* ReferenceNode = Cast<URigVMFunctionReferenceNode>(Nodes[i]))
+				if (!ReferenceNode->GetReferencedFunctionHeader().IsValid())
 				{
-					FunctionsToCompile.AddUnique(ReferenceNode->GetReferencedNode());
-				}
-				if (URigVMCollapseNode* CollapseNode = Cast<URigVMCollapseNode>(Nodes[i]))
-				{
-					Nodes.Append(CollapseNode->GetContainedGraph()->GetNodes());
+					static const FString FunctionCompilationErrorMessage = TEXT("Function reference @@ has no function data.");
+					Settings.ASTSettings.Report(EMessageSeverity::Error, ReferenceNode, FunctionCompilationErrorMessage);
+					bEncounteredGraphError = true;
+					break;
 				}
-			}
-		}
-		while (CompiledFunctions.Num() < FunctionsToCompile.Num())
-		{
-			for (int32 FunctionIndex=0; FunctionIndex<FunctionsToCompile.Num(); ++FunctionIndex)
-			{
-				if (URigVMLibraryNode* FunctionToCompile = FunctionsToCompile[FunctionIndex])
+				
+				// Try to find the compiled data
+				FString FunctionHash = ReferenceNode->GetReferencedFunctionHeader().GetHash();
+				if (!CurrentCompiledFunctions.Contains(FunctionHash))
 				{
-					if (CompiledFunctions.Contains(FunctionToCompile))
-					{
-						continue;
-					}
-			
-					bool bCanBeCompiled = true;
-					TArray<URigVMNode*> Nodes = FunctionToCompile->GetContainedNodes();
-					for (int32 i=0; i<Nodes.Num() && bCanBeCompiled; ++i)
+					if (FRigVMGraphFunctionData* FunctionData = ReferenceNode->GetReferencedFunctionHeader().GetFunctionData())
 					{
-						if (URigVMFunctionReferenceNode* FunctionReference = Cast<URigVMFunctionReferenceNode>(Nodes[i]))
+						// Clear compilation data if compiled with outdated dependency data
+						if (FunctionData->CompilationData.IsValid())
 						{
-							if (const URigVMFunctionLibrary* ReferencedFunctionLibrary = FunctionReference->GetReferencedNode()->GetLibrary())
+							for (const TPair<FRigVMGraphFunctionIdentifier, uint32>& Pair : ReferenceNode->GetReferencedFunctionHeader().Dependencies)
 							{
-								if (ReferencedFunctionLibrary == FunctionLibrary)
+								if (IRigVMGraphFunctionHost* HostObj = Cast<IRigVMGraphFunctionHost>(Pair.Key.HostObject.ResolveObject()))
 								{
-									if (!CompiledFunctions.Contains(FunctionReference->GetReferencedNode()))
+									if (FRigVMGraphFunctionData* DependencyData = HostObj->GetRigVMGraphFunctionStore()->FindFunction(Pair.Key))
 									{
-										bCanBeCompiled = false;
-										break;
+										if (Pair.Value != DependencyData->CompilationData.Hash)
+										{
+											FunctionData->ClearCompilationData();
+											break;
+										}
 									}
 								}
+							}
+						}
+						
+						if (const FRigVMFunctionCompilationData* CompilationData = &FunctionData->CompilationData)
+						{
+							bool bSuccessfullCompilation = false;
+							if (!CompilationData->IsValid())
+							{
+								if (URigVMLibraryNode* LibraryNode = Cast<URigVMLibraryNode>(FunctionData->Header.LibraryPointer.LibraryNode.TryLoad()))
+								{
+									IRigVMClientHost* ClientHost = LibraryNode->GetImplementingOuter<IRigVMClientHost>();
+									URigVMController* FunctionController = ClientHost->GetRigVMClient()->GetController(LibraryNode->GetLibrary());
+									bSuccessfullCompilation = CompileFunction(LibraryNode, FunctionController, &FunctionData->CompilationData);
+								}
 								else
 								{
-									// Public function.
-									// todo Try to find its compilation data
-									bool bFoundCompilationData = true;
-									if (!CompiledFunctions.Contains(FunctionReference->GetReferencedNode()))
-									{
-										FunctionsToCompile.AddUnique(FunctionReference->GetReferencedNode());
-										bCanBeCompiled = false;
-										break;
-									}
-
-									if (!bFoundCompilationData)
-									{
-										static const FString FunctionCompilationErrorMessage = TEXT("Could not find compilation data for public function @@.");
-										Settings.ASTSettings.Report(EMessageSeverity::Error, FunctionReference, FunctionCompilationErrorMessage);
-										bEncounteredGraphError = true;
-										bCanBeCompiled = false;
-										break;
-									}
+									static const FString FunctionCompilationErrorMessage = TEXT("Compilation data for public function @@ has no instructions.");
+									Settings.ASTSettings.Report(EMessageSeverity::Error, ReferenceNode, FunctionCompilationErrorMessage);
+									bEncounteredGraphError = true;
 								}
 							}
-						}
-
-						if (bCanBeCompiled)
-						{
-							if (URigVMLibraryNode* LibraryNode = Cast<URigVMLibraryNode>(Nodes[i]))
+							if (bSuccessfullCompilation || CompilationData->IsValid())
 							{
-								Nodes.Append(LibraryNode->GetContainedNodes());	
+								CurrentCompiledFunctions.Add(FunctionHash, CompilationData);
+							}
+							else
+							{
+								static const FString FunctionCompilationErrorMessage = TEXT("Compilation data for public function @@ has no instructions.");
+								Settings.ASTSettings.Report(EMessageSeverity::Error, ReferenceNode, FunctionCompilationErrorMessage);
+								bEncounteredGraphError = true;
 							}
 						}
-					}
-
-					if (bCanBeCompiled)
-					{
-						if (!CompileFunction(FunctionToCompile, InController, UserData))
+						else
 						{
-							static const FString FunctionCompilationErrorMessage = TEXT("Function @@ failed to compile.");
-							Settings.ASTSettings.Report(EMessageSeverity::Error, FunctionToCompile, FunctionCompilationErrorMessage);
+							static const FString FunctionCompilationErrorMessage = TEXT("Could not find compilation data for node @@.");
+							Settings.ASTSettings.Report(EMessageSeverity::Error, ReferenceNode, FunctionCompilationErrorMessage);
 							bEncounteredGraphError = true;
-							break;
 						}
 					}
+					else
+					{
+						static const FString FunctionCompilationErrorMessage = TEXT("Could not find graph function data for node @@.");
+						Settings.ASTSettings.Report(EMessageSeverity::Error, ReferenceNode, FunctionCompilationErrorMessage);
+						bEncounteredGraphError = true;
+					}
 				}
 			}
-
-			FunctionsToCompile.RemoveAll([](const URigVMLibraryNode* Node) { return Node == nullptr; });
-			
-			if(bEncounteredGraphError)
+			if (URigVMCollapseNode* CollapseNode = Cast<URigVMCollapseNode>(Nodes[i]))
 			{
-				return false;
+				Nodes.Append(CollapseNode->GetContainedGraph()->GetNodes());
 			}
 		}
 	}
 
+	if (bEncounteredGraphError)
+	{
+		return false;
+	}
+
+	CompiledFunctions = CurrentCompiledFunctions;
+
 #if WITH_EDITOR
 
 	// traverse all graphs and try to clear out orphan pins
@@ -452,9 +450,10 @@ bool URigVMCompiler::Compile(TArray<URigVMGraph*> InGraphs, URigVMController* In
 						bEncounteredGraphError = true;
 					}
 
-					if (const FunctionCompilationData* CompilationData = CompiledFunctions.Find(FunctionReferenceNode->GetReferencedNode()))
+					FString FunctionHash = FunctionReferenceNode->GetReferencedFunctionHeader().GetHash();
+					if (const FRigVMFunctionCompilationData** CompilationData = CompiledFunctions.Find(FunctionHash))
 					{
-						for (const TPair<int32, FName>& Pair : CompilationData->ExternalRegisterIndexToVariable)
+						for (const TPair<int32, FName>& Pair : (*CompilationData)->ExternalRegisterIndexToVariable)
 						{
 							FName OuterName = FunctionReferenceNode->GetOuterVariableName(Pair.Value);
 							if (OuterName.IsNone())
@@ -821,21 +820,84 @@ bool URigVMCompiler::Compile(TArray<URigVMGraph*> InGraphs, URigVMController* In
 	}
 
 	// Store function compile data
-	if (CurrentCompilationFunction)
+	if (CurrentCompilationFunction && OutFunctionCompilationData)
 	{
-		FunctionCompilationData Data;
-		Data.ByteCode = WorkData.VM->ByteCodeStorage;
-		Data.FunctionNames = WorkData.VM->FunctionNamesStorage;
-		Data.PropertyDescriptions = WorkData.PropertyDescriptions;
-		Data.PropertyPathDescriptions = WorkData.PropertyPathDescriptions;
-		Data.Operands = *OutOperands;
+		OutFunctionCompilationData->ByteCode = WorkData.VM->ByteCodeStorage;
+		OutFunctionCompilationData->FunctionNames = WorkData.VM->FunctionNamesStorage;
+		OutFunctionCompilationData->Operands = *OutOperands;
+
+		for (uint8 MemoryTypeIndex=0; MemoryTypeIndex<(uint8)ERigVMMemoryType::Invalid; ++MemoryTypeIndex)
+		{
+			TArray<FRigVMFunctionCompilationPropertyDescription>* PropertyDescriptions = nullptr;
+			TArray<FRigVMFunctionCompilationPropertyPath>* PropertyPathDescriptions = nullptr;
+			ERigVMMemoryType MemoryType = (ERigVMMemoryType) MemoryTypeIndex;
+			switch (MemoryType)
+			{
+				case ERigVMMemoryType::Work:
+				{
+					PropertyDescriptions = &OutFunctionCompilationData->WorkPropertyDescriptions;	
+					PropertyPathDescriptions = &OutFunctionCompilationData->WorkPropertyPathDescriptions;
+					break;
+				}
+				case ERigVMMemoryType::Literal:
+				{
+					PropertyDescriptions = &OutFunctionCompilationData->LiteralPropertyDescriptions;	
+					PropertyPathDescriptions = &OutFunctionCompilationData->LiteralPropertyPathDescriptions;
+					break;
+				}
+				case ERigVMMemoryType::External:
+				{
+					PropertyDescriptions = &OutFunctionCompilationData->ExternalPropertyDescriptions;	
+					PropertyPathDescriptions = &OutFunctionCompilationData->ExternalPropertyPathDescriptions;
+					break;
+				}
+				case ERigVMMemoryType::Debug:
+				{
+					PropertyDescriptions = &OutFunctionCompilationData->DebugPropertyDescriptions;	
+					PropertyPathDescriptions = &OutFunctionCompilationData->DebugPropertyPathDescriptions;
+					break;
+				}
+				default:
+				{
+					checkNoEntry();
+				}
+			}
+
+			PropertyDescriptions->Reset();
+			PropertyPathDescriptions->Reset();
+			if (const TArray<FRigVMPropertyDescription>* Descriptions = WorkData.PropertyDescriptions.Find(MemoryType))
+			{
+				PropertyDescriptions->Reserve(Descriptions->Num());
+				for (const FRigVMPropertyDescription& Description : (*Descriptions))
+				{
+					FRigVMFunctionCompilationPropertyDescription NewDescription;
+					NewDescription.Name = Description.Name;
+					NewDescription.CPPType = Description.CPPType;
+					NewDescription.CPPTypeObject = Description.CPPTypeObject;
+					NewDescription.DefaultValue = Description.DefaultValue;
+					PropertyDescriptions->Add(NewDescription);
+				}
+			}
+			if (const TArray<FRigVMPropertyPathDescription>* PathDescriptions = WorkData.PropertyPathDescriptions.Find(MemoryType))
+			{
+				PropertyPathDescriptions->Reserve(PathDescriptions->Num());
+				for (const FRigVMPropertyPathDescription& Description : (*PathDescriptions))
+				{
+					FRigVMFunctionCompilationPropertyPath NewDescription;
+					NewDescription.PropertyIndex = Description.PropertyIndex;
+					NewDescription.SegmentPath = Description.SegmentPath;
+					NewDescription.HeadCPPType = Description.HeadCPPType;
+					PropertyPathDescriptions->Add(NewDescription);
+				}
+			}
+		}
 
 		// Only add used external registers to the function compilation data
-		FRigVMInstructionArray Instructions = Data.ByteCode.GetInstructions();
+		FRigVMInstructionArray Instructions = OutFunctionCompilationData->ByteCode.GetInstructions();
 		TSet<int32> UsedExternalVariableRegisters;
 		for (const FRigVMInstruction& Instruction : Instructions)
 		{
-			const FRigVMOperandArray OperandArray = Data.ByteCode.GetOperandsForOp(Instruction);
+			const FRigVMOperandArray OperandArray = OutFunctionCompilationData->ByteCode.GetOperandsForOp(Instruction);
 			for (const FRigVMOperand& Operand : OperandArray)
 			{
 				if (Operand.GetMemoryType() == ERigVMMemoryType::External)
@@ -854,75 +916,81 @@ bool URigVMCompiler::Compile(TArray<URigVMGraph*> InGraphs, URigVMController* In
 				if (UsedExternalVariableRegisters.Contains(Pair.Value.GetRegisterIndex()))
 				{
 					FString VariableName = Pair.Key.RightChop(VariablePrefix.Len());
-					Data.ExternalRegisterIndexToVariable.Add(Pair.Value.GetRegisterIndex(), *VariableName);
+					OutFunctionCompilationData->ExternalRegisterIndexToVariable.Add(Pair.Value.GetRegisterIndex(), *VariableName);
 				}
 			}
 		}
-		CompiledFunctions.Add(CurrentCompilationFunction, Data);
 	}
 
+	if (!CurrentCompilationFunction)
+	{
+		CompileTimer.Stop();
+		ReportInfof(TEXT("Total Compilation time %f\n"), CompilationTime*1000);
+	}
+
+	FPlatformMisc::LowLevelOutputDebugStringf(TEXT("%s\n\n"), *WorkData.VM->DumpByteCodeAsText());
+
 	return true;
 }
 
-bool URigVMCompiler::CompileFunction(URigVMLibraryNode* InLibraryNode, URigVMController* InController, const TArray<FRigVMUserDataArray>& InRigVMUserData)
+bool URigVMCompiler::CompileFunction(const URigVMLibraryNode* InLibraryNode, URigVMController* InController, FRigVMFunctionCompilationData* OutFunctionCompilationData)
 {
-	TGuardValue<URigVMLibraryNode*> CompilationGuard(CurrentCompilationFunction, InLibraryNode);
+	TGuardValue<const URigVMLibraryNode*> CompilationGuard(CurrentCompilationFunction, InLibraryNode);
 	FRigVMControllerGraphGuard ControllerGraphGuard(InController, InLibraryNode->GetContainedGraph(), false);
 
+	double CompilationTime = 0;
+	FDurationTimer CompileTimer(CompilationTime);
+
+	if (OutFunctionCompilationData == nullptr)
+	{
+		return false;
+	}
+
+	OutFunctionCompilationData->ByteCode.Reset();
+
 	TArray<FRigVMExternalVariable> ExternalVariables;
 	if (InController->GetExternalVariablesDelegate.IsBound())
 	{
 		ExternalVariables = InController->GetExternalVariablesDelegate.Execute(InLibraryNode->GetContainedGraph());
 	}
 	TMap<FString, FRigVMOperand> Operands;
+	TArray<FRigVMUserDataArray> UserData;
+	UserData.Add(FRigVMUserDataArray());
+	
 	URigVM* TempVM = NewObject<URigVM>(InLibraryNode->GetContainedGraph());
-	const bool bSuccess = Compile({InLibraryNode->GetContainedGraph()}, InController, TempVM, ExternalVariables, InRigVMUserData, &Operands);
+	const bool bSuccess = Compile({InLibraryNode->GetContainedGraph()}, InController, TempVM, ExternalVariables, UserData, &Operands, nullptr, OutFunctionCompilationData);
 	TempVM->ClearMemory();
 	TempVM->Rename(nullptr, GetTransientPackage(), REN_ForceNoResetLoaders | REN_DoNotDirty | REN_DontCreateRedirectors | REN_NonTransactional);
 	TempVM->MarkAsGarbage();
 
-	return bSuccess;
-}
-
-void URigVMCompiler::ClearFunctionCompilationData()
-{
-	TArray<TSoftObjectPtr<URigVMLibraryNode>> Functions;
-	CompiledFunctions.GetKeys(Functions);
-	for (TSoftObjectPtr<URigVMLibraryNode> Function : Functions)
-	{
-		RemoveFunctionCompilationData(Function);
-	}
-	CompiledFunctions.Reset();
-}
-
-void URigVMCompiler::RemoveFunctionCompilationData(TSoftObjectPtr<URigVMLibraryNode> LibraryNode)
-{
-	FunctionCompilationData* CompilationData = CompiledFunctions.Find(LibraryNode);
-	if (CompilationData)
-	{
-		CompiledFunctions.Remove(LibraryNode);
-	}
-
-	if (!LibraryNode.IsValid())
-	{
-		return;
-	}
+	CompileTimer.Stop();
+	ReportInfof(TEXT("Compiled Function %s in %fms"), *InLibraryNode->GetName(), CompilationTime*1000);
 
-	// Clear the compilation data of any function referencing the removed function
-	URigVMFunctionLibrary* Library = LibraryNode->GetLibrary();
-	URigVMBuildData* BuildData = Library->GetBuildData();
-	if (const FRigVMFunctionReferenceArray* FunctionReferences = BuildData->FindFunctionReferences(LibraryNode.Get()))
+	// Update the compilation data of this library, and the hashes of the compilation data of its dependencies used for this compilation
+	if (IRigVMClientHost* ClientHost = InLibraryNode->GetImplementingOuter<IRigVMClientHost>())
 	{
-		for (int32 i=0; i<FunctionReferences->Num(); ++i)
+		if (IRigVMGraphFunctionHost* FunctionHost = ClientHost->GetRigVMGraphFunctionHost())
 		{
-			const TSoftObjectPtr<URigVMFunctionReferenceNode> Reference = FunctionReferences->operator[](i);
-			if (Reference.IsValid())
+			if (FRigVMGraphFunctionStore* Store = FunctionHost->GetRigVMGraphFunctionStore())
 			{
-				URigVMLibraryNode* ReferencedNode = Reference->GetReferencedNode();
-				// todo: remove compilation data of this function
+				if (FRigVMGraphFunctionData* Data = Store->FindFunction(InLibraryNode->GetFunctionIdentifier()))
+				{
+					for(TPair<FRigVMGraphFunctionIdentifier, uint32>& Pair : Data->Header.Dependencies)
+					{
+						IRigVMGraphFunctionHost* ReferencedFunctionHost = Cast<IRigVMGraphFunctionHost>(Pair.Key.HostObject.ResolveObject());
+						if (FRigVMGraphFunctionData* ReferencedData = ReferencedFunctionHost->GetRigVMGraphFunctionStore()->FindFunction(Pair.Key))
+						{
+							Pair.Value = ReferencedData->CompilationData.Hash;
+						}
+					}
+				}
+		
+				Store->UpdateFunctionCompilationData(InLibraryNode->GetFunctionIdentifier(), *OutFunctionCompilationData);
 			}
 		}
 	}
+
+	return bSuccess;
 }
 
 void URigVMCompiler::TraverseExpression(const FRigVMExprAST* InExpr, FRigVMCompilerWorkData& WorkData)
@@ -1335,37 +1403,86 @@ int32 URigVMCompiler::TraverseInlineFunction(const FRigVMInlineFunctionExprAST*
 	int32 InstructionIndexStart = INDEX_NONE;
 	int32 InstructionIndexEnd = INDEX_NONE;
 
-	const URigVMLibraryNode* ReferencedLibrary = FunctionReferenceNode->GetReferencedNode();
-	if (!CompiledFunctions.Contains(ReferencedLibrary))
+	
+	FString FunctionHash = FunctionReferenceNode->GetReferencedFunctionHeader().GetHash();
+	if (!CompiledFunctions.Contains(FunctionHash))
 	{
 		return INDEX_NONE;
 	}
-	const FunctionCompilationData& FunctionCompilationData = CompiledFunctions.FindChecked(ReferencedLibrary);
-	const FRigVMByteCode& FunctionByteCode = FunctionCompilationData.ByteCode;
+	const FRigVMFunctionCompilationData* FunctionCompilationData = CompiledFunctions.FindChecked(FunctionHash);
+	const FRigVMByteCode& FunctionByteCode = FunctionCompilationData->ByteCode;
 	
 	if (WorkData.bSetupMemory)
 	{
 		TraverseChildren(InExpr, WorkData);
 
 		// Add internal operands (not the ones represented by interface pins)
-		TArray<ERigVMMemoryType> MemoryTypes;
-		FunctionCompilationData.PropertyDescriptions.GetKeys(MemoryTypes);
-		for (const ERigVMMemoryType& MemoryType : MemoryTypes)
+		for (uint8 MemoryIndex=0; MemoryIndex< (uint8)ERigVMMemoryType::Invalid; ++MemoryIndex)
 		{
-			const TArray<FRigVMPropertyDescription>& Properties = FunctionCompilationData.PropertyDescriptions[MemoryType];
-			const int32 NumArguments = Algo::CountIf(ReferencedLibrary->GetPins(), [](const URigVMPin* Pin) -> bool
+			ERigVMMemoryType MemoryType = (ERigVMMemoryType) MemoryIndex;
+			TArray<FRigVMFunctionCompilationPropertyDescription> Properties;
+			switch (MemoryType)
 			{
-				return !Pin->IsExecuteContext();
-			});
- 
-			int32 StartIndex = MemoryType == ERigVMMemoryType::Work ? NumArguments : 0;
+				case ERigVMMemoryType::Work:
+				{
+					Properties = FunctionCompilationData->WorkPropertyDescriptions;
+					break;
+				}
+				case ERigVMMemoryType::Literal:
+				{
+					Properties = FunctionCompilationData->LiteralPropertyDescriptions;
+					break;
+				}
+				case ERigVMMemoryType::External:
+				{
+					Properties = FunctionCompilationData->ExternalPropertyDescriptions;
+					break;
+				}
+				case ERigVMMemoryType::Debug:
+				{
+					Properties = FunctionCompilationData->DebugPropertyDescriptions;
+					break;
+				}
+			}
+
+			int32 NumProperties = 0;
+			if (MemoryType == ERigVMMemoryType::Work)
+			{
+				for (const FRigVMGraphFunctionArgument& Argument : FunctionReferenceNode->GetReferencedFunctionHeader().Arguments)
+				{
+					if (Argument.CPPTypeObject.IsValid())
+					{
+						if (const UScriptStruct* ScriptStruct = Cast<UScriptStruct>(Argument.CPPTypeObject.Get()))
+						{
+							if (ScriptStruct->IsChildOf(FRigVMExecuteContext::StaticStruct()))
+							{
+								continue;
+							}
+						}
+					}
+					if (!Properties.IsValidIndex(NumProperties))
+					{
+						continue;
+					}
+					if (!Properties[NumProperties].Name.ToString().Contains(Argument.Name.ToString()))
+					{
+						continue;
+					}
+					NumProperties++;
+				}
+			}
+			
+			int32 StartIndex = MemoryType == ERigVMMemoryType::Work ? NumProperties : 0;
 			for (int32 PropertyIndex = StartIndex; PropertyIndex < Properties.Num(); ++PropertyIndex)
 			{
-				const FRigVMPropertyDescription& Description = Properties[PropertyIndex];
+				const FRigVMFunctionCompilationPropertyDescription& Description = Properties[PropertyIndex];
 				FString NewName = Description.Name.ToString();
 				static const FString FunctionLibraryPrefix = TEXT("FunctionLibrary");
-				NewName = FString::Printf(TEXT("%s%s"), *FunctionReferenceNode->GetNodePath(), *NewName.RightChop(FunctionLibraryPrefix.Len()));
-				FRigVMOperand Operand = WorkData.AddProperty(MemoryType, *NewName, Description.CPPType, Description.CPPTypeObject, Description.DefaultValue);
+				if (NewName.StartsWith(FunctionLibraryPrefix))
+				{
+					NewName = FString::Printf(TEXT("%s%s"), *FunctionReferenceNode->GetNodePath(), *NewName.RightChop(FunctionLibraryPrefix.Len()));
+				}
+				FRigVMOperand Operand = WorkData.AddProperty(MemoryType, *NewName, Description.CPPType, Description.CPPTypeObject.Get(), Description.DefaultValue);
 				FRigVMCompilerWorkData::FFunctionRegisterData Data = {FunctionReferenceNode, MemoryType, PropertyIndex};
 				WorkData.FunctionRegisterToOperand.Add(Data, Operand);
 
@@ -1428,7 +1545,7 @@ int32 URigVMCompiler::TraverseInlineFunction(const FRigVMInlineFunctionExprAST*
 				// Remap the variable: find the operand index of the outer variable
 				if (Operand->GetMemoryType() == ERigVMMemoryType::External)
 				{
-					const FName& InnerVariableName = FunctionCompilationData.ExternalRegisterIndexToVariable[Operand->GetRegisterIndex()];
+					const FName& InnerVariableName = FunctionCompilationData->ExternalRegisterIndexToVariable[Operand->GetRegisterIndex()];
 					FName OuterVariableName = InnerVariableName;
 					if (const FName* VariableRemapped = FunctionReferenceNode->GetVariableMap().Find(InnerVariableName))
 					{
@@ -1461,7 +1578,31 @@ int32 URigVMCompiler::TraverseInlineFunction(const FRigVMInlineFunctionExprAST*
 				// For all operands, check to see if we need to add a property path
 				if (Operand->GetRegisterOffset() != INDEX_NONE)
 				{
-					const FRigVMPropertyPathDescription& Description = FunctionCompilationData.PropertyPathDescriptions[OriginalMemoryType][Operand->GetRegisterOffset()];
+					
+					FRigVMFunctionCompilationPropertyPath Description;
+					switch (OriginalMemoryType)
+					{
+						case ERigVMMemoryType::Work:
+						{
+							Description = FunctionCompilationData->WorkPropertyPathDescriptions[Operand->GetRegisterOffset()];
+							break;
+						}
+						case ERigVMMemoryType::Literal:
+						{
+							Description = FunctionCompilationData->LiteralPropertyPathDescriptions[Operand->GetRegisterOffset()];
+							break;
+						}
+						case ERigVMMemoryType::External:
+						{
+							Description = FunctionCompilationData->ExternalPropertyPathDescriptions[Operand->GetRegisterOffset()];
+							break;
+						}
+						case ERigVMMemoryType::Debug:
+						{
+							Description = FunctionCompilationData->DebugPropertyPathDescriptions[Operand->GetRegisterOffset()];
+							break;
+						}
+					}
 					Operand->RegisterOffset = WorkData.FindOrAddPropertyPath(*Operand, Description.HeadCPPType, Description.SegmentPath);
 				}
 			}
@@ -1469,7 +1610,7 @@ int32 URigVMCompiler::TraverseInlineFunction(const FRigVMInlineFunctionExprAST*
 			if (Instruction.OpCode >= ERigVMOpCode::Execute_0_Operands && Instruction.OpCode <= ERigVMOpCode::Execute_64_Operands)
 			{
 				FRigVMExecuteOp& Op = ByteCode.GetOpAt<FRigVMExecuteOp>(Instruction);
-				const int32 FunctionIndex = WorkData.VM->AddRigVMFunction(FunctionCompilationData.FunctionNames[Op.FunctionIndex].ToString());
+				const int32 FunctionIndex = WorkData.VM->AddRigVMFunction(FunctionCompilationData->FunctionNames[Op.FunctionIndex].ToString());
 				Op.FunctionIndex = FunctionIndex;
 			}
 
@@ -2587,7 +2728,7 @@ void URigVMCompiler::AddCopyOperator(
 	AddCopyOperator(CopyOpInfo.Op, CopyOpInfo.AssignExpr, CopyOpInfo.SourceExpr, CopyOpInfo.TargetExpr, WorkData, bDelayCopyOperations);
 }
 
-FString URigVMCompiler::GetPinHashImpl(const URigVMPin* InPin, const FRigVMVarExprAST* InVarExpr, bool bIsDebugValue, URigVMLibraryNode* FunctionCompiling, const FRigVMASTProxy& InPinProxy)
+FString URigVMCompiler::GetPinHashImpl(const URigVMPin* InPin, const FRigVMVarExprAST* InVarExpr, bool bIsDebugValue, const URigVMLibraryNode* FunctionCompiling, const FRigVMASTProxy& InPinProxy)
 {
 	FString Prefix = bIsDebugValue ? TEXT("DebugWatch:") : TEXT("");
 	FString Suffix;
@@ -2776,7 +2917,7 @@ FString URigVMCompiler::GetPinHashImpl(const URigVMPin* InPin, const FRigVMVarEx
 	return FString::Printf(TEXT("%s%s%s"), *Prefix, *PinPath, *Suffix);
 }
 
-FString URigVMCompiler::GetPinHash(const URigVMPin* InPin, const FRigVMVarExprAST* InVarExpr, bool bIsDebugValue, URigVMLibraryNode* FunctionCompiling, const FRigVMASTProxy& InPinProxy)
+FString URigVMCompiler::GetPinHash(const URigVMPin* InPin, const FRigVMVarExprAST* InVarExpr, bool bIsDebugValue, const URigVMLibraryNode* FunctionCompiling, const FRigVMASTProxy& InPinProxy)
 {
 	const FString Hash = GetPinHashImpl(InPin, InVarExpr, bIsDebugValue, FunctionCompiling, InPinProxy);
 	if(!bIsDebugValue && FunctionCompiling == nullptr)
diff --git a/Engine/Source/Developer/RigVMDeveloper/Private/RigVMModel/Nodes/RigVMFunctionReferenceNode.cpp b/Engine/Source/Developer/RigVMDeveloper/Private/RigVMModel/Nodes/RigVMFunctionReferenceNode.cpp
--- a/Engine/Source/Developer/RigVMDeveloper/Private/RigVMModel/Nodes/RigVMFunctionReferenceNode.cpp
+++ b/Engine/Source/Developer/RigVMDeveloper/Private/RigVMModel/Nodes/RigVMFunctionReferenceNode.cpp
@@ -2,72 +2,56 @@
 
 #include "RigVMModel/Nodes/RigVMFunctionReferenceNode.h"
 #include "RigVMModel/RigVMFunctionLibrary.h"
+#include "RigVMCore/RigVMGraphFunctionDefinition.h"
+#include "RigVMCore/RigVMGraphFunctionHost.h"
 
 #include UE_INLINE_GENERATED_CPP_BY_NAME(RigVMFunctionReferenceNode)
 
 FString URigVMFunctionReferenceNode::GetNodeTitle() const
 {
-	if (URigVMLibraryNode* ReferencedNode = GetReferencedNode())
+	if (const FRigVMGraphFunctionData* Data = GetReferencedFunctionData())
 	{
-		return ReferencedNode->GetNodeTitle();
+		return Data->Header.NodeTitle;
 	}
 	return Super::GetNodeTitle();
 }
 
 FLinearColor URigVMFunctionReferenceNode::GetNodeColor() const
 {
-	if (URigVMLibraryNode* ReferencedNode = GetReferencedNode())
+	if (const FRigVMGraphFunctionData* Data = GetReferencedFunctionData())
 	{
-		return ReferencedNode->GetNodeColor();
+		return Data->Header.NodeColor;
 	}
 	return Super::GetNodeColor();
 }
 
 FText URigVMFunctionReferenceNode::GetToolTipText() const
 {
-	if (URigVMLibraryNode* ReferencedNode = GetReferencedNode())
+	if (const FRigVMGraphFunctionData* Data = GetReferencedFunctionData())
 	{
-		return ReferencedNode->GetToolTipText();
+		return Data->Header.Tooltip;
 	}
 	return Super::GetToolTipText();
 }
 
 FString URigVMFunctionReferenceNode::GetNodeCategory() const
 {
-	if (URigVMLibraryNode* ReferencedNode = GetReferencedNode())
+	if (const FRigVMGraphFunctionData* Data = GetReferencedFunctionData())
 	{
-		return ReferencedNode->GetNodeCategory();
+		return Data->Header.Category;
 	}
 	return Super::GetNodeCategory();
 }
 
 FString URigVMFunctionReferenceNode::GetNodeKeywords() const
 {
-	if (URigVMLibraryNode* ReferencedNode = GetReferencedNode())
+	if (const FRigVMGraphFunctionData* Data = GetReferencedFunctionData())
 	{
-		return ReferencedNode->GetNodeKeywords();
+		return Data->Header.Keywords;
 	}
 	return Super::GetNodeKeywords();
 }
 
-URigVMFunctionLibrary* URigVMFunctionReferenceNode::GetLibrary() const
-{
-	if(URigVMLibraryNode* ReferencedNode = GetReferencedNode())
-	{
-		return ReferencedNode->GetLibrary();
-	}
-	return nullptr;
-}
-
-URigVMGraph* URigVMFunctionReferenceNode::GetContainedGraph() const
-{
-	if (URigVMLibraryNode* ReferencedNode = GetReferencedNode())
-	{
-		return ReferencedNode->GetContainedGraph();
-	}
-	return nullptr;
-}
-
 bool URigVMFunctionReferenceNode::RequiresVariableRemapping() const
 {
 	TArray<FRigVMExternalVariable> InnerVariables;
@@ -77,62 +61,15 @@ bool URigVMFunctionReferenceNode::RequiresVariableRemapping() const
 bool URigVMFunctionReferenceNode::RequiresVariableRemappingInternal(TArray<FRigVMExternalVariable>& InnerVariables) const
 {
 	bool bHostedInDifferencePackage = false;
-	if(URigVMLibraryNode* ReferencedNode = GetReferencedNode())
-	{
-		// we only need to remap variables if we are referencing
-		// a function which is hosted in another asset.
-		const UObject* ReferencedPackage = ReferencedNode->GetOutermost(); 
-		const UObject* Package = GetOutermost();
-		if(ReferencedPackage != Package)
-		{
-			bHostedInDifferencePackage = true;
-		}
-		else
-		{
-			// this case might happen within unit tests
-			// where the graphs are not parented to a package
-			const UObject* TransientPackage = GetTransientPackage();
-			if((ReferencedPackage == TransientPackage) && (Package == TransientPackage))
-			{
-				const UObject* ReferencedOuter = ReferencedNode;
-				while(ReferencedOuter)
-				{
-					const UObject* Outer = ReferencedOuter->GetOuter();
-					if(Outer != TransientPackage)
-					{
-						ReferencedOuter = Outer;
-					}
-					else
-					{
-						break;
-					}
-				}
-
-				const UObject* CurrentOuter = this;
-				while(CurrentOuter)
-				{
-					const UObject* Outer = CurrentOuter->GetOuter();
-					if(Outer != TransientPackage)
-					{
-						CurrentOuter = Outer;
-					}
-					else
-					{
-						break;
-					}
-				}
-
-				if(CurrentOuter != ReferencedOuter)
-				{
-					bHostedInDifferencePackage = true;
-				}
-			}
-		}
-	}
-
+	
+	FRigVMGraphFunctionIdentifier LibraryPointer = ReferencedFunctionHeader.LibraryPointer;
+	const FString& LibraryPackagePath = LibraryPointer.LibraryNode.GetLongPackageName();
+	const FString& ThisPacakgePath = GetPackage()->GetPathName();
+	bHostedInDifferencePackage = LibraryPackagePath != ThisPacakgePath;
+		
 	if(bHostedInDifferencePackage)
 	{
-		InnerVariables = Super::GetExternalVariables();
+		InnerVariables = ReferencedFunctionHeader.ExternalVariables;
 		if(InnerVariables.Num() == 0)
 		{
 			return false;
@@ -169,22 +106,13 @@ TArray<FRigVMExternalVariable> URigVMFunctionReferenceNode::GetExternalVariables
 	return GetExternalVariables(true);
 }
 
-const FRigVMTemplate* URigVMFunctionReferenceNode::GetTemplate() const
-{
-	if (URigVMLibraryNode* ReferencedNode = GetReferencedNode())
-	{
-		return ReferencedNode->GetTemplate();
-	}
-	return nullptr;
-}
-
 TArray<FRigVMExternalVariable> URigVMFunctionReferenceNode::GetExternalVariables(bool bRemapped) const
 {
 	TArray<FRigVMExternalVariable> Variables;
 	
 	if(!bRemapped)
 	{
-		Variables = Super::GetExternalVariables();
+		Variables = ReferencedFunctionHeader.ExternalVariables;
 	}
 	else
 	{
@@ -214,37 +142,43 @@ FName URigVMFunctionReferenceNode::GetOuterVariableName(const FName& InInnerVari
 	return NAME_None;
 }
 
+const FRigVMGraphFunctionData* URigVMFunctionReferenceNode::GetReferencedFunctionData() const
+{
+	if (IRigVMGraphFunctionHost* Host = ReferencedFunctionHeader.GetFunctionHost())
+	{
+		return Host->GetRigVMGraphFunctionStore()->FindFunction(ReferencedFunctionHeader.LibraryPointer);
+	}
+	return nullptr;
+}
+
 FText URigVMFunctionReferenceNode::GetToolTipTextForPin(const URigVMPin* InPin) const
 {
 	check(InPin);
-	
-	if(URigVMLibraryNode* ReferencedNode = GetReferencedNode())
+
+	URigVMPin* RootPin = InPin->GetRootPin();
+	const FRigVMGraphFunctionArgument* Argument = ReferencedFunctionHeader.Arguments.FindByPredicate([RootPin](const FRigVMGraphFunctionArgument& Argument)
 	{
-		if(URigVMPin* ReferencedPin = ReferencedNode->FindPin(InPin->GetSegmentPath(true)))
+		return Argument.Name == RootPin->GetFName();
+	});
+
+	if (Argument)
+	{
+		if (const FText* Tooltip = Argument->PathToTooltip.Find(InPin->GetSegmentPath(false)))
 		{
-			const FString DefaultValue = ReferencedPin->GetDefaultValue();
-			if(!DefaultValue.IsEmpty())
-			{
-				return FText::FromString(FString::Printf(TEXT("%s\nDefault %s"),
-					*ReferencedPin->GetToolTipText().ToString(),
-					*DefaultValue));
-			}
+			return *Tooltip;
 		}
 	}
-
+	
 	return Super::GetToolTipTextForPin(InPin);
 }
 
-URigVMLibraryNode* URigVMFunctionReferenceNode::GetReferencedNode() const
+URigVMLibraryNode* URigVMFunctionReferenceNode::LoadReferencedNode() const
 {
-	if (!ReferencedNodePtr.IsValid())
+	UObject* LibraryNode = ReferencedFunctionHeader.LibraryPointer.LibraryNode.ResolveObject();
+	if (!LibraryNode)
 	{
-		ReferencedNodePtr.LoadSynchronous();
+		LibraryNode = ReferencedFunctionHeader.LibraryPointer.LibraryNode.TryLoad();
 	}
-	return ReferencedNodePtr.Get();
-}
-
-void URigVMFunctionReferenceNode::SetReferencedNode(URigVMLibraryNode* InReferenceNode)
-{
-	ReferencedNodePtr = InReferenceNode;
+	return Cast<URigVMLibraryNode>(LibraryNode);
+	
 }
diff --git a/Engine/Source/Developer/RigVMDeveloper/Private/RigVMModel/Nodes/RigVMLibraryNode.cpp b/Engine/Source/Developer/RigVMDeveloper/Private/RigVMModel/Nodes/RigVMLibraryNode.cpp
--- a/Engine/Source/Developer/RigVMDeveloper/Private/RigVMModel/Nodes/RigVMLibraryNode.cpp
+++ b/Engine/Source/Developer/RigVMDeveloper/Private/RigVMModel/Nodes/RigVMLibraryNode.cpp
@@ -6,6 +6,7 @@
 #include "RigVMModel/RigVMGraph.h"
 #include "RigVMModel/Nodes/RigVMCollapseNode.h"
 #include "RigVMModel/Nodes/RigVMFunctionReferenceNode.h"
+#include "RigVMModel/RigVMFunctionLibrary.h"
 
 #include UE_INLINE_GENERATED_CPP_BY_NAME(RigVMLibraryNode)
 
@@ -153,12 +154,10 @@ bool URigVMLibraryNode::Contains(URigVMLibraryNode* InContainedNode, bool bRecur
 
 		if(URigVMFunctionReferenceNode* FunctionReferenceNode = Cast<URigVMFunctionReferenceNode>(ContainedNode))
 		{
-			if(URigVMLibraryNode* ReferencedNode = FunctionReferenceNode->GetReferencedNode())
+			FRigVMGraphFunctionIdentifier ReferencedNode = FunctionReferenceNode->GetReferencedFunctionHeader().LibraryPointer;
+			if(ReferencedNode.LibraryNode == InContainedNode)
 			{
-				if(ReferencedNode == InContainedNode)
-				{
-					return true;
-				}
+				return true;
 			}
 		}
 		if(URigVMLibraryNode* ContainedLibraryNode = Cast<URigVMLibraryNode>(ContainedNode))
@@ -189,4 +188,106 @@ TArray<FRigVMExternalVariable> URigVMLibraryNode::GetExternalVariables() const
 	return Variables;
 }
 
+TMap<FRigVMGraphFunctionIdentifier, uint32> URigVMLibraryNode::GetDependencies() const
+{
+	TMap<FRigVMGraphFunctionIdentifier, uint32> Dependencies;
+	
+    TArray<URigVMNode*> NodesToProcess = GetContainedNodes();
+    for (int32 i=0; i<NodesToProcess.Num(); ++i)
+    {
+    	const URigVMNode* Node = NodesToProcess[i];
+    	if (const URigVMFunctionReferenceNode* RefNode = Cast<URigVMFunctionReferenceNode>(Node))
+    	{
+    		uint32 Hash = 0;
+    		if (const FRigVMGraphFunctionData* Data = RefNode->GetReferencedFunctionData())
+    		{
+    			Hash = Data->CompilationData.Hash;
+    		}
+    		Dependencies.Add(RefNode->GetReferencedFunctionHeader().LibraryPointer, Hash);
+    		for (const TPair<FRigVMGraphFunctionIdentifier, uint32>& Pair : RefNode->GetReferencedFunctionHeader().Dependencies)
+    		{
+    			Dependencies.Add(Pair);
+    		}
+    	}
+    	if (const URigVMCollapseNode* CollapseNode = Cast<URigVMCollapseNode>(Node))
+    	{
+    		NodesToProcess.Append(CollapseNode->GetContainedNodes());
+    	}
+    }
+	
+	return Dependencies;
+}
+
+FRigVMGraphFunctionIdentifier URigVMLibraryNode::GetFunctionIdentifier() const
+{
+	FRigVMGraphFunctionIdentifier Identifier;
+	Identifier.HostObject = GetLibrary()->GetFunctionHostObjectPath();
+	Identifier.LibraryNode = this;
+	return Identifier;
+}
+
+FRigVMGraphFunctionHeader URigVMLibraryNode::GetFunctionHeader(IRigVMGraphFunctionHost* InHostObject) const
+{
+	FRigVMGraphFunctionHeader Header;
+	
+	if (InHostObject == nullptr)
+	{
+		Header.LibraryPointer = GetFunctionIdentifier();
+	}
+	else
+	{
+	    Header.LibraryPointer.LibraryNode = this;
+	    Header.LibraryPointer.HostObject = Cast<UObject>(InHostObject);
+	}	
+    
+    Header.Name = GetFName();
+    Header.Category = GetNodeCategory();
+    Header.Keywords = GetNodeKeywords();						
+    Header.Tooltip = GetToolTipText();
+    Header.NodeColor = GetNodeColor();
+    Header.NodeTitle = GetNodeTitle();
+    for(URigVMPin* Pin : GetPins())
+    {
+    	FRigVMGraphFunctionArgument Arg;
+    	Arg.Name = Pin->GetFName();
+    	Arg.DisplayName = Pin->GetDisplayName();
+    	Arg.bIsArray = Pin->IsArray();
+    	Arg.Direction = Pin->GetDirection();
+    	Arg.CPPType = *Pin->GetCPPType();
+    	Arg.CPPTypeObject = Pin->GetCPPTypeObject();
+    	Arg.DefaultValue = Pin->GetDefaultValue();
+    	Arg.bIsConst = Pin->IsDefinedAsConstant();
+
+    	TArray<URigVMPin*> PinsToProcess;
+    	PinsToProcess.Add(Pin);
+    	for (int32 i=0; i<PinsToProcess.Num(); ++i)
+    	{
+    		const FString SubPinPath = PinsToProcess[i]->GetSubPinPath(Pin, false);
+    		Arg.PathToTooltip.Add(SubPinPath, GetToolTipTextForPin(PinsToProcess[i]));
+    		PinsToProcess.Append(PinsToProcess[i]->GetSubPins());
+    	}
+    	
+    	Header.Arguments.Add(Arg);
+    }
+
+	// If the header already exists, try to find it and get its external variables and dependencies
+	TArray<FString> Dependencies;
+	TArray<FRigVMExternalVariable> ExternalVariables;
+	UObject* HostPtr = Cast<UObject>(Header.LibraryPointer.HostObject.ResolveObject());
+	if (IRigVMGraphFunctionHost* Host = Cast<IRigVMGraphFunctionHost>(HostPtr))
+	{
+		if (FRigVMGraphFunctionData* Data = Host->GetRigVMGraphFunctionStore()->FindFunction(Header.LibraryPointer))
+		{
+			Header.Dependencies = Data->Header.Dependencies;
+			Header.ExternalVariables = Data->Header.ExternalVariables;
+		}
+		else
+		{
+			Header.ExternalVariables = GetExternalVariables();
+		}
+	}
+	
+	return Header;
+}
+
 
diff --git a/Engine/Source/Developer/RigVMDeveloper/Private/RigVMModel/RigVMBuildData.cpp b/Engine/Source/Developer/RigVMDeveloper/Private/RigVMModel/RigVMBuildData.cpp
--- a/Engine/Source/Developer/RigVMDeveloper/Private/RigVMModel/RigVMBuildData.cpp
+++ b/Engine/Source/Developer/RigVMDeveloper/Private/RigVMModel/RigVMBuildData.cpp
@@ -2,6 +2,8 @@
 
 #include "RigVMModel/RigVMBuildData.h"
 
+#include "RigVMModel/RigVMFunctionLibrary.h"
+
 #include UE_INLINE_GENERATED_CPP_BY_NAME(RigVMBuildData)
 
 FRigVMReferenceNodeData::FRigVMReferenceNodeData(URigVMFunctionReferenceNode* InReferenceNode)
@@ -9,8 +11,7 @@ FRigVMReferenceNodeData::FRigVMReferenceNodeData(URigVMFunctionReferenceNode* In
 	check(InReferenceNode);
 	ReferenceNodePtr = TSoftObjectPtr<URigVMFunctionReferenceNode>(InReferenceNode);
 	ReferenceNodePath = ReferenceNodePtr.ToString();
-	LibraryNodePtr = TSoftObjectPtr<URigVMLibraryNode>(InReferenceNode->GetReferencedNode());
-	ReferencedFunctionPath = LibraryNodePtr.ToString();
+	ReferencedHeader = InReferenceNode->GetReferencedFunctionHeader();
 }
 
 TSoftObjectPtr<URigVMFunctionReferenceNode> FRigVMReferenceNodeData::GetReferenceNodeObjectPath()
@@ -22,15 +23,6 @@ TSoftObjectPtr<URigVMFunctionReferenceNode> FRigVMReferenceNodeData::GetReferenc
 	return ReferenceNodePtr;
 }
 
-TSoftObjectPtr<URigVMLibraryNode> FRigVMReferenceNodeData::GetReferencedFunctionObjectPath()
-{
-	if(LibraryNodePtr.IsNull())
-	{
-		LibraryNodePtr = TSoftObjectPtr<URigVMFunctionReferenceNode>(ReferencedFunctionPath);
-	}
-	return LibraryNodePtr;
-}
-
 URigVMFunctionReferenceNode* FRigVMReferenceNodeData::GetReferenceNode()
 {
 	if(ReferenceNodePtr.IsNull())
@@ -48,42 +40,20 @@ URigVMFunctionReferenceNode* FRigVMReferenceNodeData::GetReferenceNode()
 	return nullptr;
 }
 
-URigVMLibraryNode* FRigVMReferenceNodeData::GetReferencedFunction()
-{
-	if(LibraryNodePtr.IsNull())
-	{
-		LibraryNodePtr = TSoftObjectPtr<URigVMLibraryNode>(ReferencedFunctionPath);
-	}
-	if(!LibraryNodePtr.IsValid())
-	{
-		LibraryNodePtr.LoadSynchronous();
-	}
-	if(LibraryNodePtr.IsValid())
-	{
-		return LibraryNodePtr.Get();
-	}
-	return nullptr;
-}
-
 URigVMBuildData::URigVMBuildData()
 : UObject()
 , bIsRunningUnitTest(false)
 {
 }
 
-const FRigVMFunctionReferenceArray* URigVMBuildData::FindFunctionReferences(const URigVMLibraryNode* InFunction) const
+const FRigVMFunctionReferenceArray* URigVMBuildData::FindFunctionReferences(const FRigVMGraphFunctionIdentifier& InFunction) const
 {
-	check(InFunction);
-
-	const TSoftObjectPtr<URigVMLibraryNode> Key(InFunction);
-	return FunctionReferences.Find(Key);
+	return GraphFunctionReferences.Find(InFunction);
 }
 
-void URigVMBuildData::ForEachFunctionReference(const URigVMLibraryNode* InFunction,
-                                             TFunction<void(URigVMFunctionReferenceNode*)> PerReferenceFunction) const
+void URigVMBuildData::ForEachFunctionReference(const FRigVMGraphFunctionIdentifier& InFunction,
+                                               TFunction<void(URigVMFunctionReferenceNode*)> PerReferenceFunction) const
 {
-	check(InFunction);
-	
 	if (const FRigVMFunctionReferenceArray* ReferencesEntry = FindFunctionReferences(InFunction))
 	{
 		for (int32 ReferenceIndex = 0; ReferenceIndex < ReferencesEntry->Num(); ReferenceIndex++)
@@ -101,11 +71,9 @@ void URigVMBuildData::ForEachFunctionReference(const URigVMLibraryNode* InFuncti
 	}
 }
 
-void URigVMBuildData::ForEachFunctionReferenceSoftPtr(const URigVMLibraryNode* InFunction,
-	TFunction<void(TSoftObjectPtr<URigVMFunctionReferenceNode>)> PerReferenceFunction) const
+void URigVMBuildData::ForEachFunctionReferenceSoftPtr(const FRigVMGraphFunctionIdentifier& InFunction,
+                                                      TFunction<void(TSoftObjectPtr<URigVMFunctionReferenceNode>)> PerReferenceFunction) const
 {
-	check(InFunction);
-
 	if (const FRigVMFunctionReferenceArray* ReferencesEntry = FindFunctionReferences(InFunction))
 	{
 		for (int32 ReferenceIndex = 0; ReferenceIndex < ReferencesEntry->Num(); ReferenceIndex++)
@@ -124,57 +92,53 @@ void URigVMBuildData::UpdateReferencesForFunctionReferenceNode(URigVMFunctionRef
 	{
 		return;
 	}
-	
-	if(const URigVMLibraryNode* Function = InReferenceNode->GetReferencedNode())
+
+	const FRigVMGraphFunctionHeader& Header = InReferenceNode->GetReferencedFunctionHeader();
+	FRigVMFunctionReferenceArray* ReferencesEntry = GraphFunctionReferences.Find(Header.LibraryPointer);
+	if (ReferencesEntry == nullptr)
 	{
-		const TSoftObjectPtr<URigVMLibraryNode> Key(Function);
-		FRigVMFunctionReferenceArray* ReferencesEntry = FunctionReferences.Find(Key);
-		if (ReferencesEntry == nullptr)
-		{
-			Modify();
-			FunctionReferences.Add(Key);
-			ReferencesEntry = FunctionReferences.Find(Key);
-		}
+		Modify();
+		GraphFunctionReferences.Add(Header.LibraryPointer);
+		ReferencesEntry = GraphFunctionReferences.Find(Header.LibraryPointer);
+	}
 
-		const FString ReferenceNodePathName = InReferenceNode->GetPathName();
-		for (int32 ReferenceIndex = 0; ReferenceIndex < ReferencesEntry->Num(); ReferenceIndex++)
+	const FString ReferenceNodePathName = InReferenceNode->GetPathName();
+	for (int32 ReferenceIndex = 0; ReferenceIndex < ReferencesEntry->Num(); ReferenceIndex++)
+	{
+		const TSoftObjectPtr<URigVMFunctionReferenceNode>& Reference = ReferencesEntry->operator [](ReferenceIndex);
+		if(Reference.ToString() == ReferenceNodePathName)
 		{
-			const TSoftObjectPtr<URigVMFunctionReferenceNode>& Reference = ReferencesEntry->operator [](ReferenceIndex);
-			if(Reference.ToString() == ReferenceNodePathName)
-			{
-				return;
-			}
+			return;
 		}
-
-		Modify();
-		ReferencesEntry->FunctionReferences.Add(InReferenceNode);
-		MarkPackageDirty();
 	}
+
+	Modify();
+	ReferencesEntry->FunctionReferences.Add(InReferenceNode);
+	MarkPackageDirty();
+
 }
 
-void URigVMBuildData::RegisterFunctionReference(URigVMLibraryNode* InFunction,
-	URigVMFunctionReferenceNode* InReference)
+void URigVMBuildData::RegisterFunctionReference(const FRigVMGraphFunctionIdentifier& InFunction, URigVMFunctionReferenceNode* InReference)
 {
-	if(InFunction == nullptr || InReference == nullptr)
+	if(InReference == nullptr)
 	{
 		return;
 	}
 
-	const TSoftObjectPtr<URigVMLibraryNode> FunctionKey(InFunction);
 	const TSoftObjectPtr<URigVMFunctionReferenceNode> ReferenceKey(InReference);
 
-	RegisterFunctionReference(FunctionKey, ReferenceKey);
+	RegisterFunctionReference(InFunction, ReferenceKey);
 }
 
-void URigVMBuildData::RegisterFunctionReference(TSoftObjectPtr<URigVMLibraryNode> InFunction,
-	TSoftObjectPtr<URigVMFunctionReferenceNode> InReference)
+void URigVMBuildData::RegisterFunctionReference(const FRigVMGraphFunctionIdentifier& InFunction,
+                                                TSoftObjectPtr<URigVMFunctionReferenceNode> InReference)
 {
-	if(InFunction.IsNull() || InReference.IsNull())
+	if(InReference.IsNull())
 	{
 		return;
 	}
 
-	if(FRigVMFunctionReferenceArray* ReferenceEntry = FunctionReferences.Find(InFunction))
+	if(FRigVMFunctionReferenceArray* ReferenceEntry = GraphFunctionReferences.Find(InFunction))
 	{
 		if(ReferenceEntry->FunctionReferences.Contains(InReference))
 		{
@@ -189,7 +153,7 @@ void URigVMBuildData::RegisterFunctionReference(TSoftObjectPtr<URigVMLibraryNode
 		Modify();
 		FRigVMFunctionReferenceArray NewReferenceEntry;
 		NewReferenceEntry.FunctionReferences.Add(InReference);
-		FunctionReferences.Add(InFunction, NewReferenceEntry);
+		GraphFunctionReferences.Add(InFunction, NewReferenceEntry);
 	}
 	
 	MarkPackageDirty();
@@ -197,32 +161,64 @@ void URigVMBuildData::RegisterFunctionReference(TSoftObjectPtr<URigVMLibraryNode
 
 void URigVMBuildData::RegisterFunctionReference(FRigVMReferenceNodeData InReferenceNodeData)
 {
-	RegisterFunctionReference(InReferenceNodeData.GetReferencedFunctionObjectPath(), InReferenceNodeData.GetReferenceNodeObjectPath());
+	if (InReferenceNodeData.ReferencedHeader.IsValid())
+	{
+		return RegisterFunctionReference(InReferenceNodeData.ReferencedHeader.LibraryPointer, InReferenceNodeData.GetReferenceNodeObjectPath());
+	}
+	
+	check(!InReferenceNodeData.ReferencedFunctionPath_DEPRECATED.IsEmpty());
+
+	TSoftObjectPtr<URigVMLibraryNode> LibraryNodePtr = TSoftObjectPtr<URigVMLibraryNode>(InReferenceNodeData.ReferencedFunctionPath_DEPRECATED);
+	bool bFound = false;
+	for (TPair< FRigVMGraphFunctionIdentifier, FRigVMFunctionReferenceArray >& Pair : GraphFunctionReferences)
+	{
+		if (Pair.Key.LibraryNode == LibraryNodePtr.ToSoftObjectPath())
+		{
+			Pair.Value.FunctionReferences.Add(InReferenceNodeData.GetReferenceNodeObjectPath());
+			bFound = true;
+			break;
+		}
+	}
+
+	if (!bFound)
+	{
+		FRigVMGraphFunctionIdentifier Pointer(nullptr, LibraryNodePtr.ToSoftObjectPath());
+		if (!LibraryNodePtr.IsValid())
+		{
+			LibraryNodePtr.LoadSynchronous();
+		}
+		if (LibraryNodePtr.IsValid())
+		{
+			Pointer.HostObject = Cast<UObject>(LibraryNodePtr.Get()->GetFunctionHeader().GetFunctionHost());
+		}
+		FRigVMFunctionReferenceArray RefArray;
+		RefArray.FunctionReferences.Add(InReferenceNodeData.GetReferenceNodeObjectPath());
+		GraphFunctionReferences.Add(Pointer, RefArray);
+	}
 }
 
-void URigVMBuildData::UnregisterFunctionReference(URigVMLibraryNode* InFunction,
+void URigVMBuildData::UnregisterFunctionReference(const FRigVMGraphFunctionIdentifier& InFunction,
                                                   URigVMFunctionReferenceNode* InReference)
 {
-	if(InFunction == nullptr || InReference == nullptr)
+	if(InReference == nullptr)
 	{
 		return;
 	}
 
-	const TSoftObjectPtr<URigVMLibraryNode> FunctionKey(InFunction);
 	const TSoftObjectPtr<URigVMFunctionReferenceNode> ReferenceKey(InReference);
 
-	return UnregisterFunctionReference(FunctionKey, ReferenceKey);
+	return UnregisterFunctionReference(InFunction, ReferenceKey);
 }
 
-void URigVMBuildData::UnregisterFunctionReference(TSoftObjectPtr<URigVMLibraryNode> InFunction,
-	TSoftObjectPtr<URigVMFunctionReferenceNode> InReference)
+void URigVMBuildData::UnregisterFunctionReference(const FRigVMGraphFunctionIdentifier& InFunction,
+                                                  TSoftObjectPtr<URigVMFunctionReferenceNode> InReference)
 {
-	if(InFunction.IsNull() || InReference.IsNull())
+	if(InReference.IsNull())
 	{
 		return;
 	}
 
-	if(FRigVMFunctionReferenceArray* ReferenceEntry = FunctionReferences.Find(InFunction))
+	if(FRigVMFunctionReferenceArray* ReferenceEntry = GraphFunctionReferences.Find(InFunction))
 	{
 		if(!ReferenceEntry->FunctionReferences.Contains(InReference))
 		{
@@ -246,7 +242,7 @@ void URigVMBuildData::ClearInvalidReferences()
 	
 	// check each function's each reference
 	int32 NumRemoved = 0;
-	for (TTuple<TSoftObjectPtr<URigVMLibraryNode>, FRigVMFunctionReferenceArray>& FunctionReferenceInfo : FunctionReferences)
+	for (TTuple<FRigVMGraphFunctionIdentifier, FRigVMFunctionReferenceArray>& FunctionReferenceInfo : GraphFunctionReferences)
 	{
 		FRigVMFunctionReferenceArray* ReferencesEntry = &FunctionReferenceInfo.Value;
 
diff --git a/Engine/Source/Developer/RigVMDeveloper/Private/RigVMModel/RigVMClient.cpp b/Engine/Source/Developer/RigVMDeveloper/Private/RigVMModel/RigVMClient.cpp
--- a/Engine/Source/Developer/RigVMDeveloper/Private/RigVMModel/RigVMClient.cpp
+++ b/Engine/Source/Developer/RigVMDeveloper/Private/RigVMModel/RigVMClient.cpp
@@ -91,9 +91,9 @@ URigVMGraph* FRigVMClient::GetModel(const FString& InNodePathOrName) const
 		if(InNodePathOrName.StartsWith(NodePathPrefix))
 		{
 			const FString RemainingNodePath = InNodePathOrName.Mid(NodePathPrefix.Len());
-			if(const URigVMLibraryNode* LibraryNode = Cast<URigVMLibraryNode>(Model->FindNode(RemainingNodePath)))
+			if(const URigVMCollapseNode* CollapseNode = Cast<URigVMCollapseNode>(Model->FindNode(RemainingNodePath)))
 			{
-				return LibraryNode->GetContainedGraph();
+				return CollapseNode->GetContainedGraph();
 			}
 		}
 	}
@@ -261,6 +261,7 @@ void FRigVMClient::AddModel(URigVMGraph* InModel, bool bCreateController)
 	{
 		check(FunctionLibrary == nullptr);
 		FunctionLibrary = Cast<URigVMFunctionLibrary>(InModel);
+		FunctionLibrary->SetDefaultFunctionLibrary(FunctionLibrary);		
 	}
 	else
 	{
@@ -321,6 +322,17 @@ URigVMFunctionLibrary* FRigVMClient::GetOrCreateFunctionLibrary(bool bSetupUndoR
 	{
 		NewFunctionLibrary = NewObject<URigVMFunctionLibrary>(GetOuter(), SafeGraphName);
 	}
+
+	NewFunctionLibrary->GetFunctionHostObjectPathDelegate.BindLambda([this]() -> const FSoftObjectPath 
+		{
+			if (GetOuter()->Implements<URigVMClientHost>())
+			{
+				IRigVMClientHost* ClientHost = Cast<IRigVMClientHost>(GetOuter());
+				return Cast<UObject>(ClientHost->GetRigVMGraphFunctionHost());				
+			}
+			return nullptr;
+		});
+	
 	AddModel(NewFunctionLibrary, bCreateController);
 	return NewFunctionLibrary;
 }
@@ -605,6 +617,10 @@ URigVMController* FRigVMClient::CreateController(const URigVMGraph* InModel)
 	URigVMController* Controller = NewObject<URigVMController>(GetOuter(), SafeControllerName);
 	Controllers.Add(InModel, Controller);
 	Controller->SetGraph((URigVMGraph*)InModel);
+	Controller->OnModified().AddLambda([this](ERigVMGraphNotifType InNotifType, URigVMGraph* InGraph, UObject* InSubject)
+	{
+		HandleGraphModifiedEvent(InNotifType, InGraph, InSubject);
+	});
 
 	if (GetOuter()->Implements<URigVMClientHost>())
 	{
@@ -651,3 +667,316 @@ void FRigVMClient::PatchModelsOnLoad()
 	}
 }
 
+void FRigVMClient::HandleGraphModifiedEvent(ERigVMGraphNotifType InNotifType, URigVMGraph* InGraph, UObject* InSubject)
+{
+	if (bIgnoreModelNotifications)
+	{
+		return;
+	}
+	
+#if WITH_EDITOR
+
+	IRigVMClientHost* ClientHost = Cast<IRigVMClientHost>(GetOuter());
+	IRigVMGraphFunctionHost* FunctionHost = ClientHost->GetRigVMGraphFunctionHost();
+	FRigVMGraphFunctionStore* FunctionStore = FunctionHost->GetRigVMGraphFunctionStore();
+
+	switch (InNotifType)
+	{
+		case ERigVMGraphNotifType::NodeAdded: // A node has been added to the graph (Subject == URigVMNode)
+		{
+			// A node was added directly into the function library
+			if(InGraph->IsA<URigVMFunctionLibrary>())
+			{
+				if (URigVMCollapseNode* CollapseNode = Cast<URigVMCollapseNode>(InSubject))
+				{
+					if (GetOuter()->Implements<URigVMClientHost>())
+					{
+						FunctionStore->AddFunction(CollapseNode->GetFunctionHeader(FunctionHost), false);
+					}
+				}
+			}
+			// A node was added into the contained graph of a function
+			else if(URigVMLibraryNode* LibraryNode = Cast<URigVMNode>(InSubject)->FindFunctionForNode())
+			{
+				DirtyGraphFunctionCompilationData(LibraryNode);
+				if (URigVMFunctionReferenceNode* FunctionReference = Cast<URigVMFunctionReferenceNode>(InSubject))
+				{
+					UpdateDependenciesForFunction(LibraryNode);
+					UpdateExternalVariablesForFunction(LibraryNode);
+				}
+			}
+			break;	
+		}
+		case ERigVMGraphNotifType::NodeRemoved: // A node has been removed from the graph (Subject == URigVMNode)
+		{		
+			if(InSubject->GetOuter()->IsA<URigVMFunctionLibrary>())
+			{
+				if (URigVMCollapseNode* CollapseNode = Cast<URigVMCollapseNode>(InSubject))
+				{
+					if (GetOuter()->Implements<URigVMClientHost>())
+					{
+						FunctionStore->RemoveFunction(FRigVMGraphFunctionIdentifier (Cast<UObject>(FunctionHost), CollapseNode));
+					}
+				}
+			}
+			// A node was added into the contained graph of a function
+			else if(URigVMLibraryNode* LibraryNode = Cast<URigVMNode>(InSubject)->FindFunctionForNode())
+			{
+				DirtyGraphFunctionCompilationData(LibraryNode);
+				if (URigVMFunctionReferenceNode* FunctionReference = Cast<URigVMFunctionReferenceNode>(InSubject))
+				{
+					UpdateDependenciesForFunction(LibraryNode);
+					UpdateExternalVariablesForFunction(LibraryNode);
+				}
+			}
+			break;	
+		}
+		case ERigVMGraphNotifType::VariableAdded: // A variable has been added (Subject == URigVMVariableNode)
+		case ERigVMGraphNotifType::VariableRemoved: // A variable has been removed (Subject == URigVMVariableNode)
+		case ERigVMGraphNotifType::VariableRenamed: // A variable has been renamed (Subject == URigVMVariableNode)
+		case ERigVMGraphNotifType::VariableRemappingChanged: // A function reference node's remapping has changed (Subject == URigVMFunctionReferenceNode)
+		{
+			if(URigVMLibraryNode* LibraryNode = Cast<URigVMNode>(InSubject)->FindFunctionForNode())
+			{
+				DirtyGraphFunctionCompilationData(LibraryNode);
+				UpdateExternalVariablesForFunction(LibraryNode);				
+			}
+			break;
+		}
+		case ERigVMGraphNotifType::NodeRenamed: // A node has been renamed in the graph (Subject == URigVMNode)
+		{
+			if(InSubject->GetOuter()->IsA<URigVMFunctionLibrary>())
+			{
+				if (URigVMCollapseNode* CollapseNode = Cast<URigVMCollapseNode>(InSubject))
+				{
+					URigVMBuildData* BuildData = GetFunctionLibrary()->GetBuildData();
+					IRigVMGraphFunctionHost* Host = Cast<IRigVMGraphFunctionHost>(CollapseNode->GetFunctionIdentifier().HostObject.ResolveObject());
+					FRigVMGraphFunctionData* Data = Host->GetRigVMGraphFunctionStore()->FindFunctionByName(CollapseNode->GetPreviousFName());
+					Data->Header = CollapseNode->GetFunctionHeader();
+					
+					BuildData->ForEachFunctionReference(CollapseNode->GetFunctionIdentifier(), [Data](URigVMFunctionReferenceNode* ReferenceNode)
+					{
+						ReferenceNode->ReferencedFunctionHeader = Data->Header;
+					});
+					UpdateGraphFunctionData(CollapseNode);
+				}
+			}
+			break;	
+		}
+		case ERigVMGraphNotifType::NodeColorChanged: // A node's color has changed (Subject == URigVMNode)
+		case ERigVMGraphNotifType::NodeCategoryChanged: // A node's category has changed (Subject == URigVMNode)
+		case ERigVMGraphNotifType::NodeKeywordsChanged: // A node's keywords have changed (Subject == URigVMNode)
+		case ERigVMGraphNotifType::NodeDescriptionChanged: // A node's description has changed (Subject == URigVMNode)
+		{
+			if(InSubject->GetOuter()->IsA<URigVMFunctionLibrary>())
+			{
+				if (URigVMCollapseNode* CollapseNode = Cast<URigVMCollapseNode>(InSubject))
+				{
+					UpdateGraphFunctionData(CollapseNode);
+				}
+			}
+		}
+		case ERigVMGraphNotifType::NodeReferenceChanged: // A node has changed it's referenced function (Subject == URigVMFunctionReferenceNode)
+		case ERigVMGraphNotifType::LibraryTemplateChanged: // The definition of a library node's template has changed (Subject == URigVMLibraryNode)
+		{
+			
+			
+			break;	
+		}
+		
+		case ERigVMGraphNotifType::PinAdded: // A pin has been added to a given node (Subject == URigVMPin)
+		case ERigVMGraphNotifType::PinRemoved: // A pin has been removed from a given node (Subject == URigVMPin)
+		case ERigVMGraphNotifType::PinRenamed: // A pin has been renamed (Subject == URigVMPin)
+		case ERigVMGraphNotifType::PinArraySizeChanged: // An array pin's size has changed (Subject == URigVMPin)
+		case ERigVMGraphNotifType::PinDefaultValueChanged: // A pin's default value has changed (Subject == URigVMPin)
+		case ERigVMGraphNotifType::PinDirectionChanged: // A pin's direction has changed (Subject == URigVMPin)
+		case ERigVMGraphNotifType::PinTypeChanged: // A pin's data type has changed (Subject == URigVMPin)
+		case ERigVMGraphNotifType::PinIndexChanged: // A pin's index has changed (Subject == URigVMPin)
+		{
+			if (URigVMNode* Node = Cast<URigVMNode>(InSubject->GetOuter()))
+			{
+				if (URigVMLibraryNode* LibraryNode = Node->FindFunctionForNode())
+				{
+					DirtyGraphFunctionCompilationData(LibraryNode);
+				}
+				if(Node->GetOuter()->IsA<URigVMFunctionLibrary>())
+				{
+					if (URigVMCollapseNode* CollapseNode = Cast<URigVMCollapseNode>(Node))
+					{
+						UpdateGraphFunctionData(CollapseNode);
+					}
+				}
+			}
+		}
+
+		case ERigVMGraphNotifType::LinkAdded: // A link has been added (Subject == URigVMLink)
+		case ERigVMGraphNotifType::LinkRemoved: // A link has been removed (Subject == URigVMLink)
+		{
+			if (URigVMLink* Link = Cast<URigVMLink>(InSubject))
+			{
+				if (URigVMNode* OuterNode = Cast<URigVMNode>(Link->GetGraph()->GetOuter()))
+				{
+					if (URigVMLibraryNode* LibraryNode = OuterNode->FindFunctionForNode())
+					{
+						DirtyGraphFunctionCompilationData(LibraryNode);
+					}
+				}
+			}
+		}
+		
+		
+		default:
+		{
+			break;
+		}
+	}
+	
+#endif
+}
+
+FRigVMGraphFunctionStore* FRigVMClient::FindFunctionStore(const URigVMLibraryNode* InLibraryNode)
+{
+	if (IRigVMClientHost* ClientHost = InLibraryNode->GetImplementingOuter<IRigVMClientHost>())
+	{
+		if (IRigVMGraphFunctionHost* FunctionHost = ClientHost->GetRigVMGraphFunctionHost())
+		{
+			return FunctionHost->GetRigVMGraphFunctionStore();
+		}
+	}
+	return nullptr;
+}
+
+bool FRigVMClient::UpdateFunctionReferences(const FRigVMGraphFunctionHeader& Header, bool bUpdateDependencies, bool bUpdateExternalVariables)
+{
+	URigVMBuildData* BuildData = GetFunctionLibrary()->GetBuildData();
+	if (const FRigVMFunctionReferenceArray* FunctionReferenceArray = BuildData->FindFunctionReferences(Header.LibraryPointer))
+	{
+		for (int32 i=0; i<FunctionReferenceArray->Num(); ++i)
+		{
+			TSoftObjectPtr<URigVMFunctionReferenceNode> Reference = FunctionReferenceArray->FunctionReferences[i];
+			if (Reference.IsValid())
+			{
+				URigVMFunctionReferenceNode* Node = Reference.Get();
+				Node->ReferencedFunctionHeader = Header;
+
+				if (bUpdateDependencies || bUpdateExternalVariables)
+				{
+					if (URigVMLibraryNode* LibraryNode = Node->FindFunctionForNode())
+					{
+						IRigVMClientHost* OtherClientHost = LibraryNode->GetImplementingOuter<IRigVMClientHost>();							
+						if (bUpdateDependencies)
+						{
+							OtherClientHost->GetRigVMClient()->UpdateDependenciesForFunction(LibraryNode);
+						}
+						if (bUpdateExternalVariables)
+						{								
+							OtherClientHost->GetRigVMClient()->UpdateExternalVariablesForFunction(LibraryNode);
+						}
+					}
+				}
+			}
+		}
+	}
+	return true;
+}
+
+bool FRigVMClient::UpdateGraphFunctionData(const URigVMLibraryNode* InLibraryNode)
+{
+	if (GetOuter()->Implements<URigVMClientHost>())
+	{
+		if (FRigVMGraphFunctionStore* Store = FindFunctionStore(InLibraryNode))
+		{
+			if (FRigVMGraphFunctionData* Data = Store->UpdateFunctionInterface(InLibraryNode->GetFunctionHeader()))
+			{
+				UpdateFunctionReferences(Data->Header, false, false);
+				return true;
+			}
+		}
+	}	
+	
+	return false;	
+}
+
+bool FRigVMClient::UpdateExternalVariablesForFunction(const URigVMLibraryNode* InLibraryNode)
+{
+	if (FRigVMGraphFunctionStore* Store = FindFunctionStore(InLibraryNode))
+	{
+		FRigVMGraphFunctionIdentifier Identifier = InLibraryNode->GetFunctionIdentifier();
+		if (Store->UpdateExternalVariables(Identifier, InLibraryNode->GetExternalVariables()))
+		{
+			const FRigVMGraphFunctionData* Data = Store->FindFunction(Identifier);
+			UpdateFunctionReferences(Data->Header, false, true);
+			return true;
+		}
+	}
+	
+	return false;
+}
+
+bool FRigVMClient::UpdateDependenciesForFunction(const URigVMLibraryNode* InLibraryNode)
+{
+	if (FRigVMGraphFunctionStore* Store = FindFunctionStore(InLibraryNode))
+	{
+		TMap<FRigVMGraphFunctionIdentifier, uint32> Dependencies = InLibraryNode->GetDependencies();
+		FRigVMGraphFunctionIdentifier Identifier = InLibraryNode->GetFunctionIdentifier();
+		if (Store->UpdateDependencies(Identifier, Dependencies))
+		{
+			const FRigVMGraphFunctionData* Data = Store->FindFunction(Identifier);
+			UpdateFunctionReferences(Data->Header, true, false);
+			return true;
+		}
+	}
+	
+	return false;
+}
+
+bool FRigVMClient::DirtyGraphFunctionCompilationData(URigVMLibraryNode* InLibraryNode)
+{
+	if (FRigVMGraphFunctionStore* Store = FindFunctionStore(InLibraryNode))
+	{
+		FRigVMGraphFunctionIdentifier Identifier = InLibraryNode->GetFunctionIdentifier();
+		if (const FRigVMGraphFunctionData* Data = Store->FindFunction(Identifier))
+		{
+			if (Store->RemoveFunctionCompilationData(Identifier))
+			{
+				URigVMBuildData* BuildData = GetFunctionLibrary()->GetBuildData();
+				if (const FRigVMFunctionReferenceArray* FunctionReferenceArray = BuildData->FindFunctionReferences(Identifier))
+				{
+					for (int32 i=0; i<FunctionReferenceArray->Num(); ++i)
+					{
+						TSoftObjectPtr<URigVMFunctionReferenceNode> Reference = FunctionReferenceArray->FunctionReferences[i];
+						{
+							if (!Reference.IsValid())
+							{
+								Reference.LoadSynchronous();
+							}
+							if (Reference.IsValid())
+							{
+								if (URigVMFunctionLibrary* ReferenceFunctionLibrary = Reference->GetGraph()->GetDefaultFunctionLibrary())
+								{
+									if (URigVMLibraryNode* ReferenceLibraryNode = Reference.Get()->FindFunctionForNode())
+									{
+										IRigVMClientHost* OtherClientHost = ReferenceLibraryNode->GetImplementingOuter<IRigVMClientHost>();
+										OtherClientHost->GetRigVMClient()->DirtyGraphFunctionCompilationData(ReferenceLibraryNode);
+									}
+								}							
+							}
+						}
+					}
+				}
+				return true;
+			}
+		}
+	}
+
+	return false;
+}
+
+bool FRigVMClient::IsFunctionPublic(URigVMLibraryNode* InLibraryNode)
+{
+	if (FRigVMGraphFunctionStore* Store = FindFunctionStore(InLibraryNode))
+	{
+		return Store->IsFunctionPublic(InLibraryNode->GetFunctionIdentifier());
+	}
+	return false;
+}
diff --git a/Engine/Source/Developer/RigVMDeveloper/Private/RigVMModel/RigVMController.cpp b/Engine/Source/Developer/RigVMDeveloper/Private/RigVMModel/RigVMController.cpp
--- a/Engine/Source/Developer/RigVMDeveloper/Private/RigVMModel/RigVMController.cpp
+++ b/Engine/Source/Developer/RigVMDeveloper/Private/RigVMModel/RigVMController.cpp
@@ -23,6 +23,7 @@
 #include "RigVMPythonUtils.h"
 #include "RigVMTypeUtils.h"
 #include "Engine/UserDefinedStruct.h"
+#include "RigVMModel/RigVMClient.h"
 
 #include UE_INLINE_GENERATED_CPP_BY_NAME(RigVMController)
 
@@ -675,39 +676,45 @@ TArray<FString> URigVMController::GetAddNodePythonCommands(URigVMNode* Node) con
 						*RigVMPythonUtils::Vector2DToPythonString(EnumNode->GetPosition()),
 						*NodeName));
 	}
-	else if (const URigVMLibraryNode* LibraryNode = Cast<URigVMLibraryNode>(Node))
+	else if (const URigVMFunctionReferenceNode* RefNode = Cast<URigVMFunctionReferenceNode>(Node))
 	{
-		const FString ContainedGraphName = GetSanitizedGraphName(LibraryNode->GetContainedGraph()->GetGraphName());
-
-		// AddFunctionReferenceNode(URigVMLibraryNode* InFunctionDefinition, const FVector2D& InNodePosition = FVector2D::ZeroVector, const FString& InNodeName = TEXT(""), bool bSetupUndoRedo = true);
-		URigVMFunctionLibrary* Library = LibraryNode->GetLibrary();
-		if (!Library || Library == GetGraph()->GetDefaultFunctionLibrary())
+		if (RefNode->GetReferencedFunctionHeader().LibraryPointer.HostObject == GetGraph()->GetDefaultFunctionLibrary()->GetFunctionHostObjectPath())
 		{
 			Commands.Add(FString::Printf(TEXT("blueprint.get_controller_by_name('%s').add_function_reference_node(function_%s, %s, '%s')"),
-						*GraphName,
-						*RigVMPythonUtils::PythonizeName(ContainedGraphName),
-						*RigVMPythonUtils::Vector2DToPythonString(LibraryNode->GetPosition()), 
-						*NodeName));
-		}
-		else
-		{
-			Commands.Add(FString::Printf(TEXT("function_blueprint = unreal.load_object(name = '%s', outer = None)"),
-				*Library->GetOuter()->GetPathName()));
-			Commands.Add(FString::Printf(TEXT("blueprint.get_controller_by_name('%s').add_function_reference_node(function_blueprint.get_local_function_library().find_function('%s'), %s, '%s')"),
-						*GraphName,
-						*NodeName,
-						*RigVMPythonUtils::Vector2DToPythonString(LibraryNode->GetPosition()), 
-						*NodeName));
-		}
-
-		if (Node->IsA<URigVMCollapseNode>())
-		{
-			Commands.Add(FString::Printf(TEXT("blueprint.get_controller_by_name('%s').promote_function_reference_node_to_collapse_node('%s')"),
 					*GraphName,
+					*RigVMPythonUtils::PythonizeName(RefNode->LoadReferencedNode()->GetContainedGraph()->GetGraphName()),
+					*RigVMPythonUtils::Vector2DToPythonString(RefNode->GetPosition()), 
 					*NodeName));
-			Commands.Add(FString::Printf(TEXT("library_controller.remove_function_from_library('%s')"),
-					*ContainedGraphName));
 		}
+		// todo UE-169886
+		// else
+		// {
+		// 	Commands.Add(FString::Printf(TEXT("function_blueprint = unreal.load_object(name = '%s', outer = None)"),
+		// 		*Library->GetOuter()->GetPathName()));
+		// 	Commands.Add(FString::Printf(TEXT("blueprint.get_controller_by_name('%s').add_function_reference_node(function_blueprint.get_local_function_library().find_function('%s'), %s, '%s')"),
+		// 				*GraphName,
+		// 				*NodeName,
+		// 				*RigVMPythonUtils::Vector2DToPythonString(LibraryNode->GetPosition()), 
+		// 				*NodeName));
+		// }
+	}
+	else if (const URigVMCollapseNode* CollapseNode = Cast<URigVMCollapseNode>(Node))
+	{
+		const FString ContainedGraphName = GetSanitizedGraphName(CollapseNode->GetContainedGraph()->GetGraphName());
+		// AddFunctionReferenceNode(URigVMLibraryNode* InFunctionDefinition, const FVector2D& InNodePosition = FVector2D::ZeroVector, const FString& InNodeName = TEXT(""), bool bSetupUndoRedo = true);
+		Commands.Add(FString::Printf(TEXT("blueprint.get_controller_by_name('%s').add_function_reference_node(function_%s, %s, '%s')"),
+					*GraphName,
+					*RigVMPythonUtils::PythonizeName(ContainedGraphName),
+					*RigVMPythonUtils::Vector2DToPythonString(CollapseNode->GetPosition()), 
+					*NodeName));
+	
+
+		Commands.Add(FString::Printf(TEXT("blueprint.get_controller_by_name('%s').promote_function_reference_node_to_collapse_node('%s')"),
+				*GraphName,
+				*NodeName));
+		Commands.Add(FString::Printf(TEXT("library_controller.remove_function_from_library('%s')"),
+				*ContainedGraphName));
+	
 	}
 	else if (const URigVMInvokeEntryNode* InvokeEntryNode = Cast<URigVMInvokeEntryNode>(Node))
 	{
@@ -3429,9 +3436,9 @@ TArray<FName> URigVMController::ImportNodesFromText(const FString& InText, bool
 	for (int32 i=0; i<CreatedNodes.Num(); ++i)
 	{
 		URigVMNode* CreatedNode = CreatedNodes[i];
-		if (URigVMLibraryNode* LibraryNode = Cast<URigVMLibraryNode>(CreatedNode))
+		if (URigVMCollapseNode* CollapseNode = Cast<URigVMCollapseNode>(CreatedNode))
 		{
-			if (URigVMGraph* ContainedGraph = LibraryNode->GetContainedGraph())
+			if (URigVMGraph* ContainedGraph = CollapseNode->GetContainedGraph())
 			{
 				EditGuards.Emplace(ContainedGraph->bEditable, true);
 				CreatedNodes.Append(ContainedGraph->GetNodes());
@@ -3552,12 +3559,9 @@ TArray<FName> URigVMController::ImportNodesFromText(const FString& InText, bool
 
 			if (URigVMFunctionReferenceNode* FunctionRefNode = Cast<URigVMFunctionReferenceNode>(SubNode))
 			{
-				if (URigVMLibraryNode* FunctionDefinition = FunctionRefNode->GetReferencedNode())
+				if(URigVMBuildData* BuildData = GetBuildData())
 				{
-					if(URigVMBuildData* BuildData = GetBuildData())
-					{
-						BuildData->RegisterFunctionReference(FunctionDefinition, FunctionRefNode);
-					}
+					BuildData->RegisterFunctionReference(FunctionRefNode->GetReferencedFunctionHeader().LibraryPointer, FunctionRefNode);
 				}
 			}
 
@@ -3854,9 +3858,10 @@ TMap<URigVMLibraryNode*, URigVMLibraryNode*> URigVMController::LocalizeFunctions
 				if(OtherLibrary != ThisLibrary)
 				{
 					bool bIsAvailable = false;
-					if(IsFunctionAvailableDelegate.IsBound())
+					FRigVMGraphFunctionIdentifier Identifier = CollapseNode->GetFunctionIdentifier();
+					if (IRigVMGraphFunctionHost* Host = Cast<IRigVMGraphFunctionHost>(Identifier.HostObject.ResolveObject()))
 					{
-						bIsAvailable = IsFunctionAvailableDelegate.Execute(CollapseNode);
+						bIsAvailable = Host->GetRigVMGraphFunctionStore()->IsFunctionPublic(Identifier);		
 					}
 
 					if(!bIsAvailable)
@@ -3873,11 +3878,11 @@ TMap<URigVMLibraryNode*, URigVMLibraryNode*> URigVMController::LocalizeFunctions
 			}
 		}
 
-		else if(URigVMFunctionReferenceNode* FunctionReferencedNode = Cast<URigVMFunctionReferenceNode>(NodeToVisit))
+		else if(const URigVMFunctionReferenceNode* FunctionReferencedNode = Cast<URigVMFunctionReferenceNode>(NodeToVisit))
 		{
-			if(FunctionReferencedNode->GetLibrary() != ThisLibrary)
+			if(FunctionReferencedNode->GetReferencedFunctionHeader().LibraryPointer.HostObject != ThisLibrary->GetFunctionHostObjectPath())
 			{
-				if(URigVMCollapseNode* FunctionDefinition = Cast<URigVMCollapseNode>(FunctionReferencedNode->GetReferencedNode()))
+				if(URigVMCollapseNode* FunctionDefinition = Cast<URigVMCollapseNode>(FunctionReferencedNode->LoadReferencedNode()))
 				{
 					NodesToVisit.AddUnique(FunctionDefinition);
 				}
@@ -3894,8 +3899,8 @@ TMap<URigVMLibraryNode*, URigVMLibraryNode*> URigVMController::LocalizeFunctions
 	});
 
 	// export all of the content for each node
-	TMap<URigVMLibraryNode*, FString> ExportedTextPerFunction;
-	for(URigVMLibraryNode* FunctionToLocalize : FunctionsToLocalize)
+	TMap<const URigVMLibraryNode*, FString> ExportedTextPerFunction;
+	for(const URigVMLibraryNode* FunctionToLocalize : FunctionsToLocalize)
 	{
 		URigVMFunctionLibrary* OtherLibrary = Cast<URigVMFunctionLibrary>(FunctionToLocalize->GetOuter());
 		FRigVMControllerGraphGuard GraphGuard(this, OtherLibrary, false);
@@ -3914,15 +3919,6 @@ TMap<URigVMLibraryNode*, URigVMLibraryNode*> URigVMController::LocalizeFunctions
 	// import the functions to our local function library
 	{
 		FRigVMControllerGraphGuard GraphGuard(this, ThisLibrary, bSetupUndoRedo);
-
-		// override the availability and check up later
-		TGuardValue<FRigVMController_IsFunctionAvailableDelegate> IsFunctionAvailableGuard(IsFunctionAvailableDelegate,
-			FRigVMController_IsFunctionAvailableDelegate::CreateLambda([](URigVMLibraryNode*)
-			{
-				return true;
-			})
-		);
-
 		for(URigVMLibraryNode* FunctionToLocalize : FunctionsToLocalize)
 		{
 			const FString& ExportedText = ExportedTextPerFunction.FindChecked(FunctionToLocalize);
@@ -3965,7 +3961,7 @@ TMap<URigVMLibraryNode*, URigVMLibraryNode*> URigVMController::LocalizeFunctions
 			}
 			else if(URigVMFunctionReferenceNode* FunctionReferenceNode = Cast<URigVMFunctionReferenceNode>(NodeToUpdate))
 			{
-				URigVMLibraryNode* ReferencedNode = FunctionReferenceNode->GetReferencedNode();
+				const URigVMLibraryNode* ReferencedNode = FunctionReferenceNode->LoadReferencedNode();
 				URigVMLibraryNode** RemappedNodePtr = LocalizedFunctions.Find(ReferencedNode);
 				if(RemappedNodePtr)
 				{
@@ -4700,8 +4696,22 @@ TArray<URigVMNode*> URigVMController::ExpandLibraryNode(URigVMLibraryNode* InNod
 		return TArray<URigVMNode*>();
 	}
 
-	TArray<URigVMNode*> ContainedNodes = InNode->GetContainedNodes();
-	TArray<URigVMLink*> ContainedLinks = InNode->GetContainedLinks();
+	URigVMGraph* InnerGraph = InNode->GetContainedGraph();
+	if (URigVMFunctionReferenceNode* RefNode = Cast<URigVMFunctionReferenceNode>(InNode))
+	{
+		if (URigVMLibraryNode* LibraryNode = RefNode->LoadReferencedNode())
+		{
+			InnerGraph = LibraryNode->GetContainedGraph();
+		}
+		else
+		{
+			ReportError(TEXT("Cannot expand nodes from function reference because the source graph is not found."));
+			return TArray<URigVMNode*>();			
+		}
+	}	
+
+	TArray<URigVMNode*> ContainedNodes = InnerGraph->GetNodes();
+	TArray<URigVMLink*> ContainedLinks = InnerGraph->GetLinks();
 	if (ContainedNodes.Num() == 0)
 	{
 		return TArray<URigVMNode*>();
@@ -4762,10 +4772,10 @@ TArray<URigVMNode*> URigVMController::ExpandLibraryNode(URigVMLibraryNode* InNod
 	// exist, they will be reused. If a local variable is not used, it will not be created.
 	if (URigVMFunctionReferenceNode* FunctionReferenceNode = Cast<URigVMFunctionReferenceNode>(InNode))
 	{
-		TArray<FRigVMGraphVariableDescription> LocalVariables = FunctionReferenceNode->GetContainedGraph()->LocalVariables;
+		TArray<FRigVMGraphVariableDescription> LocalVariables = InnerGraph->LocalVariables;
 		TArray<FRigVMExternalVariable> CurrentVariables = GetAllVariables();
 		TArray<FRigVMGraphVariableDescription> VariablesToAdd;
-		for (const URigVMNode* Node : FunctionReferenceNode->GetContainedGraph()->GetNodes())
+		for (const URigVMNode* Node : InnerGraph->GetNodes())
 		{
 			if (const URigVMVariableNode* VariableNode = Cast<URigVMVariableNode>(Node))
 			{
@@ -4802,7 +4812,7 @@ TArray<URigVMNode*> URigVMController::ExpandLibraryNode(URigVMLibraryNode* InNod
 						}
 						else if(bVariableIncompatible)
 						{
-							ReportErrorf(TEXT("Found variable %s of incompatible type with a local variable inside function %s"), *LocalVariable.Name.ToString(), *FunctionReferenceNode->GetReferencedNode()->GetName());
+							ReportErrorf(TEXT("Found variable %s of incompatible type with a local variable inside function %s"), *LocalVariable.Name.ToString(), *FunctionReferenceNode->GetReferencedFunctionHeader().Name.ToString());
 							if (bSetupUndoRedo)
 							{
 								ActionStack->CancelAction(ExpandAction, this);
@@ -4829,7 +4839,7 @@ TArray<URigVMNode*> URigVMController::ExpandLibraryNode(URigVMLibraryNode* InNod
 
 	FString TextContent;
 	{
-		FRigVMControllerGraphGuard GraphGuard(this, InNode->GetContainedGraph(), false);
+		FRigVMControllerGraphGuard GraphGuard(this, InnerGraph, false);
 		TextContent = ExportNodesToText(NodeNames);
 	}
 
@@ -4893,11 +4903,11 @@ TArray<URigVMNode*> URigVMController::ExpandLibraryNode(URigVMLibraryNode* InNod
 
 	// c) create a map from the entry node to the contained graph
 	TMap<FString, TArray<FString>> FromEntryNode;
-	if (URigVMFunctionEntryNode* EntryNode = InNode->GetEntryNode())
+	if (URigVMFunctionEntryNode* EntryNode = InnerGraph->GetEntryNode())
 	{
 		TArray<URigVMLink*> EntryLinks = EntryNode->GetLinks();
 
-		for (URigVMNode* Node : InNode->GetContainedGraph()->GetNodes())
+		for (URigVMNode* Node : InnerGraph->GetNodes())
 		{
 			if (URigVMVariableNode* VariableNode = Cast<URigVMVariableNode>(Node))
 			{
@@ -4946,7 +4956,7 @@ TArray<URigVMNode*> URigVMController::ExpandLibraryNode(URigVMLibraryNode* InNod
 
 	// d) create a map from the contained graph from to the return node
 	TMap<FString, TArray<FString>> ToReturnNode;
-	if (URigVMFunctionReturnNode* ReturnNode = InNode->GetReturnNode())
+	if (URigVMFunctionReturnNode* ReturnNode = InnerGraph->GetReturnNode())
 	{
 		TArray<URigVMLink*> ReturnLinks = ReturnNode->GetLinks();
 		for (URigVMLink* Link : ReturnLinks)
@@ -5236,7 +5246,7 @@ TArray<URigVMNode*> URigVMController::ExpandLibraryNode(URigVMLibraryNode* InNod
 				FString ReturnNodeName, ReturnPinPath;
 				if (URigVMPin::SplitPinPathAtStart(RemappedTargetPinPath, ReturnNodeName, ReturnPinPath))
 				{
-					if(Cast<URigVMFunctionReturnNode>(InNode->GetContainedGraph()->FindNode(ReturnNodeName)))
+					if(Cast<URigVMFunctionReturnNode>(InnerGraph->FindNode(ReturnNodeName)))
 					{
 						if(FromLibraryNode.Contains(ReturnPinPath))
 						{
@@ -5605,7 +5615,7 @@ URigVMCollapseNode* URigVMController::PromoteFunctionReferenceNodeToCollapseNode
 	URigVMGraph* Graph = GetGraph();
 	check(Graph);
 
-	URigVMCollapseNode* FunctionDefinition = Cast<URigVMCollapseNode>(InFunctionRefNode->GetReferencedNode());
+	const URigVMCollapseNode* FunctionDefinition = Cast<URigVMCollapseNode>(InFunctionRefNode->LoadReferencedNode());
 	if (FunctionDefinition == nullptr)
 	{
 		return nullptr;
@@ -5687,7 +5697,8 @@ URigVMCollapseNode* URigVMController::PromoteFunctionReferenceNodeToCollapseNode
 		{
 			FRigVMByteArray ArchiveBytes;
 			FMemoryWriter ArchiveWriter(ArchiveBytes);
-			FunctionDefinition->Template.Serialize(ArchiveWriter);
+			FRigVMTemplate Template = FunctionDefinition->Template;
+			Template.Serialize(ArchiveWriter);
 
 			FMemoryReader ArchiveReader(ArchiveBytes);
 			CollapseNode->Template.Serialize(ArchiveReader);
@@ -5751,17 +5762,18 @@ void URigVMController::SetReferencedFunction(URigVMFunctionReferenceNode* InFunc
 		return;
 	}
 	
-	URigVMLibraryNode* OldReferencedNode = InFunctionRefNode->GetReferencedNode();
-	if(OldReferencedNode != InNewReferencedNode)
+	FRigVMGraphFunctionHeader OldReferencedNode = InFunctionRefNode->GetReferencedFunctionHeader();
+	InFunctionRefNode->ReferencedFunctionHeader = InNewReferencedNode->GetFunctionHeader();
+	
+	if(!(OldReferencedNode == InFunctionRefNode->ReferencedFunctionHeader))
 	{
 		if(URigVMBuildData* BuildData = GetBuildData())
 		{
-			BuildData->UnregisterFunctionReference(OldReferencedNode, InFunctionRefNode);
-			BuildData->RegisterFunctionReference(InNewReferencedNode, InFunctionRefNode);
+			BuildData->UnregisterFunctionReference(OldReferencedNode.LibraryPointer, InFunctionRefNode);
+			BuildData->RegisterFunctionReference(InFunctionRefNode->ReferencedFunctionHeader.LibraryPointer, InFunctionRefNode);
 		}
 	}
 
-	InFunctionRefNode->SetReferencedNode(InNewReferencedNode);
 	
 	FRigVMControllerGraphGuard GraphGuard(this, InFunctionRefNode->GetGraph(), false);
 	GetGraph()->Notify(ERigVMGraphNotifType::NodeReferenceChanged, InFunctionRefNode);
@@ -5954,7 +5966,7 @@ bool URigVMController::RemoveNode(URigVMNode* InNode, bool bSetupUndoRedo, bool
 		{
 			if(URigVMBuildData* BuildData = GetBuildData())
 			{
-				BuildData->UnregisterFunctionReference(FunctionReferenceNode->GetReferencedNode(), FunctionReferenceNode);
+				BuildData->UnregisterFunctionReference(FunctionReferenceNode->GetReferencedFunctionHeader().LibraryPointer, FunctionReferenceNode);
 			}
 		}
 		// If we are removing a function, remove all the references first
@@ -5962,7 +5974,8 @@ bool URigVMController::RemoveNode(URigVMNode* InNode, bool bSetupUndoRedo, bool
 		{
 			if(URigVMBuildData* BuildData = GetBuildData())
 			{
-				if (const FRigVMFunctionReferenceArray* ReferencesEntry = BuildData->FindFunctionReferences(LibraryNode))
+				FRigVMGraphFunctionIdentifier Identifier = LibraryNode->GetFunctionIdentifier();
+				if (const FRigVMFunctionReferenceArray* ReferencesEntry = BuildData->FindFunctionReferences(Identifier))
 				{
 					// make a copy since we'll be modifying the array
 					TArray< TSoftObjectPtr<URigVMFunctionReferenceNode> > FunctionReferences = ReferencesEntry->FunctionReferences;
@@ -5981,7 +5994,7 @@ bool URigVMController::RemoveNode(URigVMNode* InNode, bool bSetupUndoRedo, bool
 					}
 				}
 
-				BuildData->FunctionReferences.Remove(LibraryNode);
+				BuildData->GraphFunctionReferences.Remove(Identifier);
 			}
 			
 			for(const auto& Pair : FunctionLibrary->LocalizedFunctions)
@@ -6185,16 +6198,15 @@ bool URigVMController::RenameNode(URigVMNode* InNode, const FName& InNewName, bo
 			// update the table in the build data
 			if(URigVMBuildData* BuildData = GetBuildData())
 			{
-				for(const TPair< TSoftObjectPtr<URigVMLibraryNode>, FRigVMFunctionReferenceArray >& Pair: BuildData->FunctionReferences)
+				for(const TPair< FRigVMGraphFunctionIdentifier, FRigVMFunctionReferenceArray >& Pair: BuildData->GraphFunctionReferences)
 				{
-					if(Pair.Key.ToSoftObjectPath() == PreviousObjectPath)
+					if(Pair.Key.LibraryNode == PreviousObjectPath)
 					{
-						const TSoftObjectPtr<URigVMLibraryNode> SoftObjectPtr(InNode);
 						const FRigVMFunctionReferenceArray FunctionReferences = Pair.Value;
 						
 						BuildData->Modify();
-						BuildData->FunctionReferences.Remove(Pair.Key);
-						BuildData->FunctionReferences.Add(SoftObjectPtr, FunctionReferences);
+						BuildData->GraphFunctionReferences.Remove(Pair.Key);
+						BuildData->GraphFunctionReferences.Add(LibraryNode->GetFunctionIdentifier(), FunctionReferences);
 						BuildData->MarkPackageDirty();
 						break;
 					}
@@ -9546,23 +9558,26 @@ bool URigVMController::AddLink(URigVMPin* OutputPin, URigVMPin* InputPin, bool b
 	{
 		for (URigVMPin* Pin : LibraryNode->GetPins())
 		{
-			if (!LibraryNode->GetTemplate()->FindArgument(Pin->GetFName()))
+			if (const FRigVMTemplate* Template = LibraryNode->GetTemplate())
 			{
-				if (ShouldPinOwnArgument(Pin))
+				if (!Template->FindArgument(Pin->GetFName()))
 				{
-					URigVMPin* InterfacePin = nullptr;
-					if (URigVMFunctionEntryNode* EntryNode = LibraryNode->GetEntryNode())
+					if (ShouldPinOwnArgument(Pin))
 					{
-						InterfacePin = EntryNode->FindPin(Pin->GetName());
-					}
-					if (!InterfacePin)
-					{
-						if (URigVMFunctionReturnNode* ReturnNode = LibraryNode->GetReturnNode())
+						URigVMPin* InterfacePin = nullptr;
+						if (URigVMFunctionEntryNode* EntryNode = LibraryNode->GetEntryNode())
+						{
+							InterfacePin = EntryNode->FindPin(Pin->GetName());
+						}
+						if (!InterfacePin)
 						{
-							InterfacePin = ReturnNode->FindPin(Pin->GetName());
+							if (URigVMFunctionReturnNode* ReturnNode = LibraryNode->GetReturnNode())
+							{
+								InterfacePin = ReturnNode->FindPin(Pin->GetName());
+							}
 						}
+						AddArgumentForPin(InterfacePin, nullptr, bSetupUndoRedo);
 					}
-					AddArgumentForPin(InterfacePin, nullptr, bSetupUndoRedo);
 				}
 			}
 		}
@@ -11262,62 +11277,70 @@ URigVMFunctionReferenceNode* URigVMController::AddFunctionReferenceNode(URigVMLi
 	{
 		return nullptr;
 	}
-
-	if (!bIsTransacting && !IsGraphEditable())
+	
+	if (!InFunctionDefinition)
 	{
 		return nullptr;
 	}
 
-	URigVMGraph* Graph = GetGraph();
-	check(Graph);
+	if (URigVMFunctionReferenceNode* ReferenceNode = AddFunctionReferenceNodeFromDescription(InFunctionDefinition->GetFunctionHeader(), InNodePosition, InNodeName, bSetupUndoRedo, bPrintPythonCommand))
+	{
+		return ReferenceNode;
+	}
 
-	if (Graph->IsA<URigVMFunctionLibrary>())
+	return nullptr;
+}
+
+URigVMFunctionReferenceNode* URigVMController::AddFunctionReferenceNodeFromDescription(const FRigVMGraphFunctionHeader& InFunctionDefinition, const FVector2D& InNodePosition, const FString& InNodeName, bool bSetupUndoRedo, bool bPrintPythonCommand, bool bAllowPrivateFunctions)
+{
+	if (!IsValidGraph())
 	{
-		ReportError(TEXT("Cannot add function reference nodes to function library graphs."));
 		return nullptr;
 	}
 
-	if (InFunctionDefinition == nullptr)
+	if (!bIsTransacting && !IsGraphEditable())
 	{
-		ReportError(TEXT("Cannot add a function reference node without a valid function definition."));
 		return nullptr;
 	}
 
-	if (!InFunctionDefinition->GetGraph()->IsA<URigVMFunctionLibrary>())
+	URigVMGraph* Graph = GetGraph();
+	check(Graph);
+
+	if (Graph->IsA<URigVMFunctionLibrary>())
 	{
-		ReportAndNotifyError(TEXT("Cannot use the function definition for a function reference node."));
+		ReportError(TEXT("Cannot add function reference nodes to function library graphs."));
 		return nullptr;
 	}
 
-	if(!CanAddFunctionRefForDefinition(InFunctionDefinition, true))
+	if(!CanAddFunctionRefForDefinition(InFunctionDefinition, true, bAllowPrivateFunctions))
 	{
 		return nullptr;
 	}
 
-	FString NodeName = GetValidNodeName(InNodeName.IsEmpty() ? InFunctionDefinition->GetName() : InNodeName);
+	FString NodeName = GetValidNodeName(InNodeName.IsEmpty() ? InFunctionDefinition.Name.ToString() : InNodeName);
 	URigVMFunctionReferenceNode* FunctionRefNode = NewObject<URigVMFunctionReferenceNode>(Graph, *NodeName);
 	FunctionRefNode->Position = InNodePosition;
-	FunctionRefNode->SetReferencedNode(InFunctionDefinition);
+	FunctionRefNode->ReferencedFunctionHeader = InFunctionDefinition;
 	Graph->Nodes.Add(FunctionRefNode);
 	FunctionRefNode->InitializeFilteredPermutations();
 
 	FRigVMControllerCompileBracketScope CompileScope(this);
 	
 	RepopulatePinsOnNode(FunctionRefNode, false, false);
-	UpdateTemplateNodePinTypes(FunctionRefNode, false);
+	//UpdateTemplateNodePinTypes(FunctionRefNode, false); // We will need to add this if we ever allow functions to be templates
 
 	Notify(ERigVMGraphNotifType::NodeAdded, FunctionRefNode);
 
 	if (URigVMBuildData* BuildData = GetBuildData())
 	{
-		BuildData->RegisterFunctionReference(InFunctionDefinition, FunctionRefNode);
+		BuildData->RegisterFunctionReference(FunctionRefNode->GetReferencedFunctionHeader().LibraryPointer, FunctionRefNode);
 	}
-
-	for (URigVMPin* SourcePin : InFunctionDefinition->Pins)
+	
+	for (const FRigVMGraphFunctionArgument& Argument : InFunctionDefinition.Arguments)
 	{
-		if (URigVMPin* TargetPin = FunctionRefNode->FindPin(SourcePin->GetName()))
+		if (URigVMPin* TargetPin = FunctionRefNode->FindPin(Argument.Name.ToString()))
 		{
-			FString DefaultValue = SourcePin->GetDefaultValue();
+			const FString& DefaultValue = Argument.DefaultValue;
 			if (!DefaultValue.IsEmpty())
 			{
 				SetPinDefaultValue(TargetPin, DefaultValue, true, false, false);
@@ -11337,31 +11360,33 @@ URigVMFunctionReferenceNode* URigVMController::AddFunctionReferenceNode(URigVMLi
 
 	if (bPrintPythonCommand)
 	{
-		const FString GraphName = GetSanitizedGraphName(GetGraph()->GetGraphName());
-		const FString FunctionDefinitionName = GetSanitizedNodeName(InFunctionDefinition->GetName());
-
-		if (InFunctionDefinition->GetLibrary() == GetGraph()->GetDefaultFunctionLibrary())
-		{
-
-			RigVMPythonUtils::Print(GetGraphOuterName(), 
-				FString::Printf(TEXT("blueprint.get_controller_by_name('%s').add_function_reference_node(library.find_function('%s'), %s, '%s')"),
-						*GraphName,
-						*FunctionDefinitionName,
-						*RigVMPythonUtils::Vector2DToPythonString(InNodePosition),
-						*NodeName));
-		}
-		else
-		{
-			RigVMPythonUtils::Print(GetGraphOuterName(), 
-				FString::Printf(TEXT("function_blueprint = unreal.load_object(name = '%s', outer = None)"),
-				*InFunctionDefinition->GetLibrary()->GetOuter()->GetPathName()));
-			RigVMPythonUtils::Print(GetGraphOuterName(), 
-				FString::Printf(TEXT("blueprint.get_controller_by_name('%s').add_function_reference_node(function_blueprint.get_local_function_library().find_function('%s'), %s, '%s')"),
-						*GraphName,
-						*FunctionDefinitionName,
-						*RigVMPythonUtils::Vector2DToPythonString(InFunctionDefinition->GetPosition()), 
-						*FunctionDefinitionName));
-		}
+		// todo: UE-169886
+		// const FString GraphName = GetSanitizedGraphName(GetGraph()->GetGraphName());
+		// const FString FunctionDefinitionName = GetSanitizedNodeName(InFunctionDefinition.Name.ToString());
+		//
+		// URigVMLibraryNode* LibraryNode = InFunctionDefinition.LibraryNode.IsValid() ? InFunctionDefinition.LibraryNode.Get() : nullptr;
+		// if (LibraryNode && LibraryNode->GetLibrary() == GetGraph()->GetDefaultFunctionLibrary())
+		// {
+		//
+		// 	RigVMPythonUtils::Print(GetGraphOuterName(), 
+		// 		FString::Printf(TEXT("blueprint.get_controller_by_name('%s').add_function_reference_node(library.find_function('%s'), %s, '%s')"),
+		// 				*GraphName,
+		// 				*FunctionDefinitionName,
+		// 				*RigVMPythonUtils::Vector2DToPythonString(InNodePosition),
+		// 				*NodeName));
+		// }
+		// else
+		// {
+		// 	RigVMPythonUtils::Print(GetGraphOuterName(), 
+		// 		FString::Printf(TEXT("function_blueprint = unreal.load_object(name = '%s', outer = None)"),
+		// 		*InFunctionDefinition->GetLibrary()->GetOuter()->GetPathName()));
+		// 	RigVMPythonUtils::Print(GetGraphOuterName(), 
+		// 		FString::Printf(TEXT("blueprint.get_controller_by_name('%s').add_function_reference_node(function_blueprint.get_local_function_library().find_function('%s'), %s, '%s')"),
+		// 				*GraphName,
+		// 				*FunctionDefinitionName,
+		// 				*RigVMPythonUtils::Vector2DToPythonString(InFunctionDefinition->GetPosition()), 
+		// 				*FunctionDefinitionName));
+		// }
 		
 	}
 
@@ -11407,8 +11432,13 @@ bool URigVMController::SetRemappedVariable(URigVMFunctionReferenceNode* InFuncti
 
 	FRigVMExternalVariable InnerExternalVariable;
 	{
-		FRigVMControllerGraphGuard GraphGuard(this, InFunctionRefNode->GetContainedGraph());
-		InnerExternalVariable = GetVariableByName(InInnerVariableName);
+		if (FRigVMExternalVariable* Variable = InFunctionRefNode->ReferencedFunctionHeader.ExternalVariables.FindByPredicate([InInnerVariableName](const FRigVMExternalVariable& Variable)
+		{
+			return Variable.Name == InInnerVariableName;
+		}))
+		{
+			InnerExternalVariable = *Variable;
+		}
 	}
 
 	if(!InnerExternalVariable.IsValid(true))
@@ -11490,6 +11520,8 @@ URigVMLibraryNode* URigVMController::AddFunctionToLibrary(const FName& InFunctio
 
 	FRigVMControllerCompileBracketScope CompileScope(this);
 	
+	Notify(ERigVMGraphNotifType::NodeAdded, CollapseNode);
+	
 	if (bMutable)
 	{
 		const UScriptStruct* ExecuteContextStruct = FRigVMExecuteContext::StaticStruct();
@@ -11503,8 +11535,6 @@ URigVMLibraryNode* URigVMController::AddFunctionToLibrary(const FName& InFunctio
 			bSetupUndoRedo);
 	}
 
-	Notify(ERigVMGraphNotifType::NodeAdded, CollapseNode);
-
 	{
 		FRigVMControllerGraphGuard GraphGuard(this, CollapseNode->GetContainedGraph(), false);
 		TGuardValue<bool> GuardEditGraph(CollapseNode->ContainedGraph->bEditable, true);
@@ -11517,6 +11547,7 @@ URigVMLibraryNode* URigVMController::AddFunctionToLibrary(const FName& InFunctio
 			EntryNode = NewObject<URigVMFunctionEntryNode>(CollapseNode->ContainedGraph, TEXT("Entry"));
 			CollapseNode->ContainedGraph->Nodes.Add(EntryNode);
 			RefreshFunctionPins(EntryNode, false);
+			EntryNode->Position = FVector2D(-250.f, 0.f);
 			Notify(ERigVMGraphNotifType::NodeAdded, EntryNode);
 		}
 
@@ -11525,11 +11556,10 @@ URigVMLibraryNode* URigVMController::AddFunctionToLibrary(const FName& InFunctio
 			ReturnNode = NewObject<URigVMFunctionReturnNode>(CollapseNode->ContainedGraph, TEXT("Return"));
 			CollapseNode->ContainedGraph->Nodes.Add(ReturnNode);
 			RefreshFunctionPins(ReturnNode, false);
+			ReturnNode->Position = FVector2D(250.f, 0.f);
 			Notify(ERigVMGraphNotifType::NodeAdded, ReturnNode);
 		}
 
-		EntryNode->Position = FVector2D(-250.f, 0.f);
-		ReturnNode->Position = FVector2D(250.f, 0.f);
 
 		if (bMutable)
 		{
@@ -12305,21 +12335,18 @@ TArray<TSoftObjectPtr<URigVMFunctionReferenceNode>> URigVMController::GetAffecte
 		{
 			if(URigVMFunctionReferenceNode* AffectedFunctionReferenceNode = FunctionReferencePtr.Get())
 			{
-				if(URigVMFunctionLibrary* AffectedFunctionLibrary = AffectedFunctionReferenceNode->GetTypedOuter<URigVMFunctionLibrary>())
+				if(URigVMLibraryNode* AffectedFunction = AffectedFunctionReferenceNode->FindFunctionForNode())
 				{
-					if(URigVMLibraryNode* AffectedFunction = AffectedFunctionLibrary->FindFunctionForNode(AffectedFunctionReferenceNode))
+					FRigVMControllerGraphGuard GraphGuard(this, AffectedFunction->GetContainedGraph(), false);
+					TArray<TSoftObjectPtr<URigVMFunctionReferenceNode>> AffectedFunctionReferencePtrs = GetAffectedReferences(InEditType, bForceLoad, false);
+					for(TSoftObjectPtr<URigVMFunctionReferenceNode> AffectedFunctionReferencePtr : AffectedFunctionReferencePtrs)
 					{
-						FRigVMControllerGraphGuard GraphGuard(this, AffectedFunction->GetContainedGraph(), false);
-						TArray<TSoftObjectPtr<URigVMFunctionReferenceNode>> AffectedFunctionReferencePtrs = GetAffectedReferences(InEditType, bForceLoad, false);
-						for(TSoftObjectPtr<URigVMFunctionReferenceNode> AffectedFunctionReferencePtr : AffectedFunctionReferencePtrs)
+						const FString Key = AffectedFunctionReferencePtr.ToSoftObjectPath().ToString();
+						if(VisitedPaths.Contains(Key))
 						{
-							const FString Key = AffectedFunctionReferencePtr.ToSoftObjectPath().ToString();
-							if(VisitedPaths.Contains(Key))
-							{
-								continue;
-							}
-							VisitedPaths.Add(Key, FunctionReferencePtrs.Add(AffectedFunctionReferencePtr));
+							continue;
 						}
+						VisitedPaths.Add(Key, FunctionReferencePtrs.Add(AffectedFunctionReferencePtr));
 					}
 				}
 			}
@@ -13741,19 +13768,20 @@ bool URigVMController::CanAddNode(URigVMNode* InNode, bool bReportErrors, bool b
 	if (URigVMFunctionReferenceNode* FunctionRefNode = Cast<URigVMFunctionReferenceNode>(InNode))
 	{
 		if (URigVMFunctionLibrary* FunctionLibrary = FunctionRefNode->GetLibrary())
-		{
-			if (URigVMLibraryNode* FunctionDefinition = FunctionRefNode->GetReferencedNode())
+		{			
+			FRigVMGraphFunctionHeader FunctionDefinition = FunctionRefNode->GetReferencedFunctionHeader();
+			if(!CanAddFunctionRefForDefinition(FunctionDefinition, false))
 			{
-				if(!CanAddFunctionRefForDefinition(FunctionDefinition, false))
+				if (URigVMLibraryNode* LibraryNode = FunctionRefNode->LoadReferencedNode())
 				{
 					URigVMFunctionLibrary* TargetLibrary = Graph->GetDefaultFunctionLibrary();
-					URigVMLibraryNode* NewFunctionDefinition = TargetLibrary->FindPreviouslyLocalizedFunction(FunctionDefinition);
-					
+					URigVMLibraryNode* NewFunctionDefinition = TargetLibrary->FindPreviouslyLocalizedFunction(LibraryNode);
+				
 					if((NewFunctionDefinition == nullptr) && RequestLocalizeFunctionDelegate.IsBound())
 					{
-						if(RequestLocalizeFunctionDelegate.Execute(FunctionDefinition))
+						if(RequestLocalizeFunctionDelegate.Execute(LibraryNode))
 						{
-							NewFunctionDefinition = TargetLibrary->FindPreviouslyLocalizedFunction(FunctionDefinition);
+							NewFunctionDefinition = TargetLibrary->FindPreviouslyLocalizedFunction(LibraryNode);
 						}
 					}
 
@@ -13761,17 +13789,17 @@ bool URigVMController::CanAddNode(URigVMNode* InNode, bool bReportErrors, bool b
 					{
 						return false;
 					}
-					
-					SetReferencedFunction(FunctionRefNode, NewFunctionDefinition, false);
-					FunctionDefinition = NewFunctionDefinition;
-				}
 				
-				if(!CanAddFunctionRefForDefinition(FunctionDefinition, bReportErrors))
-				{
-					DestroyObject(InNode);
-					return false;
+					SetReferencedFunction(FunctionRefNode, NewFunctionDefinition, false);
+					FunctionDefinition = FunctionRefNode->GetReferencedFunctionHeader();
 				}
 			}
+			
+			if(!CanAddFunctionRefForDefinition(FunctionDefinition, bReportErrors))
+			{
+				DestroyObject(InNode);
+				return false;
+			}
 		}			
 	}
 	else if(!bIgnoreFunctionEntryReturnNodes &&
@@ -13921,7 +13949,7 @@ bool URigVMController::CanAddEventNode(UScriptStruct* InScriptStruct, const bool
 	return !bHasEventNode;
 }
 
-bool URigVMController::CanAddFunctionRefForDefinition(URigVMLibraryNode* InFunctionDefinition, bool bReportErrors)
+bool URigVMController::CanAddFunctionRefForDefinition(const FRigVMGraphFunctionHeader& InFunctionDefinition, bool bReportErrors, bool bAllowPrivateFunctions)
 {
 	if(!IsValidGraph())
 	{
@@ -13933,14 +13961,22 @@ bool URigVMController::CanAddFunctionRefForDefinition(URigVMLibraryNode* InFunct
 		return false;
 	}
 
-	check(InFunctionDefinition);
-
 	URigVMGraph* Graph = GetGraph();
 	check(Graph);
 
-	if(IsFunctionAvailableDelegate.IsBound())
+	if (!bAllowPrivateFunctions)
 	{
-		if(!IsFunctionAvailableDelegate.Execute(InFunctionDefinition))
+		IRigVMClientHost* ClientHost = GetImplementingOuter<IRigVMClientHost>();
+		bool bIsAvailable = ClientHost->GetRigVMClient()->GetFunctionLibrary()->GetFunctionHostObjectPath() ==
+			InFunctionDefinition.LibraryPointer.HostObject;
+		if (!bIsAvailable)
+		{
+			if (IRigVMGraphFunctionHost* Host = Cast<IRigVMGraphFunctionHost>(InFunctionDefinition.LibraryPointer.HostObject.ResolveObject()))
+			{
+				bIsAvailable = Host->GetRigVMGraphFunctionStore()->IsFunctionPublic(InFunctionDefinition.LibraryPointer);		
+			}
+		}
+		if (!bIsAvailable)
 		{
 			if(bReportErrors)
 			{
@@ -13950,22 +13986,25 @@ bool URigVMController::CanAddFunctionRefForDefinition(URigVMLibraryNode* InFunct
 		}
 	}
 
-	if(IsDependencyCyclicDelegate.IsBound())
+	if (URigVMNode* Node = Cast<URigVMNode>(Graph->GetOuter()))
 	{
-		if(IsDependencyCyclicDelegate.Execute(Graph, InFunctionDefinition))
+		if (URigVMLibraryNode* LibraryNode = Node->FindFunctionForNode())
 		{
-			if(bReportErrors)
+			if (InFunctionDefinition.Dependencies.Contains(LibraryNode->GetFunctionIdentifier()))
 			{
-				ReportAndNotifyError(TEXT("Function is not available for placement in this graph host due to dependency cycles."));
+				if(bReportErrors)
+				{
+					ReportAndNotifyError(TEXT("Function is not available for placement in this graph host due to dependency cycles."));
+				}
+				return false;
 			}
-			return false;
 		}
 	}
-
+	
 	URigVMLibraryNode* ParentLibraryNode = Cast<URigVMLibraryNode>(Graph->GetOuter());
 	while (ParentLibraryNode)
 	{
-		if (ParentLibraryNode == InFunctionDefinition)
+		if (TSoftObjectPtr<UObject>(ParentLibraryNode).ToSoftObjectPath() == InFunctionDefinition.LibraryPointer.LibraryNode)
 		{
 			if(bReportErrors)
 			{
@@ -14345,6 +14384,29 @@ void URigVMController::ConfigurePinFromPin(URigVMPin* InOutPin, URigVMPin* InPin
 	}
 }
 
+void URigVMController::ConfigurePinFromArgument(URigVMPin* InOutPin, const FRigVMGraphFunctionArgument& InArgument, bool bCopyDisplayName)
+{
+	// it is important we copy things that define the identity of the pin
+	// things that defines the state of the pin is copied during GetPinState()
+	// though addmittedly these two functions have overlaps currently
+	InOutPin->bIsConstant = InArgument.bIsConst;
+	InOutPin->Direction = InArgument.Direction;
+	InOutPin->CPPType = InArgument.CPPType.ToString();
+	InOutPin->CPPTypeObjectPath = *InArgument.CPPTypeObject.ToSoftObjectPath().ToString();
+	InOutPin->CPPTypeObject = InArgument.CPPTypeObject.Get();
+	InOutPin->DefaultValue = InArgument.DefaultValue;
+	InOutPin->bIsDynamicArray = InArgument.bIsArray;
+	if(bCopyDisplayName)
+	{
+		InOutPin->SetDisplayName(InArgument.DisplayName);
+	}
+
+	if(InOutPin->IsExecuteContext() && InOutPin->CPPTypeObject != FRigVMExecuteContext::StaticStruct())
+	{
+		MakeExecutePin(InOutPin);
+	}
+}
+
 bool URigVMController::ShouldStructBeUnfolded(const UStruct* Struct)
 {
 	if (Struct == nullptr)
@@ -15239,31 +15301,38 @@ void URigVMController::RepopulatePinsOnNode(URigVMNode* InNode, bool bFollowCore
 	}
 	else if (FunctionRefNode)
 	{
-		if (URigVMLibraryNode* ReferencedNode = FunctionRefNode->GetReferencedNode())
+		FRigVMGraphFunctionHeader& Header = FunctionRefNode->ReferencedFunctionHeader;
+		
+		// we want to make sure notify the graph of a potential name change
+		// when repopulating the function ref node
+		Notify(ERigVMGraphNotifType::NodeRenamed, FunctionRefNode);
+		RemovePinsDuringRepopulate(InNode, InNode->Pins, bNotify, bSetupOrphanedPins);
+
+		TMap<FString, FPinState> ReferencedPinStates;
+		if (Header.LibraryPointer.LibraryNode.ResolveObject())
 		{
-			// we want to make sure notify the graph of a potential name change
-			// when repopulating the function ref node
-			Notify(ERigVMGraphNotifType::NodeRenamed, FunctionRefNode);
-			RemovePinsDuringRepopulate(InNode, InNode->Pins, bNotify, bSetupOrphanedPins);
+			URigVMLibraryNode* LibraryNode = Cast<URigVMLibraryNode>(Header.LibraryPointer.LibraryNode.ResolveObject());
+			ReferencedPinStates = GetPinStates(LibraryNode);
+		}
 
-			TMap<FString, FPinState> ReferencedPinStates = GetPinStates(ReferencedNode);
+		for (const FRigVMGraphFunctionArgument& ReferencedArgument : Header.Arguments)
+		{
+			URigVMPin* NewPin = NewObject<URigVMPin>(InNode, ReferencedArgument.Name);
+			ConfigurePinFromArgument(NewPin, ReferencedArgument, true);
+			EnsurePinValidity(NewPin, false);
+			
+			AddNodePin(InNode, NewPin);
 
-			for (URigVMPin* ReferencedPin : ReferencedNode->Pins)
+			if (NewPin->IsStruct())
 			{
-				URigVMPin* NewPin = NewObject<URigVMPin>(InNode, ReferencedPin->GetFName());
-				ConfigurePinFromPin(NewPin, ReferencedPin, true);
-				EnsurePinValidity(NewPin, false);
-				
-				AddNodePin(InNode, NewPin);
-
-				if (NewPin->IsStruct())
-				{
-					AddPinsForStruct(NewPin->GetScriptStruct(), InNode, NewPin, NewPin->GetDirection(), FString(), false);
-				}
-
-				Notify(ERigVMGraphNotifType::PinAdded, NewPin);
+				AddPinsForStruct(NewPin->GetScriptStruct(), InNode, NewPin, NewPin->GetDirection(), FString(), false);
 			}
 
+			Notify(ERigVMGraphNotifType::PinAdded, NewPin);
+		}
+
+		if (ReferencedPinStates.Num() > 0)
+		{
 			ApplyPinStates(InNode, ReferencedPinStates);
 		}
 	}
diff --git a/Engine/Source/Developer/RigVMDeveloper/Private/RigVMModel/RigVMFunctionLibrary.cpp b/Engine/Source/Developer/RigVMDeveloper/Private/RigVMModel/RigVMFunctionLibrary.cpp
--- a/Engine/Source/Developer/RigVMDeveloper/Private/RigVMModel/RigVMFunctionLibrary.cpp
+++ b/Engine/Source/Developer/RigVMDeveloper/Private/RigVMModel/RigVMFunctionLibrary.cpp
@@ -107,7 +107,7 @@ void URigVMFunctionLibrary::ForEachReference(const FName& InFunctionName,
 	{
 		if(const URigVMBuildData* BuildData = GetBuildData())
 		{
-			BuildData->ForEachFunctionReference(Function, PerReferenceFunction);
+			BuildData->ForEachFunctionReference(Function->GetFunctionIdentifier(),PerReferenceFunction);
 		}
 	}
 }
@@ -119,7 +119,7 @@ void URigVMFunctionLibrary::ForEachReferenceSoftPtr(const FName& InFunctionName,
 	{
 		if(const URigVMBuildData* BuildData = GetBuildData())
 		{
-			BuildData->ForEachFunctionReferenceSoftPtr(Function, PerReferenceFunction);
+			BuildData->ForEachFunctionReferenceSoftPtr(Function->GetFunctionIdentifier(), PerReferenceFunction);
 		}
 	}
 }
@@ -162,4 +162,13 @@ URigVMLibraryNode* URigVMFunctionLibrary::FindPreviouslyLocalizedFunction(URigVM
 	return LocalizedFunction;
 }
 
+const FSoftObjectPath URigVMFunctionLibrary::GetFunctionHostObjectPath() const
+{
+	if (GetFunctionHostObjectPathDelegate.IsBound())
+	{
+		return GetFunctionHostObjectPathDelegate.Execute();
+	}
+	return FSoftObjectPath();
+}
+
 
diff --git a/Engine/Source/Developer/RigVMDeveloper/Private/RigVMModel/RigVMGraph.cpp b/Engine/Source/Developer/RigVMDeveloper/Private/RigVMModel/RigVMGraph.cpp
--- a/Engine/Source/Developer/RigVMDeveloper/Private/RigVMModel/RigVMGraph.cpp
+++ b/Engine/Source/Developer/RigVMDeveloper/Private/RigVMModel/RigVMGraph.cpp
@@ -169,9 +169,9 @@ URigVMNode* URigVMGraph::FindNode(const FString& InNodePath) const
 		return FindNodeByName(*Left);
 	}
 
-	if (URigVMLibraryNode* LibraryNode = Cast< URigVMLibraryNode>(FindNodeByName(*Left)))
+	if (URigVMCollapseNode* CollapseNode = Cast< URigVMCollapseNode>(FindNodeByName(*Left)))
 	{
-		return LibraryNode->GetContainedGraph()->FindNode(Right);
+		return CollapseNode->GetContainedGraph()->FindNode(Right);
 	}
 
 	return nullptr;
diff --git a/Engine/Source/Developer/RigVMDeveloper/Private/RigVMModel/RigVMNode.cpp b/Engine/Source/Developer/RigVMDeveloper/Private/RigVMModel/RigVMNode.cpp
--- a/Engine/Source/Developer/RigVMDeveloper/Private/RigVMModel/RigVMNode.cpp
+++ b/Engine/Source/Developer/RigVMDeveloper/Private/RigVMModel/RigVMNode.cpp
@@ -3,6 +3,8 @@
 #include "RigVMModel/RigVMNode.h"
 #include "RigVMModel/Nodes/RigVMUnitNode.h"
 #include "RigVMModel/RigVMGraph.h"
+#include "RigVMModel/Nodes/RigVMLibraryNode.h"
+#include "RigVMModel/RigVMFunctionLibrary.h"
 #include "RigVMCore/RigVMExecuteContext.h"
 #include "RigVMCore/RigVMStruct.h"
 #include "RigVMUserWorkflowRegistry.h"
@@ -390,6 +392,21 @@ bool URigVMNode::IsLinkedTo(URigVMNode* InNode) const
 	return false;
 }
 
+URigVMLibraryNode* URigVMNode::FindFunctionForNode() 
+{
+	UObject* Subject = this;
+	while (Subject->GetOuter() && !Subject->GetOuter()->IsA<URigVMFunctionLibrary>())
+	{
+		Subject = Subject->GetOuter();
+		if(Subject == nullptr)
+		{
+			return nullptr;
+		}
+	}
+
+	return Cast<URigVMLibraryNode>(Subject);
+}
+
 bool URigVMNode::IsLinkedToRecursive(URigVMPin* InPin, URigVMNode* InNode) const
 {
 	for (URigVMPin* LinkedPin : InPin->GetLinkedSourcePins())
diff --git a/Engine/Source/Developer/RigVMDeveloper/Public/RigVMCompiler/RigVMAST.h b/Engine/Source/Developer/RigVMDeveloper/Public/RigVMCompiler/RigVMAST.h
--- a/Engine/Source/Developer/RigVMDeveloper/Public/RigVMCompiler/RigVMAST.h
+++ b/Engine/Source/Developer/RigVMDeveloper/Public/RigVMCompiler/RigVMAST.h
@@ -714,6 +714,9 @@ public:
 		return InType == EType::InlineFunction;
 	};
 
+	// todo: For now, all function references are considered varying (UE-170129)
+	virtual bool IsConstant() const override { return false; }
+
 protected:
 
 	// default constructor (protected so that only parser can access it)
diff --git a/Engine/Source/Developer/RigVMDeveloper/Public/RigVMCompiler/RigVMCompiler.h b/Engine/Source/Developer/RigVMDeveloper/Public/RigVMCompiler/RigVMCompiler.h
--- a/Engine/Source/Developer/RigVMDeveloper/Public/RigVMCompiler/RigVMCompiler.h
+++ b/Engine/Source/Developer/RigVMDeveloper/Public/RigVMCompiler/RigVMCompiler.h
@@ -10,7 +10,6 @@
 #include "RigVMModel/Nodes/RigVMParameterNode.h"
 #include "RigVMModel/Nodes/RigVMCommentNode.h"
 #include "RigVMModel/Nodes/RigVMRerouteNode.h"
-#include "RigVMModel/Nodes/RigVMFunctionReferenceNode.h"
 #include "RigVMCore/RigVM.h"
 #include "RigVMCore/RigVMStruct.h"
 #include "RigVMCompiler/RigVMAST.h"
@@ -18,6 +17,11 @@
 
 #include "RigVMCompiler.generated.h"
 
+class URigVMFunctionReferenceNode;
+struct FRigVMGraphFunctionData;
+struct FRigVMGraphFunctionHeader;
+struct FRigVMFunctionCompilationData;
+
 USTRUCT(BlueprintType)
 struct RIGVMDEVELOPER_API FRigVMCompileSettings
 {
@@ -109,8 +113,8 @@ public:
 
 		friend inline uint32 GetTypeHash(const FFunctionRegisterData& Data)
 		{
-			uint32 Result = HashCombine(::GetTypeHash(Data.ReferenceNode.IsValid() ? Data.ReferenceNode.Get() : 0), ::GetTypeHash(Data.MemoryType));
-			return HashCombine(Result, ::GetTypeHash(Data.RegisterIndex));
+			uint32 Result = HashCombine(GetTypeHash(Data.ReferenceNode), GetTypeHash(Data.MemoryType));
+			return HashCombine(Result, GetTypeHash(Data.RegisterIndex));
 		}
 
 		bool operator==(const FFunctionRegisterData& Other) const
@@ -146,6 +150,8 @@ public:
 	TMap<FRigVMOperand, FCopyOpInfo> DeferredCopyOps;
 };
 
+DECLARE_DELEGATE_RetVal_OneParam(const FRigVMFunctionCompilationData*, FRigVMCompiler_GetFunctionCompilationData, const FRigVMGraphFunctionHeader& Header);
+
 UCLASS(BlueprintType)
 class RIGVMDEVELOPER_API URigVMCompiler : public UObject
 {
@@ -164,21 +170,15 @@ public:
 		return Compile(InGraphs, InController, OutVM, TArray<FRigVMExternalVariable>(), TArray<FRigVMUserDataArray>(), nullptr);
 	}
 
-	bool Compile(TArray<URigVMGraph*> InGraphs, URigVMController* InController, URigVM* OutVM, const TArray<FRigVMExternalVariable>& InExternalVariables, const TArray<FRigVMUserDataArray>& InRigVMUserData, TMap<FString, FRigVMOperand>* OutOperands, TSharedPtr<FRigVMParserAST> InAST = TSharedPtr<FRigVMParserAST>());
+	bool Compile(TArray<URigVMGraph*> InGraphs, URigVMController* InController, URigVM* OutVM, const TArray<FRigVMExternalVariable>& InExternalVariables, const TArray<FRigVMUserDataArray>& InRigVMUserData, TMap<FString, FRigVMOperand>* OutOperands, TSharedPtr<FRigVMParserAST> InAST = TSharedPtr<FRigVMParserAST>(), FRigVMFunctionCompilationData* OutFunctionCompilationData = nullptr);
 
-	struct FunctionCompilationData
-	{
-		FRigVMByteCode ByteCode;
-		TArray<FName> FunctionNames;
-		TMap<ERigVMMemoryType, TArray<FRigVMPropertyDescription>> PropertyDescriptions;
-		TMap<ERigVMMemoryType, TArray<FRigVMPropertyPathDescription>> PropertyPathDescriptions;
-		TMap<int32, FName> ExternalRegisterIndexToVariable;
-		TMap<FString, FRigVMOperand> Operands;
-	};
-	TMap<TSoftObjectPtr<URigVMLibraryNode>, FunctionCompilationData> CompiledFunctions;
+	bool CompileFunction(const URigVMLibraryNode* InLibraryNode, URigVMController* InController, FRigVMFunctionCompilationData* OutFunctionCompilationData);
+
+	FRigVMCompiler_GetFunctionCompilationData GetFunctionCompilationData;
+	TMap<FString, const FRigVMFunctionCompilationData*> CompiledFunctions;
 
 	static UScriptStruct* GetScriptStructForCPPType(const FString& InCPPType);
-	static FString GetPinHash(const URigVMPin* InPin, const FRigVMVarExprAST* InVarExpr, bool bIsDebugValue = false, URigVMLibraryNode* FunctionCompiling = nullptr, const FRigVMASTProxy& InPinProxy = FRigVMASTProxy());
+	static FString GetPinHash(const URigVMPin* InPin, const FRigVMVarExprAST* InVarExpr, bool bIsDebugValue = false, const URigVMLibraryNode* FunctionCompiling = nullptr, const FRigVMASTProxy& InPinProxy = FRigVMASTProxy());
 
 	// follows assignment expressions to find the source ref counted containers
 	// since ref counted containers are not copied for assignments.
@@ -189,15 +189,12 @@ public:
 
 private:
 
-	bool CompileFunction(URigVMLibraryNode* InLibraryNode, URigVMController* InController, const TArray<FRigVMUserDataArray>& InRigVMUserData);
-	void ClearFunctionCompilationData();
-	void RemoveFunctionCompilationData(TSoftObjectPtr<URigVMLibraryNode> LibraryNode);
-	URigVMLibraryNode* CurrentCompilationFunction;
+	const URigVMLibraryNode* CurrentCompilationFunction;
 
 	TArray<URigVMPin*> GetLinkedPins(URigVMPin* InPin, bool bInputs = true, bool bOutputs = true, bool bRecursive = true);
 	uint16 GetElementSizeFromCPPType(const FString& InCPPType, UScriptStruct* InScriptStruct);
 
-	static FString GetPinHashImpl(const URigVMPin* InPin, const FRigVMVarExprAST* InVarExpr, bool bIsDebugValue = false, URigVMLibraryNode* FunctionCompiling = nullptr, const FRigVMASTProxy& InPinProxy = FRigVMASTProxy());
+	static FString GetPinHashImpl(const URigVMPin* InPin, const FRigVMVarExprAST* InVarExpr, bool bIsDebugValue = false, const URigVMLibraryNode* FunctionCompiling = nullptr, const FRigVMASTProxy& InPinProxy = FRigVMASTProxy());
 
 	void TraverseExpression(const FRigVMExprAST* InExpr, FRigVMCompilerWorkData& WorkData);
 	void TraverseChildren(const FRigVMExprAST* InExpr, FRigVMCompilerWorkData& WorkData);
diff --git a/Engine/Source/Developer/RigVMDeveloper/Public/RigVMModel/Nodes/RigVMFunctionReferenceNode.h b/Engine/Source/Developer/RigVMDeveloper/Public/RigVMModel/Nodes/RigVMFunctionReferenceNode.h
--- a/Engine/Source/Developer/RigVMDeveloper/Public/RigVMModel/Nodes/RigVMFunctionReferenceNode.h
+++ b/Engine/Source/Developer/RigVMDeveloper/Public/RigVMModel/Nodes/RigVMFunctionReferenceNode.h
@@ -2,6 +2,7 @@
 
 #pragma once
 
+#include "RigVMCore/RigVMGraphFunctionDefinition.h"
 #include "RigVMModel/RigVMGraph.h"
 #include "RigVMModel/Nodes/RigVMLibraryNode.h"
 #include "RigVMFunctionReferenceNode.generated.h"
@@ -28,13 +29,11 @@ public:
 	// URigVMLibraryNode interface
 	virtual FString GetNodeCategory() const override;
 	virtual FString GetNodeKeywords() const override;
-	virtual URigVMFunctionLibrary* GetLibrary() const override;
-	virtual URigVMGraph* GetContainedGraph() const override;
 	virtual TArray<FRigVMExternalVariable> GetExternalVariables() const override;
-	virtual const FRigVMTemplate* GetTemplate() const override;
+	virtual const FRigVMTemplate* GetTemplate() const override { return nullptr; }
 	// end URigVMLibraryNode interface
 
-	URigVMLibraryNode* GetReferencedNode() const;
+	URigVMLibraryNode* LoadReferencedNode() const;
 
 	// Variable remapping
 	bool RequiresVariableRemapping() const;
@@ -44,15 +43,23 @@ public:
 	FName GetOuterVariableName(const FName& InInnerVariableName) const;
 	// end Variable remapping
 
+	const FRigVMGraphFunctionHeader& GetReferencedFunctionHeader() const { return ReferencedFunctionHeader; }
+
+	const FRigVMGraphFunctionData* GetReferencedFunctionData() const;
+
+	
 private:
 
 	virtual FText GetToolTipTextForPin(const URigVMPin* InPin) const override;
 	bool RequiresVariableRemappingInternal(TArray<FRigVMExternalVariable>& InnerVariables) const;
 
-	void SetReferencedNode(URigVMLibraryNode* InReferenceNode);
+	//void SetReferencedFunctionData(FRigVMGraphFunctionData* Data);
 
 	UPROPERTY(AssetRegistrySearchable)
-	mutable TSoftObjectPtr<URigVMLibraryNode> ReferencedNodePtr;
+	FRigVMGraphFunctionHeader ReferencedFunctionHeader;
+	
+	UPROPERTY(AssetRegistrySearchable, meta=(DeprecatedProperty=5.2))
+	TSoftObjectPtr<URigVMLibraryNode> ReferencedNodePtr_DEPRECATED;
 
 	UPROPERTY()
 	TMap<FName, FName> VariableMap;
@@ -60,5 +67,7 @@ private:
 	friend class URigVMController;
 	friend class FRigVMParserAST;
 	friend class UControlRigBlueprint;
+	friend struct FRigVMClient;
+	friend struct EngineTestRigVMFramework;
 };
 
diff --git a/Engine/Source/Developer/RigVMDeveloper/Public/RigVMModel/Nodes/RigVMLibraryNode.h b/Engine/Source/Developer/RigVMDeveloper/Public/RigVMModel/Nodes/RigVMLibraryNode.h
--- a/Engine/Source/Developer/RigVMDeveloper/Public/RigVMModel/Nodes/RigVMLibraryNode.h
+++ b/Engine/Source/Developer/RigVMDeveloper/Public/RigVMModel/Nodes/RigVMLibraryNode.h
@@ -3,6 +3,8 @@
 #pragma once
 
 #include "RigVMModel/Nodes/RigVMTemplateNode.h"
+#include "RigVMCore/RigVMGraphFunctionDefinition.h"
+#include "RigVMCore/RigVMGraphFunctionHost.h"
 #include "RigVMLibraryNode.generated.h"
 
 class URigVMGraph;
@@ -54,6 +56,10 @@ public:
 	virtual URigVMFunctionReturnNode* GetReturnNode() const;
 	virtual bool Contains(URigVMLibraryNode* InContainedNode, bool bRecursive = true) const;
 	virtual TArray<FRigVMExternalVariable> GetExternalVariables() const;
+	virtual TMap<FRigVMGraphFunctionIdentifier, uint32> GetDependencies() const;
+
+	FRigVMGraphFunctionIdentifier GetFunctionIdentifier() const;
+	FRigVMGraphFunctionHeader GetFunctionHeader(IRigVMGraphFunctionHost* InHostObject = nullptr) const;
 
 protected:
 
diff --git a/Engine/Source/Developer/RigVMDeveloper/Public/RigVMModel/RigVMBuildData.h b/Engine/Source/Developer/RigVMDeveloper/Public/RigVMModel/RigVMBuildData.h
--- a/Engine/Source/Developer/RigVMDeveloper/Public/RigVMModel/RigVMBuildData.h
+++ b/Engine/Source/Developer/RigVMDeveloper/Public/RigVMModel/RigVMBuildData.h
@@ -35,7 +35,6 @@ public:
 	
 	FRigVMReferenceNodeData()
 		:ReferenceNodePath()
-		,ReferencedFunctionPath()
 	{}
 
 	FRigVMReferenceNodeData(URigVMFunctionReferenceNode* InReferenceNode);
@@ -43,18 +42,18 @@ public:
 	UPROPERTY()
 	FString ReferenceNodePath;
 
+	UPROPERTY(meta=(DeprecatedProperty))
+	FString ReferencedFunctionPath_DEPRECATED;
+	
 	UPROPERTY()
-	FString ReferencedFunctionPath;
+	FRigVMGraphFunctionHeader ReferencedHeader;
 
 	TSoftObjectPtr<URigVMFunctionReferenceNode> GetReferenceNodeObjectPath();
-	TSoftObjectPtr<URigVMLibraryNode> GetReferencedFunctionObjectPath();
 	URigVMFunctionReferenceNode* GetReferenceNode();
-	URigVMLibraryNode* GetReferencedFunction();
 
 private:
 
 	TSoftObjectPtr<URigVMFunctionReferenceNode> ReferenceNodePtr;
-	TSoftObjectPtr<URigVMLibraryNode> LibraryNodePtr;
 };
 
 /**
@@ -73,39 +72,40 @@ public:
 	URigVMBuildData();
 
 	// Returns the list of references for a given function definition
-	const FRigVMFunctionReferenceArray* FindFunctionReferences(const URigVMLibraryNode* InFunction) const;
+	const FRigVMFunctionReferenceArray* FindFunctionReferences(const FRigVMGraphFunctionIdentifier& InFunction) const;
 
 	/**
 	 * Iterator function to invoke a lambda / TFunction for each reference of a function
 	 * @param InFunction The function to iterate all references for
 	 * @param PerReferenceFunction The function to invoke for each reference
 	 */
-	void ForEachFunctionReference(const URigVMLibraryNode* InFunction, TFunction<void(URigVMFunctionReferenceNode*)> PerReferenceFunction) const;
+	void ForEachFunctionReference(const FRigVMGraphFunctionIdentifier& InFunction, TFunction<void(URigVMFunctionReferenceNode*)> PerReferenceFunction) const;
 
 	/**
 	* Iterator function to invoke a lambda / TFunction for each reference of a function
 	* @param InFunction The function to iterate all references for
 	* @param PerReferenceFunction The function to invoke for each reference
 	*/
-	void ForEachFunctionReferenceSoftPtr(const URigVMLibraryNode* InFunction, TFunction<void(TSoftObjectPtr<URigVMFunctionReferenceNode>)> PerReferenceFunction) const;
+	void ForEachFunctionReferenceSoftPtr(const FRigVMGraphFunctionIdentifier& InFunction, TFunction<void(TSoftObjectPtr<URigVMFunctionReferenceNode>)>
+	                                     PerReferenceFunction) const;
 
 	// Update the references list for a given reference node
 	void UpdateReferencesForFunctionReferenceNode(URigVMFunctionReferenceNode* InReferenceNode);
 
 	// registers a new reference node for a given function
-	void RegisterFunctionReference(URigVMLibraryNode* InFunction, URigVMFunctionReferenceNode* InReference);
+	void RegisterFunctionReference(const FRigVMGraphFunctionIdentifier& InFunction, URigVMFunctionReferenceNode* InReference);
 
 	// registers a new reference node for a given function
-	void RegisterFunctionReference(TSoftObjectPtr<URigVMLibraryNode> InFunction, TSoftObjectPtr<URigVMFunctionReferenceNode> InReference);
+	void RegisterFunctionReference(const FRigVMGraphFunctionIdentifier& InFunction, TSoftObjectPtr<URigVMFunctionReferenceNode> InReference);
 
 	// registers a new reference node for a given function
 	void RegisterFunctionReference(FRigVMReferenceNodeData InReferenceNodeData);
 
 	// unregisters a new reference node for a given function
-	void UnregisterFunctionReference(URigVMLibraryNode* InFunction, URigVMFunctionReferenceNode* InReference);
+	void UnregisterFunctionReference(const FRigVMGraphFunctionIdentifier& InFunction, URigVMFunctionReferenceNode* InReference);
 
 	// unregisters a new reference node for a given function
-	void UnregisterFunctionReference(TSoftObjectPtr<URigVMLibraryNode> InFunction, TSoftObjectPtr<URigVMFunctionReferenceNode> InReference);
+	void UnregisterFunctionReference(const FRigVMGraphFunctionIdentifier& InFunction, TSoftObjectPtr<URigVMFunctionReferenceNode> InReference);
 
 	// Clear references to temp assets
 	void ClearInvalidReferences();
@@ -115,8 +115,11 @@ public:
 
 private:
 
+	UPROPERTY(meta=(DeprecatedProperty))
+	TMap< TSoftObjectPtr<URigVMLibraryNode>, FRigVMFunctionReferenceArray > FunctionReferences_DEPRECATED;
+
 	UPROPERTY(VisibleAnywhere, Category = "BuildData")
-	TMap< TSoftObjectPtr<URigVMLibraryNode>, FRigVMFunctionReferenceArray > FunctionReferences;
+	TMap< FRigVMGraphFunctionIdentifier, FRigVMFunctionReferenceArray > GraphFunctionReferences;
 
 	bool bIsRunningUnitTest;
 
diff --git a/Engine/Source/Developer/RigVMDeveloper/Public/RigVMModel/RigVMClient.h b/Engine/Source/Developer/RigVMDeveloper/Public/RigVMModel/RigVMClient.h
--- a/Engine/Source/Developer/RigVMDeveloper/Public/RigVMModel/RigVMClient.h
+++ b/Engine/Source/Developer/RigVMDeveloper/Public/RigVMModel/RigVMClient.h
@@ -28,6 +28,12 @@ public:
 	// Returns the rigvm client for this host
 	virtual const FRigVMClient* GetRigVMClient() const = 0;
 
+	// Returns the rigvm function host
+	virtual IRigVMGraphFunctionHost* GetRigVMGraphFunctionHost() = 0;
+
+	// Returns the rigvm function host
+	virtual const IRigVMGraphFunctionHost* GetRigVMGraphFunctionHost() const = 0;
+
 	// Returns the editor object corresponding with the supplied editor object
 	virtual UObject* GetEditorObjectForRigVMGraph(URigVMGraph* InVMGraph) const = 0;
 
@@ -78,6 +84,7 @@ public:
 	FRigVMClient()
 		: FunctionLibrary(nullptr)
 		, bSuspendNotifications(false)
+		, bIgnoreModelNotifications(false)
 		, OuterClientHost(nullptr)
 		, OuterClientPropertyName(NAME_None)
 	{
@@ -130,6 +137,16 @@ public:
 	// backwards compatibility
 	void PatchModelsOnLoad();
 
+	void HandleGraphModifiedEvent(ERigVMGraphNotifType InNotifType, URigVMGraph* InGraph, UObject* InSubject);
+
+	FRigVMGraphFunctionStore* FindFunctionStore(const URigVMLibraryNode* InLibraryNode);
+	bool UpdateGraphFunctionData(const URigVMLibraryNode* InLibraryNode);
+	bool UpdateExternalVariablesForFunction(const URigVMLibraryNode* InLibraryNode);
+	bool UpdateDependenciesForFunction(const URigVMLibraryNode* InLibraryNode);
+	bool UpdateFunctionReferences(const FRigVMGraphFunctionHeader& InHeader, bool bUpdateDependencies, bool bUpdateExternalVariables);
+	bool DirtyGraphFunctionCompilationData(URigVMLibraryNode* InLibraryNode);
+	bool IsFunctionPublic(URigVMLibraryNode* InLibraryNode);
+
 private:
 
 	enum ERigVMClientAction
@@ -165,6 +182,7 @@ private:
 
 public:
 	bool bSuspendNotifications;
+	bool bIgnoreModelNotifications;
 private:
 	TWeakObjectPtr<UObject> OuterClientHost;
 	FName OuterClientPropertyName;
diff --git a/Engine/Source/Developer/RigVMDeveloper/Public/RigVMModel/RigVMController.h b/Engine/Source/Developer/RigVMDeveloper/Public/RigVMModel/RigVMController.h
--- a/Engine/Source/Developer/RigVMDeveloper/Public/RigVMModel/RigVMController.h
+++ b/Engine/Source/Developer/RigVMDeveloper/Public/RigVMModel/RigVMController.h
@@ -28,6 +28,8 @@
 #endif
 
 class URigVMActionStack;
+struct FRigVMGraphFunctionArgument;
+struct FRigVMGraphFunctionHeader;
 
 UENUM()
 enum class ERigVMControllerBulkEditType : uint8
@@ -82,10 +84,9 @@ private:
 DECLARE_DELEGATE_RetVal_OneParam(bool, FRigVMController_ShouldStructUnfoldDelegate, const UStruct*)
 DECLARE_DELEGATE_RetVal_OneParam(TArray<FRigVMExternalVariable>, FRigVMController_GetExternalVariablesDelegate, URigVMGraph*)
 DECLARE_DELEGATE_RetVal(const FRigVMByteCode*, FRigVMController_GetByteCodeDelegate)
-DECLARE_DELEGATE_RetVal_OneParam(bool, FRigVMController_IsFunctionAvailableDelegate, URigVMLibraryNode*)
-DECLARE_DELEGATE_RetVal_OneParam(bool, FRigVMController_RequestLocalizeFunctionDelegate, URigVMLibraryNode*)
+DECLARE_DELEGATE_RetVal_OneParam(bool, FRigVMController_RequestLocalizeFunctionDelegate, TSoftObjectPtr<URigVMLibraryNode>)
 DECLARE_DELEGATE_RetVal_ThreeParams(FName, FRigVMController_RequestNewExternalVariableDelegate, FRigVMGraphVariableDescription, bool, bool);
-DECLARE_DELEGATE_RetVal_TwoParams(bool, FRigVMController_IsDependencyCyclicDelegate, UObject*, UObject*)
+DECLARE_DELEGATE_RetVal_TwoParams(bool, FRigVMController_IsDependencyCyclicDelegate, const FRigVMGraphFunctionHeader& Dependent, const FRigVMGraphFunctionHeader& Dependency)
 DECLARE_DELEGATE_RetVal_TwoParams(FRigVMController_BulkEditResult, FRigVMController_RequestBulkEditDialogDelegate, URigVMLibraryNode*, ERigVMControllerBulkEditType)
 DECLARE_DELEGATE_RetVal_OneParam(bool, FRigVMController_RequestBreakLinksDialogDelegate, TArray<URigVMLink*>)
 DECLARE_DELEGATE_FiveParams(FRigVMController_OnBulkEditProgressDelegate, TSoftObjectPtr<URigVMFunctionReferenceNode>, ERigVMControllerBulkEditType, ERigVMControllerBulkEditProgress, int32, int32)
@@ -765,6 +766,10 @@ public:
 	bool SetExposedPinIndex(const FName& InPinName, int32 InNewIndex, bool bSetupUndoRedo = true, bool bPrintPythonCommand = false);
 
 	// Adds a function reference / invocation to the graph
+	UFUNCTION(BlueprintCallable, Category = RigVMController)
+	URigVMFunctionReferenceNode* AddFunctionReferenceNodeFromDescription(const FRigVMGraphFunctionHeader& InFunctionDefinition, const FVector2D& InNodePosition = FVector2D::ZeroVector, const
+	                                                                     FString& InNodeName = TEXT(""), bool bSetupUndoRedo = true, bool bPrintPythonCommand = false, bool bAllowPrivateFunctions = false);
+	
 	UFUNCTION(BlueprintCallable, Category = RigVMController)
 	URigVMFunctionReferenceNode* AddFunctionReferenceNode(URigVMLibraryNode* InFunctionDefinition, const FVector2D& InNodePosition = FVector2D::ZeroVector, const FString& InNodeName = TEXT(""), bool bSetupUndoRedo = true, bool bPrintPythonCommand = false);
 
@@ -832,17 +837,11 @@ public:
 	// A delegate to retrieve the current bytecode of the graph
 	FRigVMController_GetByteCodeDelegate GetCurrentByteCodeDelegate;
 
-	// A delegate to determine if a function is public
-	FRigVMController_IsFunctionAvailableDelegate IsFunctionAvailableDelegate;
-
 	// A delegate to localize a function on demand
 	FRigVMController_RequestLocalizeFunctionDelegate RequestLocalizeFunctionDelegate;
 
 	// A delegate to create a new blueprint member variable
 	FRigVMController_RequestNewExternalVariableDelegate RequestNewExternalVariableDelegate;
-	
-	// A delegate to validate if we are allowed to introduce a dependency between two objects
-	FRigVMController_IsDependencyCyclicDelegate IsDependencyCyclicDelegate;
 
 	// A delegate to ask the host / client for a dialog to confirm a bulk edit
 	FRigVMController_RequestBulkEditDialogDelegate RequestBulkEditDialogDelegate;
@@ -988,12 +987,13 @@ private:
 	bool CanAddNode(URigVMNode* InNode, bool bReportErrors, bool bIgnoreFunctionEntryReturnNodes = false);
 	TObjectPtr<URigVMNode> FindEventNode(const UScriptStruct* InScriptStruct) const;
 	bool CanAddEventNode(UScriptStruct* InScriptStruct, const bool bReportErrors) const;
-	bool CanAddFunctionRefForDefinition(URigVMLibraryNode* InFunctionDefinition, bool bReportErrors);
+	bool CanAddFunctionRefForDefinition(const FRigVMGraphFunctionHeader& InFunctionDefinition, bool bReportErrors, bool bAllowPrivateFunctions=false);
 	void AddPinsForStruct(UStruct* InStruct, URigVMNode* InNode, URigVMPin* InParentPin, ERigVMPinDirection InPinDirection, const FString& InDefaultValue, bool bAutoExpandArrays, bool bNotify = false);
 	void AddPinsForArray(FArrayProperty* InArrayProperty, URigVMNode* InNode, URigVMPin* InParentPin, ERigVMPinDirection InPinDirection, const TArray<FString>& InDefaultValues, bool bAutoExpandArrays);
 	void AddPinsForTemplate(const FRigVMTemplate* InTemplate, const FRigVMTemplateTypeMap& InPinTypeMap, URigVMNode* InNode);
 	void ConfigurePinFromProperty(FProperty* InProperty, URigVMPin* InOutPin, ERigVMPinDirection InPinDirection = ERigVMPinDirection::Invalid);
 	void ConfigurePinFromPin(URigVMPin* InOutPin, URigVMPin* InPin, bool bCopyDisplayName = false);
+	void ConfigurePinFromArgument(URigVMPin* InOutPin, const FRigVMGraphFunctionArgument& InArgument, bool bCopyDisplayName = false);
 	virtual bool ShouldStructBeUnfolded(const UStruct* InStruct);
 	virtual bool ShouldPinBeUnfolded(URigVMPin* InPin);
 	bool SetPinDefaultValue(URigVMPin* InPin, const FString& InDefaultValue, bool bResizeArrays, bool bSetupUndoRedo, bool bMergeUndoAction, bool bNotify = true);
diff --git a/Engine/Source/Developer/RigVMDeveloper/Public/RigVMModel/RigVMFunctionLibrary.h b/Engine/Source/Developer/RigVMDeveloper/Public/RigVMModel/RigVMFunctionLibrary.h
--- a/Engine/Source/Developer/RigVMDeveloper/Public/RigVMModel/RigVMFunctionLibrary.h
+++ b/Engine/Source/Developer/RigVMDeveloper/Public/RigVMModel/RigVMFunctionLibrary.h
@@ -6,8 +6,11 @@
 #include "RigVMBuildData.h"
 #include "RigVMModel/Nodes/RigVMLibraryNode.h"
 #include "RigVMModel/Nodes/RigVMFunctionReferenceNode.h"
+#include "RigVMModel/Nodes/RigVMLibraryNode.h"
 #include "RigVMFunctionLibrary.generated.h"
 
+DECLARE_DELEGATE_RetVal(const FSoftObjectPath, URigVMFunctionLibrary_GetFunctionHostObjectPath);
+
 /**
  * The Function Library is a graph used only to store
  * the sub graphs used for functions.
@@ -69,7 +72,10 @@ public:
 	// We maintain meta data on what functions have been created locally based on which other ones,
 	// and use this method to avoid redundant localizations.
 	URigVMLibraryNode* FindPreviouslyLocalizedFunction(URigVMLibraryNode* InFunctionToLocalize);
-	
+
+	const FSoftObjectPath GetFunctionHostObjectPath() const;
+	URigVMFunctionLibrary_GetFunctionHostObjectPath GetFunctionHostObjectPathDelegate;
+
 private:
 
 #if WITH_EDITORONLY_DATA
@@ -81,7 +87,7 @@ private:
 	// The source pathname is the full path of the source function that was localized
 	// to the local copy stored in the value of the pair.
 	UPROPERTY()
-	TMap< FString, TObjectPtr<URigVMLibraryNode> > LocalizedFunctions;
+	TMap< FSoftObjectPath, TObjectPtr<URigVMLibraryNode> > LocalizedFunctions;
 
 	friend class URigVMController;
 	friend class URigVMCompiler;
diff --git a/Engine/Source/Developer/RigVMDeveloper/Public/RigVMModel/RigVMNode.h b/Engine/Source/Developer/RigVMDeveloper/Public/RigVMModel/RigVMNode.h
--- a/Engine/Source/Developer/RigVMDeveloper/Public/RigVMModel/RigVMNode.h
+++ b/Engine/Source/Developer/RigVMDeveloper/Public/RigVMModel/RigVMNode.h
@@ -254,6 +254,9 @@ public:
 	UFUNCTION(BlueprintPure, Category = RigVMNode)
 	virtual FName GetNextAggregateName(const FName& InLastAggregatePinName) const { return NAME_None; }
 
+	UFUNCTION(BlueprintCallable, Category = RigVMNode)
+	virtual URigVMLibraryNode* FindFunctionForNode();
+
 	virtual FRigVMStructUpgradeInfo GetUpgradeInfo() const { return FRigVMStructUpgradeInfo(); }
 
 private:
diff --git a/Engine/Source/Runtime/RigVM/Private/RigVMCore/RigVMGraphFunctionDefinition.cpp b/Engine/Source/Runtime/RigVM/Private/RigVMCore/RigVMGraphFunctionDefinition.cpp
new file mode 100644
--- /dev/null
+++ b/Engine/Source/Runtime/RigVM/Private/RigVMCore/RigVMGraphFunctionDefinition.cpp
@@ -0,0 +1,67 @@
+// Copyright Epic Games, Inc. All Rights Reserved.
+
+#include "RigVMCore/RigVMGraphFunctionDefinition.h"
+#include "RigVMCore/RigVMGraphFunctionHost.h"
+
+#include UE_INLINE_GENERATED_CPP_BY_NAME(RigVMGraphFunctionDefinition)
+
+FRigVMExternalVariable FRigVMGraphFunctionArgument::GetExternalVariable() const
+{
+	FRigVMExternalVariable Variable;
+	Variable.Name = Name;
+	Variable.bIsArray = bIsArray;
+	Variable.TypeName = CPPType;
+	
+	if(CPPTypeObject.IsValid())
+	{
+		Variable.TypeObject = CPPTypeObject.Get();
+	}
+
+	return Variable;
+}
+
+bool FRigVMGraphFunctionHeader::IsMutable() const
+{
+	for(const FRigVMGraphFunctionArgument& Arg : Arguments)
+	{
+		if(Arg.CPPTypeObject.IsValid())
+		{
+			if(UScriptStruct* Struct = Cast<UScriptStruct>(Arg.CPPTypeObject.Get()))
+			{
+				if(Struct->IsChildOf(FRigVMExecuteContext::StaticStruct()))
+				{
+					return true;
+				}
+			}
+		}
+	}
+	return false;
+}
+
+IRigVMGraphFunctionHost* FRigVMGraphFunctionHeader::GetFunctionHost() const
+{
+	UObject* HostObj = LibraryPointer.HostObject.ResolveObject();
+	if (!HostObj)
+	{
+		HostObj = LibraryPointer.HostObject.TryLoad();
+	}
+	if (HostObj)
+	{
+		return Cast<IRigVMGraphFunctionHost>(HostObj);
+	}
+	return nullptr;
+}
+
+FRigVMGraphFunctionData* FRigVMGraphFunctionHeader::GetFunctionData() const
+{
+	if (IRigVMGraphFunctionHost* Host = GetFunctionHost())
+	{
+		return Host->GetRigVMGraphFunctionStore()->FindFunction(LibraryPointer);
+	}
+	return nullptr;
+}
+
+bool FRigVMGraphFunctionData::IsMutable() const
+{
+	return Header.IsMutable();
+}
diff --git a/Engine/Source/Runtime/RigVM/Private/RigVMCore/RigVMGraphFunctionHost.cpp b/Engine/Source/Runtime/RigVM/Private/RigVMCore/RigVMGraphFunctionHost.cpp
new file mode 100644
--- /dev/null
+++ b/Engine/Source/Runtime/RigVM/Private/RigVMCore/RigVMGraphFunctionHost.cpp
@@ -0,0 +1,252 @@
+// Copyright Epic Games, Inc. All Rights Reserved.
+
+#include "RigVMCore/RigVMGraphFunctionHost.h"
+
+const FRigVMGraphFunctionData* FRigVMGraphFunctionStore::FindFunctionImpl(const FRigVMGraphFunctionIdentifier& InLibraryPointer, bool* bOutIsPublic) const
+{
+	const FRigVMGraphFunctionData* Info = PublicFunctions.FindByPredicate([InLibraryPointer](const FRigVMGraphFunctionData& Info)
+	{
+		return Info.Header.LibraryPointer == InLibraryPointer;
+	});
+	if (Info)
+	{
+		if (bOutIsPublic)
+		{
+			(*bOutIsPublic) = true; 
+		}
+		return Info;
+	}
+
+	Info = PrivateFunctions.FindByPredicate([InLibraryPointer](const FRigVMGraphFunctionData& Info)
+	{
+		return Info.Header.LibraryPointer == InLibraryPointer;
+	});
+	if (Info)
+	{
+		if (bOutIsPublic)
+		{
+			(*bOutIsPublic) = false; 
+		}
+		return Info;
+	}
+	return nullptr;
+}
+
+const FRigVMGraphFunctionData* FRigVMGraphFunctionStore::FindFunction(const FRigVMGraphFunctionIdentifier& InLibraryPointer, bool* bOutIsPublic) const
+{
+	return FindFunctionImpl(InLibraryPointer, bOutIsPublic);
+}
+
+FRigVMGraphFunctionData* FRigVMGraphFunctionStore::FindFunction(const FRigVMGraphFunctionIdentifier& InLibraryPointer, bool* bOutIsPublic)
+{
+	return const_cast<FRigVMGraphFunctionData*>(FindFunctionImpl(InLibraryPointer, bOutIsPublic));
+}
+
+FRigVMGraphFunctionData* FRigVMGraphFunctionStore::FindFunctionByName(const FName& Name, bool* bOutIsPublic)
+{
+	FRigVMGraphFunctionData* Info = PublicFunctions.FindByPredicate([Name](const FRigVMGraphFunctionData& Info)
+	{
+		return Info.Header.Name == Name;
+	});
+	if (Info)
+	{
+		if (bOutIsPublic)
+		{
+			(*bOutIsPublic) = true; 
+		}
+		return Info;
+	}
+
+	Info = PrivateFunctions.FindByPredicate([Name](const FRigVMGraphFunctionData& Info)
+	{
+		return Info.Header.Name == Name;
+	});
+	if (Info)
+	{
+		if (bOutIsPublic)
+		{
+			(*bOutIsPublic) = false; 
+		}
+		return Info;
+	}
+	return nullptr;
+}
+
+bool FRigVMGraphFunctionStore::ContainsFunction(const FRigVMGraphFunctionIdentifier& InLibraryPointer) const
+{
+	if(FindFunction(InLibraryPointer))
+	{
+		return true;
+	}
+	return false;
+}
+
+bool FRigVMGraphFunctionStore::IsFunctionPublic(const FRigVMGraphFunctionIdentifier& InLibraryPointer) const
+{
+	bool bIsPublic;
+	if (FindFunction(InLibraryPointer, &bIsPublic))
+	{
+		return bIsPublic;
+	}
+	return false;
+}
+
+FRigVMGraphFunctionData* FRigVMGraphFunctionStore::AddFunction(const FRigVMGraphFunctionHeader& FunctionHeader, bool bIsPublic)
+{
+	// Fail if the function already exists
+	if (ContainsFunction(FunctionHeader.LibraryPointer))
+	{
+		return nullptr;
+	}
+
+	FRigVMGraphFunctionData NewInfo(FunctionHeader);
+	if(bIsPublic)
+	{
+		int32 Index = PublicFunctions.Add(NewInfo);
+		return &PublicFunctions[Index];
+	}
+	else
+	{
+		int32 Index =  PrivateFunctions.Add(NewInfo);
+		return &PrivateFunctions[Index];
+	}
+
+	return nullptr;
+}
+
+bool FRigVMGraphFunctionStore::RemoveFunction(const FRigVMGraphFunctionIdentifier& InLibraryPointer, bool* bIsPublic)
+{
+	int32 NumRemoved = PublicFunctions.RemoveAll([InLibraryPointer](const FRigVMGraphFunctionData& Info)
+	{
+		return Info.Header.LibraryPointer == InLibraryPointer;
+	});
+	if (NumRemoved > 0 && bIsPublic)
+	{
+		*bIsPublic = true;
+	}
+		
+	if (NumRemoved == 0)
+	{
+		NumRemoved = PrivateFunctions.RemoveAll([InLibraryPointer](const FRigVMGraphFunctionData& Info)
+		{
+		   return Info.Header.LibraryPointer == InLibraryPointer;
+		});
+		if (NumRemoved > 0 && bIsPublic)
+		{
+			*bIsPublic = false;
+		}
+	}
+
+	return NumRemoved > 0;
+}
+
+bool FRigVMGraphFunctionStore::MarkFunctionAsPublic(const FRigVMGraphFunctionIdentifier& InLibraryPointer, bool bIsPublic)
+{
+	TArray<FRigVMGraphFunctionData>* OldContainer = nullptr;
+	TArray<FRigVMGraphFunctionData>* NewContainer = nullptr;
+	if(bIsPublic)
+	{
+		OldContainer = &PrivateFunctions;
+		NewContainer = &PublicFunctions;
+	}
+	else
+	{
+		OldContainer = &PublicFunctions;
+		NewContainer = &PrivateFunctions;
+	}
+	
+	int32 Index = OldContainer->IndexOfByPredicate([InLibraryPointer](const FRigVMGraphFunctionData& FunctionData)
+		{
+			return FunctionData.Header.LibraryPointer == InLibraryPointer;
+		});
+	if(Index == INDEX_NONE)
+	{
+		return false;
+	}
+	NewContainer->Add(OldContainer->operator[](Index));
+	OldContainer->RemoveAt(Index);
+	return true;
+}
+
+FRigVMGraphFunctionData* FRigVMGraphFunctionStore::UpdateFunctionInterface(const FRigVMGraphFunctionHeader& Header)
+{
+	bool bIsPublic;
+
+	if (const FRigVMGraphFunctionData* OldData = FindFunction(Header.LibraryPointer))
+	{
+		TMap<FRigVMGraphFunctionIdentifier, uint32> Dependencies = OldData->Header.Dependencies;
+		TArray<FRigVMExternalVariable> ExternalVariables = OldData->Header.ExternalVariables;
+		if (RemoveFunction(Header.LibraryPointer, &bIsPublic))
+		{
+			if (FRigVMGraphFunctionData* NewData = AddFunction(Header, bIsPublic))
+			{
+				NewData->Header.Dependencies = Dependencies;
+				NewData->Header.ExternalVariables = ExternalVariables;
+				return NewData;
+			}
+		}
+	}
+	return nullptr;
+}
+
+bool FRigVMGraphFunctionStore::UpdateDependencies(const FRigVMGraphFunctionIdentifier& InLibraryPointer, TMap<FRigVMGraphFunctionIdentifier, uint32>& Dependencies)
+{
+	if (FRigVMGraphFunctionData* Data = FindFunction(InLibraryPointer))
+	{
+		Data->Header.Dependencies = Dependencies;		
+		return true;
+	}
+	return false;
+}
+
+bool FRigVMGraphFunctionStore::UpdateExternalVariables(const FRigVMGraphFunctionIdentifier& InLibraryPointer, TArray<FRigVMExternalVariable> ExternalVariables)
+{
+	if (FRigVMGraphFunctionData* Data = FindFunction(InLibraryPointer))
+	{
+		FRigVMGraphFunctionHeader& Header = Data->Header;
+		
+		// Find Dependencies
+		Header.ExternalVariables = ExternalVariables;
+
+		return true;
+	}
+	return false;
+}
+
+bool FRigVMGraphFunctionStore::UpdateFunctionCompilationData(const FRigVMGraphFunctionIdentifier& InLibraryPointer, const FRigVMFunctionCompilationData& CompilationData)
+{
+	FRigVMGraphFunctionData* Info = FindFunction(InLibraryPointer);
+	if (Info)
+	{
+		Info->CompilationData = CompilationData;
+		return true;
+	}
+
+	return false;
+}
+
+bool FRigVMGraphFunctionStore::RemoveFunctionCompilationData(const FRigVMGraphFunctionIdentifier& InLibraryPointer)
+{
+	FRigVMGraphFunctionData* FunctionData = FindFunction(InLibraryPointer);
+	if (FunctionData)
+	{
+		FunctionData->ClearCompilationData();
+		return true;
+	}
+	return false;
+}
+
+bool FRigVMGraphFunctionStore::RemoveAllCompilationData()
+{
+	for (FRigVMGraphFunctionData& Data : PublicFunctions)
+	{
+		Data.ClearCompilationData();
+	}
+
+	for (FRigVMGraphFunctionData& Data : PrivateFunctions)
+	{
+		Data.ClearCompilationData();
+	}
+
+	return true;
+}
diff --git a/Engine/Source/Runtime/RigVM/Public/RigVMCore/RigVMExternalVariable.h b/Engine/Source/Runtime/RigVM/Public/RigVMCore/RigVMExternalVariable.h
--- a/Engine/Source/Runtime/RigVM/Public/RigVMCore/RigVMExternalVariable.h
+++ b/Engine/Source/Runtime/RigVM/Public/RigVMCore/RigVMExternalVariable.h
@@ -9,13 +9,17 @@
 #include "RigVMModule.h"
 #include "RigVMTypeUtils.h"
 #include "UObject/UnrealType.h"
+#include "RigVMExternalVariable.generated.h"
 
 /**
  * The external variable can be used to map external / unowned
  * memory into the VM and back out
  */
+USTRUCT(BlueprintType)
 struct RIGVM_API FRigVMExternalVariable
 {
+	GENERATED_BODY()
+	
 	FORCEINLINE FRigVMExternalVariable()
 		: Name(NAME_None)
 		, Property(nullptr)
@@ -587,7 +591,7 @@ struct RIGVM_API FRigVMExternalVariable
 		return TypeName;
 	}
 
-	FORCEINLINE static void MergeExternalVariable(TArray<FRigVMExternalVariable>& OutVariables, const FRigVMExternalVariable& InVariable)
+	FORCEINLINE_DEBUGGABLE static void MergeExternalVariable(TArray<FRigVMExternalVariable>& OutVariables, const FRigVMExternalVariable& InVariable)
 	{
 		if(!InVariable.IsValid(true))
 		{
@@ -617,3 +621,28 @@ struct RIGVM_API FRigVMExternalVariable
 	int32 Size;
 	uint8* Memory;
 };
+
+FORCEINLINE FArchive& operator<<(FArchive& Ar, FRigVMExternalVariable& Variable)
+{
+	Ar << Variable.Name;
+	Ar << Variable.TypeName;
+
+	if (Ar.IsSaving())
+	{
+		FSoftObjectPath TypeObjectPath(Variable.TypeObject);
+		Ar << TypeObjectPath;
+	}
+	else if (Ar.IsLoading())
+	{
+		FSoftObjectPath TypeObjectPath;
+		Ar << TypeObjectPath;
+		Variable.TypeObject = TypeObjectPath.ResolveObject();
+	}
+	
+	Ar << Variable.bIsArray;
+	Ar << Variable.bIsPublic;
+	Ar << Variable.bIsReadOnly;
+	Ar << Variable.Size;
+	return Ar;
+}
+
diff --git a/Engine/Source/Runtime/RigVM/Public/RigVMCore/RigVMGraphFunctionDefinition.h b/Engine/Source/Runtime/RigVM/Public/RigVMCore/RigVMGraphFunctionDefinition.h
new file mode 100644
--- /dev/null
+++ b/Engine/Source/Runtime/RigVM/Public/RigVMCore/RigVMGraphFunctionDefinition.h
@@ -0,0 +1,446 @@
+// Copyright Epic Games, Inc. All Rights Reserved.
+
+#pragma once
+
+#include "RigVMCore/RigVMExternalVariable.h"
+#include "RigVMCore/RigVMByteCode.h"
+#include "RigVMGraphFunctionDefinition.generated.h"
+
+class IRigVMGraphFunctionHost;
+struct FRigVMGraphFunctionData;
+
+USTRUCT()
+struct FRigVMFunctionCompilationPropertyDescription
+{
+	GENERATED_BODY()
+
+	// The name of the property to create
+	UPROPERTY()
+	FName Name;
+
+	// The complete CPP type to base a new property off of (for ex: 'TArray<TArray<FVector>>')
+	UPROPERTY()
+	FString CPPType;
+
+	// The tail CPP Type object, for example the UScriptStruct for a struct
+	UPROPERTY()
+	TSoftObjectPtr<UObject> CPPTypeObject;
+
+	// The default value to use for this property (for example: '(((X=1.000000, Y=2.000000, Z=3.000000)))')
+	UPROPERTY()
+	FString DefaultValue;
+
+	friend FORCEINLINE uint32 GetTypeHash(const FRigVMFunctionCompilationPropertyDescription& Description) 
+	{
+		uint32 Hash = GetTypeHash(Description.Name);
+		Hash = HashCombine(Hash, GetTypeHash(Description.CPPType));
+		Hash = HashCombine(Hash, GetTypeHash(Description.CPPTypeObject));
+		Hash = HashCombine(Hash, GetTypeHash(Description.DefaultValue));
+		return Hash;
+	}
+};
+
+FORCEINLINE FArchive& operator<<(FArchive& Ar, FRigVMFunctionCompilationPropertyDescription& Data)
+{
+	Ar << Data.Name;
+	Ar << Data.CPPType;
+	Ar << Data.CPPTypeObject;
+	Ar << Data.DefaultValue;
+	return Ar;
+}
+
+USTRUCT()
+struct FRigVMFunctionCompilationPropertyPath
+{
+	GENERATED_BODY()
+
+	UPROPERTY()
+	int32 PropertyIndex;
+
+	UPROPERTY()
+	FString HeadCPPType;
+
+	UPROPERTY()
+	FString SegmentPath;
+
+	friend FORCEINLINE uint32 GetTypeHash(const FRigVMFunctionCompilationPropertyPath& Path)
+	{
+		uint32 Hash = GetTypeHash(Path.PropertyIndex);
+		Hash = HashCombine(Hash, GetTypeHash(Path.HeadCPPType));
+		Hash = HashCombine(Hash, GetTypeHash(Path.SegmentPath));
+		return Hash;
+	}
+};
+
+FORCEINLINE FArchive& operator<<(FArchive& Ar, FRigVMFunctionCompilationPropertyPath& Data)
+{
+	Ar << Data.PropertyIndex;
+	Ar << Data.HeadCPPType;
+	Ar << Data.SegmentPath;
+	return Ar;
+}
+
+USTRUCT(BlueprintType)
+struct RIGVM_API FRigVMFunctionCompilationData
+{
+	GENERATED_BODY()
+
+	FRigVMFunctionCompilationData() : Hash(0) {}
+
+	UPROPERTY()
+	FRigVMByteCode ByteCode;
+
+	UPROPERTY()
+	TArray<FName> FunctionNames;
+
+	UPROPERTY()
+	TArray<FRigVMFunctionCompilationPropertyDescription> WorkPropertyDescriptions;
+
+	UPROPERTY()
+	TArray<FRigVMFunctionCompilationPropertyPath> WorkPropertyPathDescriptions;
+
+	UPROPERTY()
+	TArray<FRigVMFunctionCompilationPropertyDescription> LiteralPropertyDescriptions;
+
+	UPROPERTY()
+	TArray<FRigVMFunctionCompilationPropertyPath> LiteralPropertyPathDescriptions;
+
+	UPROPERTY()
+	TArray<FRigVMFunctionCompilationPropertyDescription> DebugPropertyDescriptions;
+
+	UPROPERTY()
+	TArray<FRigVMFunctionCompilationPropertyPath> DebugPropertyPathDescriptions;
+
+	UPROPERTY()
+	TArray<FRigVMFunctionCompilationPropertyDescription> ExternalPropertyDescriptions;
+
+	UPROPERTY()
+	TArray<FRigVMFunctionCompilationPropertyPath> ExternalPropertyPathDescriptions;
+
+	UPROPERTY()
+	TMap<int32, FName> ExternalRegisterIndexToVariable;
+
+	UPROPERTY()
+	TMap<FString, FRigVMOperand> Operands;
+
+	UPROPERTY()
+	uint32 Hash;
+
+	bool IsValid() const
+	{
+		return ByteCode.Num() > 0 || WorkPropertyDescriptions.Num() > 0;
+	}
+
+	friend FORCEINLINE uint32 GetTypeHash(const FRigVMFunctionCompilationData& Data) 
+	{
+		uint32 DataHash = Data.ByteCode.GetByteCodeHash();
+		for (const FName& Name : Data.FunctionNames)
+		{
+			DataHash = HashCombine(DataHash, GetTypeHash(Name));
+		}
+
+		for (const FRigVMFunctionCompilationPropertyDescription& Description : Data.WorkPropertyDescriptions)
+		{
+			DataHash = HashCombine(DataHash, GetTypeHash(Description));
+		}
+		for (const FRigVMFunctionCompilationPropertyPath& Path : Data.WorkPropertyPathDescriptions)
+		{
+			DataHash = HashCombine(DataHash, GetTypeHash(Path));
+		}
+
+		for (const FRigVMFunctionCompilationPropertyDescription& Description : Data.LiteralPropertyDescriptions)
+		{
+			DataHash = HashCombine(DataHash, GetTypeHash(Description));
+		}
+		for (const FRigVMFunctionCompilationPropertyPath& Path : Data.LiteralPropertyPathDescriptions)
+		{
+			DataHash = HashCombine(DataHash, GetTypeHash(Path));
+		}
+
+		for (const FRigVMFunctionCompilationPropertyDescription& Description : Data.DebugPropertyDescriptions)
+		{
+			DataHash = HashCombine(DataHash, GetTypeHash(Description));
+		}
+		for (const FRigVMFunctionCompilationPropertyPath& Path : Data.DebugPropertyPathDescriptions)
+		{
+			DataHash = HashCombine(DataHash, GetTypeHash(Path));
+		}
+
+		for (const FRigVMFunctionCompilationPropertyDescription& Description : Data.ExternalPropertyDescriptions)
+		{
+			DataHash = HashCombine(DataHash, GetTypeHash(Description));
+		}
+		for (const FRigVMFunctionCompilationPropertyPath& Path : Data.ExternalPropertyPathDescriptions)
+		{
+			DataHash = HashCombine(DataHash, GetTypeHash(Path));
+		}
+		
+		for (const TPair<int32,FName>& Pair : Data.ExternalRegisterIndexToVariable)
+		{
+			DataHash = HashCombine(DataHash, GetTypeHash(Pair.Key));
+			DataHash = HashCombine(DataHash, GetTypeHash(Pair.Value));
+		}
+
+		for (const TPair<FString, FRigVMOperand>& Pair : Data.Operands)
+		{
+			DataHash = HashCombine(DataHash, GetTypeHash(Pair.Key));
+			DataHash = HashCombine(DataHash, GetTypeHash(Pair.Value));
+		}
+
+		return DataHash;
+	}
+
+};
+
+FORCEINLINE FArchive& operator<<(FArchive& Ar, FRigVMFunctionCompilationData& Data)
+{
+	Ar << Data.ByteCode;
+	Ar << Data.FunctionNames;
+	Ar << Data.WorkPropertyDescriptions;
+	Ar << Data.WorkPropertyPathDescriptions;
+	Ar << Data.LiteralPropertyDescriptions;
+	Ar << Data.LiteralPropertyPathDescriptions;
+	Ar << Data.DebugPropertyDescriptions;
+	Ar << Data.DebugPropertyPathDescriptions;
+	Ar << Data.ExternalPropertyDescriptions;
+	Ar << Data.ExternalPropertyPathDescriptions;
+	Ar << Data.ExternalRegisterIndexToVariable;
+	Ar << Data.Operands;
+	Ar << Data.Hash;
+	return Ar;
+}
+
+
+USTRUCT(BlueprintType)
+struct RIGVM_API FRigVMGraphFunctionArgument
+{
+	GENERATED_BODY()
+	
+	FRigVMGraphFunctionArgument()
+	: Name(NAME_None)
+	, DisplayName(NAME_None)
+	, CPPType(NAME_None)
+	, CPPTypeObject(nullptr)
+	, bIsArray(false)
+	, Direction(ERigVMPinDirection::Input)
+	, bIsConst(false)
+	{}
+
+	UPROPERTY()
+	FName Name;
+
+	UPROPERTY()
+	FName DisplayName;
+
+	UPROPERTY()
+	FName CPPType;
+
+	UPROPERTY()
+	TSoftObjectPtr<UObject> CPPTypeObject;
+
+	UPROPERTY()
+	bool bIsArray;
+
+	UPROPERTY()
+	ERigVMPinDirection Direction;
+
+	UPROPERTY()
+	FString DefaultValue;
+
+	UPROPERTY()
+	bool bIsConst;
+	
+	UPROPERTY()
+	TMap<FString, FText> PathToTooltip;
+
+	FRigVMExternalVariable GetExternalVariable() const;
+
+	friend FORCEINLINE uint32 GetTypeHash(const FRigVMGraphFunctionArgument& Argument)
+	{
+		uint32 Hash = HashCombine(GetTypeHash(Argument.Name), GetTypeHash(Argument.DisplayName));
+		Hash = HashCombine(Hash, GetTypeHash(Argument.CPPType));
+		Hash = HashCombine(Hash, GetTypeHash(Argument.CPPTypeObject));
+		Hash = HashCombine(Hash, GetTypeHash(Argument.bIsArray));
+		Hash = HashCombine(Hash, GetTypeHash(Argument.Direction));
+		Hash = HashCombine(Hash, GetTypeHash(Argument.DefaultValue));
+		Hash = HashCombine(Hash, GetTypeHash(Argument.bIsConst));
+		for (const TPair<FString, FText>& Pair : Argument.PathToTooltip)
+		{
+			Hash = HashCombine(Hash, GetTypeHash(Pair.Key));
+			Hash = HashCombine(Hash, GetTypeHash(Pair.Value.ToString()));
+		}
+		return Hash;
+	}
+
+	bool operator==(const FRigVMGraphFunctionArgument& Other) const
+	{
+		return true;
+	}
+};
+
+FORCEINLINE FArchive& operator<<(FArchive& Ar, FRigVMGraphFunctionArgument& Data)
+{
+	Ar << Data.Name;
+	Ar << Data.DisplayName;
+	Ar << Data.CPPType;
+	Ar << Data.CPPTypeObject;
+	Ar << Data.bIsArray;
+	Ar << Data.Direction;
+	Ar << Data.DefaultValue;
+	Ar << Data.bIsConst;
+	Ar << Data.PathToTooltip;
+	return Ar;
+}
+
+USTRUCT(BlueprintType)
+struct RIGVM_API FRigVMGraphFunctionIdentifier
+{
+	GENERATED_BODY()
+
+	UPROPERTY()
+	FSoftObjectPath LibraryNode;
+
+	// A path to the IRigVMGraphFunctionHost that stores the function information, and compilation data (e.g. ControlRigBlueprintGeneratedClass)
+	UPROPERTY()
+	FSoftObjectPath HostObject;
+
+	FRigVMGraphFunctionIdentifier()
+		: LibraryNode(nullptr), HostObject(nullptr) {}
+	
+	FRigVMGraphFunctionIdentifier(FSoftObjectPath InHostObject, FSoftObjectPath InLibraryNode)
+		: LibraryNode(InLibraryNode), HostObject(InHostObject) {}
+
+	FORCEINLINE FString GetHash() const
+	{
+		return LibraryNode.ToString();
+	}
+
+	friend FORCEINLINE uint32 GetTypeHash(const FRigVMGraphFunctionIdentifier& Pointer)
+	{
+		return HashCombine(GetTypeHash(Pointer.LibraryNode), GetTypeHash(Pointer.HostObject));
+	}
+
+	bool operator==(const FRigVMGraphFunctionIdentifier& Other) const
+	{
+		return HostObject == Other.HostObject && LibraryNode == Other.LibraryNode;
+	}
+};
+
+FORCEINLINE FArchive& operator<<(FArchive& Ar, FRigVMGraphFunctionIdentifier& Data)
+{
+	Ar << Data.LibraryNode;
+	Ar << Data.HostObject;
+	return Ar;
+}
+
+USTRUCT(BlueprintType)
+struct RIGVM_API FRigVMGraphFunctionHeader
+{
+	GENERATED_BODY()
+
+	FRigVMGraphFunctionHeader()
+		: LibraryPointer(nullptr, nullptr)
+		, Name(NAME_None)
+	{}
+
+	UPROPERTY()
+	FRigVMGraphFunctionIdentifier LibraryPointer;
+
+	UPROPERTY()
+	FName Name;
+
+	UPROPERTY()
+	FString NodeTitle;
+	
+	UPROPERTY()
+	FLinearColor NodeColor;
+
+	UPROPERTY()
+	FText Tooltip;
+	
+	UPROPERTY()
+	FString Category;
+
+	UPROPERTY()
+	FString Keywords;	
+
+	UPROPERTY()
+	TArray<FRigVMGraphFunctionArgument> Arguments;
+
+	UPROPERTY()
+	TMap<FRigVMGraphFunctionIdentifier, uint32> Dependencies;
+
+	UPROPERTY()
+	TArray<FRigVMExternalVariable> ExternalVariables;
+
+	bool IsMutable() const;
+
+	bool IsValid() const { return !LibraryPointer.HostObject.IsNull(); }
+
+	FORCEINLINE FString GetHash() const
+	{
+		return FString::Printf(TEXT("%s:%s"), *LibraryPointer.HostObject.ToString(), *Name.ToString());
+	}
+
+	friend FORCEINLINE uint32 GetTypeHash(const FRigVMGraphFunctionHeader& Header)
+	{
+		return GetTypeHash(Header.LibraryPointer);
+	}
+
+	bool operator==(const FRigVMGraphFunctionHeader& Other) const
+	{
+		return LibraryPointer == Other.LibraryPointer;
+	}
+
+	IRigVMGraphFunctionHost* GetFunctionHost() const;
+
+	FRigVMGraphFunctionData* GetFunctionData() const;
+};
+
+FORCEINLINE FArchive& operator<<(FArchive& Ar, FRigVMGraphFunctionHeader& Data)
+{
+	Ar << Data.LibraryPointer;
+	Ar << Data.Name;
+	Ar << Data.NodeTitle;
+	Ar << Data.NodeColor;
+	Ar << Data.Tooltip;
+	Ar << Data.Category;
+	Ar << Data.Keywords;
+	Ar << Data.Arguments;
+	Ar << Data.Dependencies;
+	Ar << Data.ExternalVariables;
+	return Ar;
+}
+
+USTRUCT(BlueprintType)
+struct RIGVM_API FRigVMGraphFunctionData
+{
+	GENERATED_BODY()
+
+	FRigVMGraphFunctionData(){}
+
+	FRigVMGraphFunctionData(const FRigVMGraphFunctionHeader& InHeader)
+		: Header(InHeader) {	}
+
+	UPROPERTY()
+	FRigVMGraphFunctionHeader Header;
+
+	UPROPERTY()
+	FRigVMFunctionCompilationData CompilationData;
+
+	bool IsMutable() const;
+
+	bool operator==(const FRigVMGraphFunctionData& Other) const
+	{
+		return Header == Other.Header;
+	}
+
+	void ClearCompilationData() { CompilationData = FRigVMFunctionCompilationData(); }
+};
+
+FORCEINLINE FArchive& operator<<(FArchive& Ar, FRigVMGraphFunctionData& Data)
+{
+	Ar << Data.Header;
+	Ar << Data.CompilationData;
+	return Ar;
+}
\ No newline at end of file
diff --git a/Engine/Source/Runtime/RigVM/Public/RigVMCore/RigVMGraphFunctionHost.h b/Engine/Source/Runtime/RigVM/Public/RigVMCore/RigVMGraphFunctionHost.h
new file mode 100644
--- /dev/null
+++ b/Engine/Source/Runtime/RigVM/Public/RigVMCore/RigVMGraphFunctionHost.h
@@ -0,0 +1,81 @@
+// Copyright Epic Games, Inc. All Rights Reserved.
+
+#pragma once
+
+#include "RigVMCore/RigVMGraphFunctionDefinition.h"
+#include "RigVMGraphFunctionHost.generated.h"
+
+struct FRigVMGraphFunctionStore;
+
+UINTERFACE()
+class RIGVM_API URigVMGraphFunctionHost : public UInterface
+{
+	GENERATED_BODY()
+};
+
+// Interface that allows an object to host a rig VM graph function Host.
+class RIGVM_API IRigVMGraphFunctionHost
+{
+	GENERATED_BODY()
+
+public:
+	
+	// IRigVMClientHost interface
+	virtual FRigVMGraphFunctionStore* GetRigVMGraphFunctionStore() = 0;
+	virtual const FRigVMGraphFunctionStore* GetRigVMGraphFunctionStore() const = 0;
+};
+
+// A management struct containing graph functions
+USTRUCT()
+struct RIGVM_API FRigVMGraphFunctionStore
+{
+public:
+
+	GENERATED_BODY()
+
+	/** Exposed public functions on this rig */
+	UPROPERTY()
+	TArray<FRigVMGraphFunctionData> PublicFunctions;
+
+	UPROPERTY(Transient)
+	TArray<FRigVMGraphFunctionData> PrivateFunctions;
+
+	const FRigVMGraphFunctionData* FindFunction(const FRigVMGraphFunctionIdentifier& InLibraryPointer, bool* bOutIsPublic = nullptr) const;
+
+	FRigVMGraphFunctionData* FindFunction(const FRigVMGraphFunctionIdentifier& InLibraryPointer, bool* bOutIsPublic = nullptr);
+
+	FRigVMGraphFunctionData* FindFunctionByName(const FName& Name, bool *bOutIsPublic = nullptr);
+
+	bool ContainsFunction(const FRigVMGraphFunctionIdentifier& InLibraryPointer) const;
+
+	bool IsFunctionPublic(const FRigVMGraphFunctionIdentifier& InLibraryPointer) const;
+
+	FRigVMGraphFunctionData* AddFunction(const FRigVMGraphFunctionHeader& FunctionHeader, bool bIsPublic);
+
+	bool RemoveFunction(const FRigVMGraphFunctionIdentifier& InLibraryPointer, bool* bIsPublic = nullptr);
+
+	bool MarkFunctionAsPublic(const FRigVMGraphFunctionIdentifier& InLibraryPointer, bool bIsPublic);
+
+	FRigVMGraphFunctionData* UpdateFunctionInterface(const FRigVMGraphFunctionHeader& Header);
+
+	bool UpdateDependencies(const FRigVMGraphFunctionIdentifier& InLibraryPointer, TMap<FRigVMGraphFunctionIdentifier, uint32>& Dependencies);
+
+	bool UpdateExternalVariables(const FRigVMGraphFunctionIdentifier& InLibraryPointer, TArray<FRigVMExternalVariable> ExternalVariables);
+
+	bool UpdateFunctionCompilationData(const FRigVMGraphFunctionIdentifier& InLibraryPointer, const FRigVMFunctionCompilationData& CompilationData);
+
+	bool RemoveFunctionCompilationData(const FRigVMGraphFunctionIdentifier& InLibraryPointer);
+
+	bool RemoveAllCompilationData();
+
+	FORCEINLINE friend FArchive& operator<<(FArchive& Ar, FRigVMGraphFunctionStore& Host)
+	{
+		Ar << Host.PublicFunctions;
+		return Ar;
+	}
+
+private:
+
+	const FRigVMGraphFunctionData* FindFunctionImpl(const FRigVMGraphFunctionIdentifier& InLibraryPointer, bool* bOutIsPublic = nullptr) const;
+};
+
