RigVM: Fix python unit test breakages

diff --git a/Engine/Source/Developer/RigVMDeveloper/Private/RigVMModel/RigVMController.cpp b/Engine/Source/Developer/RigVMDeveloper/Private/RigVMModel/RigVMController.cpp
--- a/Engine/Source/Developer/RigVMDeveloper/Private/RigVMModel/RigVMController.cpp
+++ b/Engine/Source/Developer/RigVMDeveloper/Private/RigVMModel/RigVMController.cpp
@@ -985,6 +985,17 @@ URigVMUnitNode* URigVMController::AddUnitNodeFromStructPath(const FString& InScr
 	}
 
 	UScriptStruct* ScriptStruct = URigVMPin::FindObjectFromCPPTypeObjectPath<UScriptStruct>(InScriptStructPath);
+	if (ScriptStruct == nullptr)
+	{
+		const FCoreRedirectObjectName OldObjectName(InScriptStructPath);
+		const FCoreRedirectObjectName NewObjectName = FCoreRedirects::GetRedirectedName(ECoreRedirectFlags::Type_Struct, OldObjectName);
+		if (OldObjectName != NewObjectName)
+		{
+			ScriptStruct = URigVMPin::FindObjectFromCPPTypeObjectPath<UScriptStruct>(NewObjectName.ToString());
+		}
+		
+	}
+	
 	if (ScriptStruct == nullptr)
 	{
 		ReportErrorf(TEXT("Cannot find struct for path '%s'."), *InScriptStructPath);
@@ -13966,23 +13977,25 @@ bool URigVMController::CanAddFunctionRefForDefinition(const FRigVMGraphFunctionH
 
 	if (!bAllowPrivateFunctions)
 	{
-		IRigVMClientHost* ClientHost = GetImplementingOuter<IRigVMClientHost>();
-		bool bIsAvailable = ClientHost->GetRigVMClient()->GetFunctionLibrary()->GetFunctionHostObjectPath() ==
-			InFunctionDefinition.LibraryPointer.HostObject;
-		if (!bIsAvailable)
+		if(IRigVMClientHost* ClientHost = GetImplementingOuter<IRigVMClientHost>())
 		{
-			if (IRigVMGraphFunctionHost* Host = Cast<IRigVMGraphFunctionHost>(InFunctionDefinition.LibraryPointer.HostObject.ResolveObject()))
+			bool bIsAvailable = ClientHost->GetRigVMClient()->GetFunctionLibrary()->GetFunctionHostObjectPath() ==
+				InFunctionDefinition.LibraryPointer.HostObject;
+			if (!bIsAvailable)
 			{
-				bIsAvailable = Host->GetRigVMGraphFunctionStore()->IsFunctionPublic(InFunctionDefinition.LibraryPointer);		
+				if (IRigVMGraphFunctionHost* Host = Cast<IRigVMGraphFunctionHost>(InFunctionDefinition.LibraryPointer.HostObject.ResolveObject()))
+				{
+					bIsAvailable = Host->GetRigVMGraphFunctionStore()->IsFunctionPublic(InFunctionDefinition.LibraryPointer);		
+				}
 			}
-		}
-		if (!bIsAvailable)
-		{
-			if(bReportErrors)
+			if (!bIsAvailable)
 			{
-				ReportAndNotifyError(TEXT("Function is not available for placement in another graph host."));
+				if(bReportErrors)
+				{
+					ReportAndNotifyError(TEXT("Function is not available for placement in another graph host."));
+				}
+				return false;
 			}
-			return false;
 		}
 	}
 
