RigVM: lazy evaluation of input pins - developer side

diff --git a/Engine/Source/Developer/RigVMDeveloper/Private/RigVMCompiler/RigVMAST.cpp b/Engine/Source/Developer/RigVMDeveloper/Private/RigVMCompiler/RigVMAST.cpp
--- a/Engine/Source/Developer/RigVMDeveloper/Private/RigVMCompiler/RigVMAST.cpp
+++ b/Engine/Source/Developer/RigVMDeveloper/Private/RigVMCompiler/RigVMAST.cpp
@@ -1298,14 +1298,74 @@ FRigVMExprAST* FRigVMParserAST::TraversePin(const FRigVMASTProxy& InPinProxy, FR
 
 	FRigVMExprAST* ParentExprForLinks = PinExpr;
 
-	if ((Pin->GetDirection() == ERigVMPinDirection::IO || Pin->GetDirection() == ERigVMPinDirection::Input) &&
-		(InParentExpr->IsA(FRigVMExprAST::If) || InParentExpr->IsA(FRigVMExprAST::Select)) &&
-		LinkIndices.Num() > 0)
+	if(LinkIndices.Num() > 0)
 	{
-		FRigVMBlockExprAST* BlockExpr = MakeExpr<FRigVMBlockExprAST>(FRigVMExprAST::EType::Block, FRigVMASTProxy());
-		BlockExpr->AddParent(PinExpr);
-		BlockExpr->Name = Pin->GetFName();
-		ParentExprForLinks = BlockExpr;
+		if (Pin->IsLazy())
+		{
+			// create a block for each lazily executing pin
+			FRigVMBlockExprAST* BlockExpr = MakeExpr<FRigVMBlockExprAST>(FRigVMExprAST::EType::Block, FRigVMASTProxy());
+			BlockExpr->Name = Pin->GetFName();
+			BlockExpr->AddParent(PinExpr);
+			ParentExprForLinks = BlockExpr;
+		}
+		else if (Pin->GetNode()->HasLazyPin(true))
+		{
+			// for greedy pins on nodes containing a lazy pin - we need to also create a block for the node itself
+			FRigVMBlockExprAST* BlockExpr = nullptr;
+
+			if(const FRigVMExprAST* NodeExpr = PinExpr->GetFirstParentOfType(FRigVMExprAST::CallExtern))
+			{
+				const FRigVMCallExternExprAST* CallExternExpr = NodeExpr->To<FRigVMCallExternExprAST>();
+				
+				// try to find the block under all non-lazy pins.
+				// this is the block that is run before anything else - so for example
+				// if you have a lazy interplate with values A and B being lazy and a greedy T blend pin,
+				// you need a block to store the instructions related to computing T. once T is computed,
+				// the callextern can run and lazily pull on A or B.
+				for(const URigVMPin* NodePin : Pin->GetNode()->GetPins())
+				{
+					if(NodePin == Pin)
+					{
+						continue;
+					}
+					
+					if(!NodePin->IsLazy() &&
+						(NodePin->GetDirection() == ERigVMPinDirection::Input ||
+						NodePin->GetDirection() == ERigVMPinDirection::IO))
+					{
+						if(const FRigVMExprAST* NodePinExpr = CallExternExpr->FindExprWithPinName(NodePin->GetFName()))
+						{
+							if(const FRigVMExprAST* ExistingBlockExpr = NodePinExpr->GetFirstChildOfType(FRigVMExprAST::Block))
+							{
+								BlockExpr = (FRigVMBlockExprAST*)ExistingBlockExpr->To<FRigVMBlockExprAST>();
+								break;
+							}
+						}
+					}
+				}
+			}
+
+			if(BlockExpr == nullptr)
+			{
+				BlockExpr = MakeExpr<FRigVMBlockExprAST>(FRigVMExprAST::EType::Block, FRigVMASTProxy());
+				BlockExpr->Name = Pin->GetFName();
+			}
+
+			if(BlockExpr)
+			{
+				BlockExpr->AddParent(PinExpr);
+				ParentExprForLinks = BlockExpr;
+			}
+		}
+		else if ((Pin->GetDirection() == ERigVMPinDirection::IO || Pin->GetDirection() == ERigVMPinDirection::Input) &&
+			(InParentExpr->IsA(FRigVMExprAST::If) || InParentExpr->IsA(FRigVMExprAST::Select)) &&
+			LinkIndices.Num() > 0)
+		{
+			FRigVMBlockExprAST* BlockExpr = MakeExpr<FRigVMBlockExprAST>(FRigVMExprAST::EType::Block, FRigVMASTProxy());
+			BlockExpr->Name = Pin->GetFName();
+			BlockExpr->AddParent(PinExpr);
+			ParentExprForLinks = BlockExpr;
+		}
 	}
 
 	for (const int32 LinkIndex : LinkIndices)
@@ -1539,28 +1599,84 @@ void FRigVMParserAST::BubbleUpExpressions()
 			BlockCandidates.Append(Blocks);
 			FRigVMBlockExprAST* OuterBlock = nullptr;
 
-			// deal with a case where an expression is linked within both the true and false case of an "if" node
-			if(Blocks.Num() == 2)
+			// check if we need to add this expression to the block representing the non-lazy pins on the node.
+			// the expressions in that block are supposed to run before the node can execute - since they are greedy. 
+			// note: for backwards compatibility this also works for the condition pin on the if node and the index pin
+			// on the select node
+			TArray<const FRigVMExprAST*> Parents;
+			TArray<const FRigVMExprAST*> GrandParents;
+
+			for(const FRigVMBlockExprAST* Block : Blocks)
 			{
-				const FRigVMExprAST* Parent0 = Blocks[0]->GetParent();
-				const FRigVMExprAST* Parent1 = Blocks[1]->GetParent();
-				if(Parent0 && Parent1)
+				const FRigVMExprAST* Parent = Block->GetParent();
+				Parents.AddUnique(Parent);
+				if(Parent)
 				{
-					const FRigVMExprAST* GrandParent0 = Parent0->GetParent();
-					const FRigVMExprAST* GrandParent1 = Parent1->GetParent();
-					if(GrandParent0 && GrandParent1 && GrandParent0 == GrandParent1)
+					GrandParents.AddUnique(Parent->GetParent());
+				}
+			}
+
+			// if any of the parents is nullptr - it means that the block didn't
+			// have a parent and thus is the top level block. in that case we cannot
+			// bubble up further.
+			if(!Parents.Contains(nullptr))
+			{
+				// if all blocks here are part of the same grandparent
+				// parent == pin, grandparent == node
+				if(GrandParents.Num() == 1 && !GrandParents.Contains(nullptr))
+				{
+					if(GrandParents[0]->IsA(FRigVMExprAST::EType::CallExtern) ||
+						GrandParents[0]->IsA(FRigVMExprAST::EType::If) ||
+						GrandParents[0]->IsA(FRigVMExprAST::EType::Select))
 					{
-						if(GrandParent0->IsA(FRigVMExprAST::EType::If))
+						const FRigVMNodeExprAST* NodeExpr = GrandParents[0]->To<FRigVMNodeExprAST>();
+
+						// find a pin on the node which is either an input or an io and not lazy
+						const URigVMNode* Node = NodeExpr->GetNode();
+						const FRigVMExprAST* TopLevelBlockExpression = nullptr;
+						for(const URigVMPin* Pin : Node->GetPins())
 						{
-							const FRigVMIfExprAST* IfExpression = GrandParent0->To<FRigVMIfExprAST>();
-							const FRigVMExprAST* ConditionBlockExpression = IfExpression->GetConditionExpr()->GetFirstChildOfType(FRigVMExprAST::Block);
-							if(ConditionBlockExpression)
+							if(Pin->IsLazy())
 							{
-								OuterBlock = (FRigVMBlockExprAST*)ConditionBlockExpression->To<FRigVMBlockExprAST>();
-								OuterBlock->Children.Add(Expression);
-								Expression->Parents.Insert(OuterBlock, 0);
 								continue;
 							}
+
+							if(Pin->GetDirection() == ERigVMPinDirection::Input ||
+								Pin->GetDirection() == ERigVMPinDirection::IO)
+							{
+								if(Pin->GetNode()->IsA<URigVMIfNode>())
+								{
+									if(Pin->GetName() != URigVMIfNode::ConditionName)
+									{
+										continue;
+									}
+								}
+								else if(Pin->GetNode()->IsA<URigVMSelectNode>())
+								{
+									if(Pin->GetName() != URigVMSelectNode::IndexName)
+									{
+										continue;
+									}
+								}
+
+								if(const FRigVMExprAST* PinExpr = NodeExpr->FindExprWithPinName(Pin->GetFName()))
+								{
+									TopLevelBlockExpression = PinExpr->GetFirstChildOfType(FRigVMExprAST::Block);
+								}
+							}
+
+							if(TopLevelBlockExpression)
+							{
+								break;
+							}
+						}
+
+						if(TopLevelBlockExpression)
+						{
+							OuterBlock = (FRigVMBlockExprAST*)TopLevelBlockExpression->To<FRigVMBlockExprAST>();
+							OuterBlock->Children.Add(Expression);
+							Expression->Parents.Insert(OuterBlock, 0);
+							continue;
 						}
 					}
 				}
diff --git a/Engine/Source/Developer/RigVMDeveloper/Private/RigVMCompiler/RigVMCompiler.cpp b/Engine/Source/Developer/RigVMDeveloper/Private/RigVMCompiler/RigVMCompiler.cpp
--- a/Engine/Source/Developer/RigVMDeveloper/Private/RigVMCompiler/RigVMCompiler.cpp
+++ b/Engine/Source/Developer/RigVMDeveloper/Private/RigVMCompiler/RigVMCompiler.cpp
@@ -1126,7 +1126,97 @@ void URigVMCompiler::TraverseBlock(const FRigVMBlockExprAST* InExpr, FRigVMCompi
 	{
 		return;
 	}
+
+	if (InExpr->NumChildren() == 0)
+	{
+		return;
+	}
+
+	// check if the block is under a lazy pin, in which case we need to set up a branch info
+	URigVMNode* CallExternNode = nullptr;
+	FRigVMBranchInfo BranchInfo;
+	if(!WorkData.bSetupMemory)
+	{
+		if(const FRigVMExprAST* ParentExpr = InExpr->GetParent())
+		{
+			if(const FRigVMExprAST* GrandParentExpr = ParentExpr->GetParent())
+			{
+				if(GrandParentExpr->IsA(FRigVMExprAST::CallExtern))
+				{
+					const URigVMPin* Pin = nullptr;
+					if(ParentExpr->IsA(FRigVMExprAST::Var))
+					{
+						Pin = ParentExpr->To<FRigVMVarExprAST>()->GetPin();
+					}
+					else if(ParentExpr->IsA(FRigVMExprAST::CachedValue))
+					{
+						Pin = ParentExpr->To<FRigVMCachedValueExprAST>()->GetVarExpr()->GetPin();
+					}
+
+					if(Pin)
+					{
+						Pin = Pin->GetRootPin();
+						if(Pin->IsLazy())
+						{
+							CallExternNode = Pin->GetNode();
+							
+							BranchInfo.Label = Pin->GetFName();
+							BranchInfo.InstructionIndex = INDEX_NONE; // we'll fill in the instruction info later
+							BranchInfo.FirstInstruction = WorkData.VM->GetByteCode().GetNumInstructions();
+
+							// find the argument index for the given pin
+							if(const URigVMTemplateNode* TemplateNode = Cast<URigVMTemplateNode>(CallExternNode))
+							{
+								if(const FRigVMTemplate* Template = TemplateNode->GetTemplate())
+								{
+									for(int32 ArgumentIndex = 0; ArgumentIndex != Template->NumArguments(); ArgumentIndex++)
+									{
+										const FRigVMTemplateArgument* Argument = Template->GetArgument(ArgumentIndex);
+										if(Template->GetArgument(ArgumentIndex)->GetName() == Pin->GetFName())
+										{
+											BranchInfo.ArgumentIndex = ArgumentIndex;
+											break;
+										}
+									}
+								}
+								// we also need to deal with unit nodes separately here. if a unit node does
+								// not offer a valid backing template - we need to visit its properties. since
+								// templates don't contain executecontext type arguments anymore - we need
+								// to step over them as well here.
+								else if(const URigVMUnitNode* UnitNode = Cast<URigVMUnitNode>(CallExternNode))
+								{
+									int32 ArgumentIndex = 0;
+									for(const URigVMPin* NodePin : UnitNode->GetPins())
+									{
+										if(NodePin->IsExecuteContext())
+										{
+											continue;
+										}
+										if(NodePin->GetFName() == Pin->GetFName())
+										{
+											BranchInfo.ArgumentIndex = ArgumentIndex;
+											break;
+										}
+										ArgumentIndex++;
+									}
+								}	
+							}
+
+							check(BranchInfo.ArgumentIndex != INDEX_NONE);
+						}
+					}
+				}
+			}
+		}
+	}
+	
 	TraverseChildren(InExpr, WorkData);
+
+	if(!BranchInfo.Label.IsNone())
+	{
+		BranchInfo.LastInstruction = WorkData.VM->GetByteCode().GetNumInstructions() - 1;
+		WorkData.BranchInfos.FindOrAdd(CallExternNode).Add(BranchInfo);
+	}
 }
 
 void URigVMCompiler::TraverseEntry(const FRigVMEntryExprAST* InExpr, FRigVMCompilerWorkData& WorkData)
@@ -1333,7 +1423,55 @@ int32 URigVMCompiler::TraverseCallExtern(const FRigVMCallExternExprAST* InExpr,
 			}
 		}
 
-		TraverseChildren(InExpr, WorkData);
+		// traverse all non-lazy children
+		TArray<const FRigVMExprAST*> LazyChildExprs;
+		for (const FRigVMExprAST* ChildExpr : *InExpr)
+		{
+			// if there's a direct child block under this - the pin may be lazy
+			if(ChildExpr->IsA(FRigVMExprAST::Var) || ChildExpr->IsA(FRigVMExprAST::CachedValue))
+			{
+				if(const FRigVMExprAST* BlockExpr = ChildExpr->GetFirstChildOfType(FRigVMExprAST::Block))
+				{
+					if(BlockExpr->GetParent() == ChildExpr)
+					{
+						URigVMPin* Pin = nullptr;
+						if(ChildExpr->IsA(FRigVMExprAST::Var))
+						{
+							Pin = ChildExpr->To<FRigVMVarExprAST>()->GetPin();
+						}
+						else
+						{
+							Pin = ChildExpr->To<FRigVMCachedValueExprAST>()->GetVarExpr()->GetPin();
+						}
+						check(Pin);
+						
+						if(Pin->IsLazy())
+						{
+							LazyChildExprs.Add(ChildExpr);
+							continue;
+						}
+					}
+				}
+			}
+			TraverseExpression(ChildExpr, WorkData);
+		}
+
+		if(!LazyChildExprs.IsEmpty())
+		{
+			// set up an operator to skip the lazy branches 
+			const uint64 JumpToCallExternByte = WorkData.VM->GetByteCode().AddJumpOp(ERigVMOpCode::JumpForward, INDEX_NONE);
+			const int32 JumpToCallExternInstruction = WorkData.VM->GetByteCode().GetNumInstructions() - 1;
+
+			// traverse the lazy children 
+			for (const FRigVMExprAST* ChildExpr : LazyChildExprs)
+			{
+				TraverseExpression(ChildExpr, WorkData);
+			}
+
+			// update the operator with the target instruction 
+			const int32 InstructionsToJump = WorkData.VM->GetByteCode().GetNumInstructions() - JumpToCallExternInstruction;
+			WorkData.VM->GetByteCode().GetOpAt<FRigVMJumpOp>(JumpToCallExternByte).InstructionIndex = InstructionsToJump;
+		}
 
 		// setup the instruction
 		const int32 FunctionIndex = WorkData.VM->AddRigVMFunction(Function->GetName());
@@ -1341,6 +1479,17 @@ int32 URigVMCompiler::TraverseCallExtern(const FRigVMCallExternExprAST* InExpr,
 		WorkData.VM->GetByteCode().AddExecuteOp(FunctionIndex, Operands);
 		InstructionIndex = WorkData.VM->GetByteCode().GetNumInstructions() - 1;
 
+		// setup the branch infos for this call extern instruction
+		if(const TArray<FRigVMBranchInfo>* BranchInfosPtr = WorkData.BranchInfos.Find(Node))
+		{
+			const TArray<FRigVMBranchInfo>& BranchInfos = *BranchInfosPtr;
+			for(FRigVMBranchInfo BranchInfo : BranchInfos)
+			{
+				BranchInfo.InstructionIndex = InstructionIndex;
+				WorkData.VM->GetByteCode().AddBranchInfo(BranchInfo);
+			}
+		}
+
 #if WITH_EDITORONLY_DATA
 		TArray<FRigVMOperand> InputsOperands, OutputOperands;
 
diff --git a/Engine/Source/Developer/RigVMDeveloper/Private/RigVMModel/Nodes/RigVMDispatchNode.cpp b/Engine/Source/Developer/RigVMDeveloper/Private/RigVMModel/Nodes/RigVMDispatchNode.cpp
--- a/Engine/Source/Developer/RigVMDeveloper/Private/RigVMModel/Nodes/RigVMDispatchNode.cpp
+++ b/Engine/Source/Developer/RigVMDeveloper/Private/RigVMModel/Nodes/RigVMDispatchNode.cpp
@@ -174,6 +174,18 @@ void URigVMDispatchNode::InvalidateCache()
 	CachedFactory = nullptr;
 }
 
+bool URigVMDispatchNode::ShouldInputPinComputeLazily(const URigVMPin* InPin) const
+{
+	const URigVMPin* RootPin = InPin->GetRootPin();
+	check(RootPin->GetNode() == this);
+
+	if(const FRigVMDispatchFactory* Factory = GetFactory())
+	{
+		return Factory->HasArgumentMetaData(RootPin->GetFName(), FRigVMStruct::ComputeLazilyMetaName); 
+	}
+	return Super::ShouldInputPinComputeLazily(InPin);
+}
+
 FRigVMStructUpgradeInfo URigVMDispatchNode::GetUpgradeInfo() const
 {
 	if(const FRigVMDispatchFactory* Factory = GetFactory())
diff --git a/Engine/Source/Developer/RigVMDeveloper/Private/RigVMModel/Nodes/RigVMUnitNode.cpp b/Engine/Source/Developer/RigVMDeveloper/Private/RigVMModel/Nodes/RigVMUnitNode.cpp
--- a/Engine/Source/Developer/RigVMDeveloper/Private/RigVMModel/Nodes/RigVMUnitNode.cpp
+++ b/Engine/Source/Developer/RigVMDeveloper/Private/RigVMModel/Nodes/RigVMUnitNode.cpp
@@ -136,6 +136,22 @@ FText URigVMUnitNode::GetToolTipTextForPin(const URigVMPin* InPin) const
 	return URigVMNode::GetToolTipTextForPin(InPin);
 }
 
+bool URigVMUnitNode::ShouldInputPinComputeLazily(const URigVMPin* InPin) const
+{
+	const URigVMPin* RootPin = InPin->GetRootPin();
+	check(RootPin->GetNode() == this);
+
+	if(const UScriptStruct* Struct = GetScriptStruct())
+	{
+		if(const FProperty* Property = Struct->FindPropertyByName(RootPin->GetFName()))
+		{
+			return Property->HasMetaData(FRigVMStruct::ComputeLazilyMetaName);
+		}
+	}
+
+	return Super::ShouldInputPinComputeLazily(InPin);
+}
+
 bool URigVMUnitNode::IsDeprecated() const
 {
 	return !GetDeprecatedMetadata().IsEmpty();
diff --git a/Engine/Source/Developer/RigVMDeveloper/Private/RigVMModel/RigVMNode.cpp b/Engine/Source/Developer/RigVMDeveloper/Private/RigVMModel/RigVMNode.cpp
--- a/Engine/Source/Developer/RigVMDeveloper/Private/RigVMModel/RigVMNode.cpp
+++ b/Engine/Source/Developer/RigVMDeveloper/Private/RigVMModel/RigVMNode.cpp
@@ -343,6 +343,22 @@ bool URigVMNode::HasIOPin() const
 	return HasPinOfDirection(ERigVMPinDirection::IO);
 }
 
+bool URigVMNode::HasLazyPin(bool bOnlyConsiderPinsWithLinks) const
+{
+	return Pins.ContainsByPredicate([bOnlyConsiderPinsWithLinks](const URigVMPin* Pin) -> bool
+	{
+		if(Pin->IsLazy())
+		{
+			if(bOnlyConsiderPinsWithLinks)
+			{
+				return Pin->GetLinkedSourcePins(true).Num() > 0;
+			}
+			return true;
+		}
+		return false;
+	});
+}
+
 bool URigVMNode::HasOutputPin(bool bIncludeIO) const
 {
 	if (HasPinOfDirection(ERigVMPinDirection::Output))
diff --git a/Engine/Source/Developer/RigVMDeveloper/Private/RigVMModel/RigVMPin.cpp b/Engine/Source/Developer/RigVMDeveloper/Private/RigVMModel/RigVMPin.cpp
--- a/Engine/Source/Developer/RigVMDeveloper/Private/RigVMModel/RigVMPin.cpp
+++ b/Engine/Source/Developer/RigVMDeveloper/Private/RigVMModel/RigVMPin.cpp
@@ -565,6 +565,18 @@ bool URigVMPin::IsDynamicArray() const
 	return bIsDynamicArray;
 }
 
+bool URigVMPin::IsLazy() const
+{
+	if(GetDirection() == ERigVMPinDirection::Input)
+	{
+		if(const URigVMNode* Node = GetNode())
+		{
+			return Node->ShouldInputPinComputeLazily(this);
+		}
+	}
+	return false;
+}
+
 int32 URigVMPin::GetPinIndex() const
 {
 	int32 Index = INDEX_NONE;
diff --git a/Engine/Source/Developer/RigVMDeveloper/Public/RigVMCompiler/RigVMAST.h b/Engine/Source/Developer/RigVMDeveloper/Public/RigVMCompiler/RigVMAST.h
--- a/Engine/Source/Developer/RigVMDeveloper/Public/RigVMCompiler/RigVMAST.h
+++ b/Engine/Source/Developer/RigVMDeveloper/Public/RigVMCompiler/RigVMAST.h
@@ -13,6 +13,7 @@
 
 class FRigVMParserAST;
 class FRigVMBlockExprAST;
+class FRigVMNodeExprAST;
 class FRigVMEntryExprAST;
 class FRigVMInvokeEntryExprAST;
 class FRigVMCallExternExprAST;
@@ -297,7 +298,7 @@ protected:
 	friend class URigVMCompiler;
 };
 
-// specialized cast for type checkiFRigVMAssignExprASTng
+// specialized cast for type checking
 // for a Block / FRigVMBlockExprAST expression
 // will raise if types are not compatible
 // @return this expression cast to FRigVMBlockExprAST
@@ -308,6 +309,27 @@ FORCEINLINE const FRigVMBlockExprAST* FRigVMExprAST::To() const
 	return (const FRigVMBlockExprAST*)this;
 }
 
+// specialized cast for type checking
+// for a Node / FRigVMNodeExprAST expression
+// will raise if types are not compatible
+// @return this expression cast to FRigVMNodeExprAST
+template<>
+FORCEINLINE const FRigVMNodeExprAST* FRigVMExprAST::To() const
+{
+	ensure(
+		IsA(EType::Entry) ||
+		IsA(EType::InvokeEntry) ||
+		IsA(EType::CallExtern) ||
+		IsA(EType::InlineFunction) ||
+		IsA(EType::NoOp) ||
+		IsA(EType::Branch) ||
+		IsA(EType::If) ||
+		IsA(EType::Select) ||
+		IsA(EType::Array)
+	);
+	return (const FRigVMNodeExprAST*)this;
+}
+
 // specialized cast for type checking
 // for a Entry / FRigVMEntryExprAST expression
 // will raise if types are not compatible
diff --git a/Engine/Source/Developer/RigVMDeveloper/Public/RigVMCompiler/RigVMCompiler.h b/Engine/Source/Developer/RigVMDeveloper/Public/RigVMCompiler/RigVMCompiler.h
--- a/Engine/Source/Developer/RigVMDeveloper/Public/RigVMCompiler/RigVMCompiler.h
+++ b/Engine/Source/Developer/RigVMDeveloper/Public/RigVMCompiler/RigVMCompiler.h
@@ -104,6 +104,7 @@ public:
 	TMap<FRigVMASTProxy, FRigVMASTProxyArray> CachedProxiesWithSharedOperand;
 	const FRigVMASTProxySourceMap* ProxySources;
 	FRigVMASTProxyTargetsMap ProxyTargets;
+	TMap<URigVMNode*, TArray<FRigVMBranchInfo>> BranchInfos;
 
 	struct FFunctionRegisterData
 	{
diff --git a/Engine/Source/Developer/RigVMDeveloper/Public/RigVMModel/Nodes/RigVMDispatchNode.h b/Engine/Source/Developer/RigVMDeveloper/Public/RigVMModel/Nodes/RigVMDispatchNode.h
--- a/Engine/Source/Developer/RigVMDeveloper/Public/RigVMModel/Nodes/RigVMDispatchNode.h
+++ b/Engine/Source/Developer/RigVMDeveloper/Public/RigVMModel/Nodes/RigVMDispatchNode.h
@@ -54,6 +54,7 @@ protected:
 	const UScriptStruct* GetFactoryStruct() const;
 	const FRigVMTemplateTypeMap& GetFilteredTypes() const;
 	virtual void InvalidateCache() override;
+	virtual bool ShouldInputPinComputeLazily(const URigVMPin* InPin) const override;
 
 private:
 
diff --git a/Engine/Source/Developer/RigVMDeveloper/Public/RigVMModel/Nodes/RigVMUnitNode.h b/Engine/Source/Developer/RigVMDeveloper/Public/RigVMModel/Nodes/RigVMUnitNode.h
--- a/Engine/Source/Developer/RigVMDeveloper/Public/RigVMModel/Nodes/RigVMUnitNode.h
+++ b/Engine/Source/Developer/RigVMDeveloper/Public/RigVMModel/Nodes/RigVMUnitNode.h
@@ -76,6 +76,7 @@ public:
 protected:
 
 	virtual FText GetToolTipTextForPin(const URigVMPin* InPin) const override;
+	virtual bool ShouldInputPinComputeLazily(const URigVMPin* InPin) const override;
 
 private:
 
diff --git a/Engine/Source/Developer/RigVMDeveloper/Public/RigVMModel/RigVMNode.h b/Engine/Source/Developer/RigVMDeveloper/Public/RigVMModel/RigVMNode.h
--- a/Engine/Source/Developer/RigVMDeveloper/Public/RigVMModel/RigVMNode.h
+++ b/Engine/Source/Developer/RigVMDeveloper/Public/RigVMModel/RigVMNode.h
@@ -166,6 +166,10 @@ public:
 	UFUNCTION(BlueprintCallable, Category = RigVMNode)
 	virtual bool HasIOPin() const;
 
+	// Returns true if the node has any lazily evaluating pins
+	UFUNCTION(BlueprintCallable, Category = RigVMNode)
+	virtual bool HasLazyPin(bool bOnlyConsiderPinsWithLinks = false) const;
+
 	// Returns true if the node has any output pins
 	UFUNCTION(BlueprintCallable, Category = RigVMNode)
 	virtual bool HasOutputPin(bool bIncludeIO = true) const;
@@ -272,6 +276,7 @@ protected:
 	virtual TArray<int32> GetInstructionsForVMImpl(URigVM* InVM, const FRigVMASTProxy& InProxy = FRigVMASTProxy()) const; 
 	virtual FText GetToolTipTextForPin(const URigVMPin* InPin) const;
 	virtual bool AllowsLinksOn(const URigVMPin* InPin) const { return true; }
+	virtual bool ShouldInputPinComputeLazily(const URigVMPin* InPin) const { return false; }
 
 	UPROPERTY()
 	FString NodeTitle;
diff --git a/Engine/Source/Developer/RigVMDeveloper/Public/RigVMModel/RigVMPin.h b/Engine/Source/Developer/RigVMDeveloper/Public/RigVMModel/RigVMPin.h
--- a/Engine/Source/Developer/RigVMDeveloper/Public/RigVMModel/RigVMPin.h
+++ b/Engine/Source/Developer/RigVMDeveloper/Public/RigVMModel/RigVMPin.h
@@ -225,6 +225,10 @@ public:
 	UFUNCTION(BlueprintCallable, Category = RigVMPin)
 	bool IsReferenceCountedContainer() const { return IsDynamicArray(); }
 
+	// Returns true if this pin's value may be executed lazily
+	UFUNCTION(BlueprintCallable, Category = RigVMPin)
+	bool IsLazy() const;
+
 	// Returns the index of the Pin within the node / parent Pin
 	UFUNCTION(BlueprintCallable, Category = RigVMPin)
 	int32 GetPinIndex() const;
diff --git a/Engine/Source/Runtime/RigVM/Private/RigVMCore/RigVM.cpp b/Engine/Source/Runtime/RigVM/Private/RigVMCore/RigVM.cpp
--- a/Engine/Source/Runtime/RigVM/Private/RigVMCore/RigVM.cpp
+++ b/Engine/Source/Runtime/RigVM/Private/RigVMCore/RigVM.cpp
@@ -1119,7 +1119,7 @@ void URigVM::CacheMemoryHandlesIfRequired(TArrayView<URigVMMemoryStorage*> InMem
 	auto InstructionOpEval = [&](
 		int32 InstructionIndex,
 		int32 InHandleBaseIndex,
-		TFunctionRef<void(int32 InHandleIndex, const TRigVMBranchInfoKey& InBranchInfoKey, const FRigVMOperand& InArg)> InOpFunc
+		TFunctionRef<void(int32 InHandleIndex, const FRigVMBranchInfoKey& InBranchInfoKey, const FRigVMOperand& InArg)> InOpFunc
 		) -> void
 	{
 		const ERigVMOpCode OpCode = Instructions[InstructionIndex].OpCode; 
@@ -1267,7 +1267,7 @@ void URigVM::CacheMemoryHandlesIfRequired(TArrayView<URigVMMemoryStorage*> InMem
 	for (int32 InstructionIndex = 0; InstructionIndex < Instructions.Num(); InstructionIndex++)
 	{
 		InstructionOpEval(InstructionIndex, INDEX_NONE,
-			[&HandleCount](int32, const TRigVMBranchInfoKey&, const FRigVMOperand& )
+			[&HandleCount](int32, const FRigVMBranchInfoKey&, const FRigVMOperand& )
 			{
 				HandleCount++;
 			});
@@ -1290,7 +1290,7 @@ void URigVM::CacheMemoryHandlesIfRequired(TArrayView<URigVMMemoryStorage*> InMem
 		[&](int32 InstructionIndex)
 		{
 			InstructionOpEval(InstructionIndex, FirstHandleForInstruction[InstructionIndex],
-				[&](int32 InHandleIndex, const TRigVMBranchInfoKey& InBranchInfoKey, const FRigVMOperand& InOp)
+				[&](int32 InHandleIndex, const FRigVMBranchInfoKey& InBranchInfoKey, const FRigVMOperand& InOp)
 				{
 					CacheSingleMemoryHandle(InHandleIndex, InBranchInfoKey, InOp);
 				});
@@ -3374,7 +3374,7 @@ void URigVM::ClearDebugMemory()
 #endif
 }
 
-void URigVM::CacheSingleMemoryHandle(int32 InHandleIndex, const TRigVMBranchInfoKey& InBranchInfoKey, const FRigVMOperand& InArg, bool bForExecute)
+void URigVM::CacheSingleMemoryHandle(int32 InHandleIndex, const FRigVMBranchInfoKey& InBranchInfoKey, const FRigVMOperand& InArg, bool bForExecute)
 {
 	URigVMMemoryStorage* Memory = GetMemoryByType(InArg.GetMemoryType(), false);
 
diff --git a/Engine/Source/Runtime/RigVM/Private/RigVMCore/RigVMByteCode.cpp b/Engine/Source/Runtime/RigVM/Private/RigVMCore/RigVMByteCode.cpp
--- a/Engine/Source/Runtime/RigVM/Private/RigVMCore/RigVMByteCode.cpp
+++ b/Engine/Source/Runtime/RigVM/Private/RigVMCore/RigVMByteCode.cpp
@@ -695,13 +695,13 @@ void FRigVMByteCode::Load(FArchive& Ar)
 		Ar << BranchInfos;
 
 		// make sure the lookup table is up 2 date
-		BranchInfoPerInstruction.Reset();
+		BranchInfoLookup.Reset();
 		(void)GetBranchInfo({0,0});
 	}
 	else
 	{
 		BranchInfos.Reset();
-		BranchInfoPerInstruction.Reset();
+		BranchInfoLookup.Reset();
 	}
 }
 
@@ -712,7 +712,7 @@ void FRigVMByteCode::Reset()
 	NumInstructions = 0;
 	Entries.Reset();
 	BranchInfos.Reset();
-	BranchInfoPerInstruction.Reset();
+	BranchInfoLookup.Reset();
 
 #if WITH_EDITORONLY_DATA
 	SubjectPerInstruction.Reset();
@@ -1614,18 +1614,24 @@ uint64 FRigVMByteCode::AddInvokeEntryOp(const FName& InEntryName)
 	return AddOp(FRigVMInvokeEntryOp(InEntryName));
 }
 
+void FRigVMByteCode::AddBranchInfo(const FRigVMBranchInfo& InBranchInfo)
+{
+	FRigVMBranchInfo BranchInfo = InBranchInfo;
+	BranchInfo.Index = BranchInfos.Num();
+	BranchInfos.Add(BranchInfo);
+	BranchInfoLookup.Reset();
+}
+
 void FRigVMByteCode::AddBranchInfo(const FName& InBranchLabel, int32 InInstructionIndex, int32 InArgumentIndex,
 	int32 InFirstBranchInstruction, int32 InLastBranchInstruction)
 {
 	FRigVMBranchInfo BranchInfo;
-	BranchInfo.Index = BranchInfos.Num();
 	BranchInfo.Label = InBranchLabel;
 	BranchInfo.InstructionIndex = InInstructionIndex;
 	BranchInfo.ArgumentIndex = InArgumentIndex;
 	BranchInfo.FirstInstruction = InFirstBranchInstruction;
 	BranchInfo.LastInstruction = InLastBranchInstruction;
-	BranchInfos.Add(BranchInfo);
-	BranchInfoPerInstruction.Reset();
+	AddBranchInfo(BranchInfo);
 }
 
 FRigVMOperandArray FRigVMByteCode::GetOperandsForOp(const FRigVMInstruction& InInstruction) const
@@ -2377,23 +2383,23 @@ void FRigVMByteCode::SetOperandsForInstruction(int32 InInstructionIndex, const F
 
 #endif
 
-const FRigVMBranchInfo* FRigVMByteCode::GetBranchInfo(const TRigVMBranchInfoKey& InBranchInfoKey) const
+const FRigVMBranchInfo* FRigVMByteCode::GetBranchInfo(const FRigVMBranchInfoKey& InBranchInfoKey) const
 {
 	if(BranchInfos.IsEmpty())
 	{
 		return nullptr;
 	}
 
-	if(BranchInfoPerInstruction.IsEmpty())
+	if(BranchInfoLookup.IsEmpty())
 	{
 		for(const FRigVMBranchInfo& BranchInfo : BranchInfos)
 		{
-			const TRigVMBranchInfoKey Key(BranchInfo.InstructionIndex, BranchInfo.ArgumentIndex);
-			BranchInfoPerInstruction.FindOrAdd(Key) = &BranchInfo;
+			const FRigVMBranchInfoKey Key(BranchInfo.InstructionIndex, BranchInfo.ArgumentIndex, BranchInfo.Label);
+			BranchInfoLookup.FindOrAdd(Key) = &BranchInfo;
 		}
 	}
 	
-	if(const FRigVMBranchInfo** BranchInfoPtr = BranchInfoPerInstruction.Find(InBranchInfoKey))
+	if(const FRigVMBranchInfo** BranchInfoPtr = BranchInfoLookup.Find(InBranchInfoKey))
 	{
 		return *BranchInfoPtr;
 	}
diff --git a/Engine/Source/Runtime/RigVM/Private/RigVMCore/RigVMMemoryStorage.cpp b/Engine/Source/Runtime/RigVM/Private/RigVMCore/RigVMMemoryStorage.cpp
--- a/Engine/Source/Runtime/RigVM/Private/RigVMCore/RigVMMemoryStorage.cpp
+++ b/Engine/Source/Runtime/RigVM/Private/RigVMCore/RigVMMemoryStorage.cpp
@@ -112,8 +112,10 @@ const uint8* TRigVMLazyValueBase::GetData() const
 {
 	if(MemoryHandle)
 	{
-		checkf(MemoryHandle->IsLazy(), TEXT("If you are hitting this the compiler didn't add a branch info to the bytecode for a lazy argument."));
-		MemoryHandle->ComputeLazyValueIfNecessary(SliceIndex);
+		if(MemoryHandle->IsLazy())
+		{
+			MemoryHandle->ComputeLazyValueIfNecessary(SliceIndex);
+		}
 		return MemoryHandle->GetData(bFollowPropertyPath, SliceIndex);
 	}
 	return nullptr;
diff --git a/Engine/Source/Runtime/RigVM/Public/RigVMCore/RigVM.h b/Engine/Source/Runtime/RigVM/Public/RigVMCore/RigVM.h
--- a/Engine/Source/Runtime/RigVM/Public/RigVMCore/RigVM.h
+++ b/Engine/Source/Runtime/RigVM/Public/RigVMCore/RigVM.h
@@ -785,7 +785,7 @@ private:
 	// debug watch register memory needs to be cleared for each execution
 	void ClearDebugMemory();
 	
-	void CacheSingleMemoryHandle(int32 InHandleIndex, const TRigVMBranchInfoKey& InBranchInfoKey, const FRigVMOperand& InArg, bool bForExecute = false);
+	void CacheSingleMemoryHandle(int32 InHandleIndex, const FRigVMBranchInfoKey& InBranchInfoKey, const FRigVMOperand& InArg, bool bForExecute = false);
 
 	FORCEINLINE_DEBUGGABLE void CopyOperandForDebuggingIfNeeded(const FRigVMOperand& InArg, const FRigVMMemoryHandle& InHandle)
 	{
diff --git a/Engine/Source/Runtime/RigVM/Public/RigVMCore/RigVMByteCode.h b/Engine/Source/Runtime/RigVM/Public/RigVMCore/RigVMByteCode.h
--- a/Engine/Source/Runtime/RigVM/Public/RigVMCore/RigVMByteCode.h
+++ b/Engine/Source/Runtime/RigVM/Public/RigVMCore/RigVMByteCode.h
@@ -27,7 +27,67 @@ class FArchive;
 class UObject;
 struct FRigVMByteCode;
 
-typedef TTuple<int32,int32> TRigVMBranchInfoKey;
+struct FRigVMBranchInfoKey
+{
+	FRigVMBranchInfoKey()
+		: InstructionIndex(INDEX_NONE)
+		, ArgumentIndex(INDEX_NONE)
+		, Label(NAME_None)
+	{}
+
+	FRigVMBranchInfoKey(int32 InInstructionIndex, int32 InArgumentIndex)
+		: InstructionIndex(InInstructionIndex)
+		, ArgumentIndex(InArgumentIndex)
+		, Label(NAME_None)
+	{}
+
+	FRigVMBranchInfoKey(int32 InInstructionIndex, const FName& InLabel)
+		: InstructionIndex(InInstructionIndex)
+		, ArgumentIndex(INDEX_NONE)
+		, Label(InLabel)
+	{}
+
+	FRigVMBranchInfoKey(int32 InInstructionIndex, int32 InArgumentIndex, const FName& InLabel)
+		: InstructionIndex(InInstructionIndex)
+		, ArgumentIndex(InArgumentIndex)
+		, Label(InLabel)
+	{}
+
+	friend uint32 GetTypeHash(const FRigVMBranchInfoKey& InKey)
+	{
+		return HashCombine(
+			HashCombine(
+				GetTypeHash(InKey.InstructionIndex),
+				GetTypeHash(InKey.ArgumentIndex)
+			),
+			GetTypeHash(InKey.Label)
+		);
+	}
+
+	bool operator ==(const FRigVMBranchInfoKey& InOther) const
+	{
+		if(InstructionIndex != InOther.InstructionIndex)
+		{
+			return false;
+		}
+
+		if(ArgumentIndex != INDEX_NONE && InOther.ArgumentIndex != INDEX_NONE)
+		{
+			return ArgumentIndex == InOther.ArgumentIndex;
+		}
+
+		if(!Label.IsNone() && !InOther.Label.IsNone())
+		{
+			return Label == InOther.Label;
+		}
+
+		return true;
+	}
+
+	int32 InstructionIndex;
+	int32 ArgumentIndex;
+	FName Label;
+};
 
 
 // The code for a single operation within the RigVM
@@ -996,6 +1056,7 @@ public:
 	uint64 AddInvokeEntryOp(const FName& InEntryName);
 
 	// adds information about a branch for an instruction's argument
+	void AddBranchInfo(const FRigVMBranchInfo& InBranchInfo);
 	void AddBranchInfo(const FName& InBranchLabel, int32 InInstructionIndex, int32 InArgumentIndex, int32 InFirstBranchInstruction, int32 InLastBranchInstruction);
 
 	// returns an instruction array for iterating over all operators
@@ -1217,8 +1278,8 @@ private:
 	UPROPERTY()
 	TArray<FRigVMBranchInfo> BranchInfos;
 
-	const FRigVMBranchInfo* GetBranchInfo(const TRigVMBranchInfoKey& InBranchInfoKey) const;
-	mutable TMap<TRigVMBranchInfoKey, const FRigVMBranchInfo*> BranchInfoPerInstruction;
+	const FRigVMBranchInfo* GetBranchInfo(const FRigVMBranchInfoKey& InBranchInfoKey) const;
+	mutable TMap<FRigVMBranchInfoKey, const FRigVMBranchInfo*> BranchInfoLookup;
 
 	// if this is set to true the stored bytecode is aligned / padded
 	bool bByteCodeIsAligned;
diff --git a/Engine/Source/Runtime/RigVM/Public/RigVMCore/RigVMMemoryStorage.h b/Engine/Source/Runtime/RigVM/Public/RigVMCore/RigVMMemoryStorage.h
--- a/Engine/Source/Runtime/RigVM/Public/RigVMCore/RigVMMemoryStorage.h
+++ b/Engine/Source/Runtime/RigVM/Public/RigVMCore/RigVMMemoryStorage.h
@@ -780,9 +780,7 @@ private:
 
 	FORCEINLINE_DEBUGGABLE TRigVMLazyValueBase GetDataLazily_Internal(bool bFollowPropertyPath, int32 InSliceIndex) const
 	{
-		// this should only get called on memory handles which can execute lazily
-		check(LazyBranch != nullptr);
-		
+		// note: this works also for memory handles which don't provide a lazy branch
 		TRigVMLazyValueBase LazyValue;
 		LazyValue.MemoryHandle = (FRigVMMemoryHandle*)this;
 		LazyValue.bFollowPropertyPath = bFollowPropertyPath;
