Control Rig: Fix static analysis warning

diff --git a/Engine/Source/Developer/RigVMDeveloper/Private/RigVMModel/RigVMClient.cpp b/Engine/Source/Developer/RigVMDeveloper/Private/RigVMModel/RigVMClient.cpp
--- a/Engine/Source/Developer/RigVMDeveloper/Private/RigVMModel/RigVMClient.cpp
+++ b/Engine/Source/Developer/RigVMDeveloper/Private/RigVMModel/RigVMClient.cpp
@@ -261,7 +261,6 @@ void FRigVMClient::AddModel(URigVMGraph* InModel, bool bCreateController)
 	{
 		check(FunctionLibrary == nullptr);
 		FunctionLibrary = Cast<URigVMFunctionLibrary>(InModel);
-		FunctionLibrary->SetDefaultFunctionLibrary(FunctionLibrary);		
 	}
 	else
 	{
@@ -279,6 +278,7 @@ void FRigVMClient::AddModel(URigVMGraph* InModel, bool bCreateController)
 		{
 			Model->SetDefaultFunctionLibrary(FunctionLibrary);
 		}
+		InModel->SetDefaultFunctionLibrary(FunctionLibrary);
 	}
 	else if(FunctionLibrary)
 	{
