RigVM: Control Flow functionality - core side

diff --git a/Engine/Plugins/Animation/ControlRig/Source/ControlRig/Private/ControlRig.cpp b/Engine/Plugins/Animation/ControlRig/Source/ControlRig/Private/ControlRig.cpp
--- a/Engine/Plugins/Animation/ControlRig/Source/ControlRig/Private/ControlRig.cpp
+++ b/Engine/Plugins/Animation/ControlRig/Source/ControlRig/Private/ControlRig.cpp
@@ -729,6 +729,11 @@ bool UControlRig::Execute(const EControlRigState InState, const FName& InEventNa
 	
 	LatestExecutedState = InState;
 
+	FRigVMExtendedExecuteContext& ExtendedExecuteContext = GetVM()->GetContext();
+	FRigVMExecuteContext& PublicContext = ExtendedExecuteContext.GetPublicData<>();
+	PublicContext.SetDeltaTime(DeltaTime);
+	PublicContext.SetAbsoluteTime(AbsoluteTime);
+
 	if (VM)
 	{
 		if (VM->GetOuter() != this)
@@ -752,8 +757,8 @@ bool UControlRig::Execute(const EControlRigState InState, const FName& InEventNa
 			{
 				HierarchyController->LogFunction = [this](EMessageSeverity::Type InSeverity, const FString& Message)
 				{
-					const FRigVMExtendedExecuteContext& Context = GetVM()->GetContext();
-					const FRigVMExecuteContext& PublicContext = Context.GetPublicData<>(); 
+					const FRigVMExtendedExecuteContext& ExtendedExecuteContext = GetVM()->GetContext();
+					const FRigVMExecuteContext& PublicContext = ExtendedExecuteContext.GetPublicData<>(); 
 					if(ControlRigLog)
 					{
 						ControlRigLog->Report(InSeverity,PublicContext.FunctionName,PublicContext.InstructionIndex, Message);
diff --git a/Engine/Plugins/Animation/ControlRig/Source/ControlRig/Private/Units/Execution/RigUnit_Collection.cpp b/Engine/Plugins/Animation/ControlRig/Source/ControlRig/Private/Units/Execution/RigUnit_Collection.cpp
--- a/Engine/Plugins/Animation/ControlRig/Source/ControlRig/Private/Units/Execution/RigUnit_Collection.cpp
+++ b/Engine/Plugins/Animation/ControlRig/Source/ControlRig/Private/Units/Execution/RigUnit_Collection.cpp
@@ -492,18 +492,29 @@ FRigVMStructUpgradeInfo FRigUnit_CollectionItemAtIndex::GetUpgradeInfo() const
 FRigUnit_CollectionLoop_Execute()
 {
     DECLARE_SCOPE_HIERARCHICAL_COUNTER_RIGUNIT()
-    Count = Collection.Num();
-   	Continue = Collection.IsValidIndex(Index);
-	Ratio = GetRatioFromIndex(Index, Count);
+	
+	if(BlockToRun.IsNone())
+	{
+	    Count = Collection.Num();
+		Index = 0;
+		BlockToRun = ExecuteContextName;
+	}
+	else if(BlockToRun == ExecuteContextName)
+	{
+		Index++;
+	}
 
-	if(Continue)
+	if(Collection.IsValidIndex(Index))
 	{
 		Item = Collection[Index];
 	}
 	else
 	{
 		Item = FRigElementKey();
+		BlockToRun = ControlFlowCompletedName;
 	}
+	
+	Ratio = GetRatioFromIndex(Index, Count);
 }
 
 FRigVMStructUpgradeInfo FRigUnit_CollectionLoop::GetUpgradeInfo() const
diff --git a/Engine/Plugins/Animation/ControlRig/Source/ControlRig/Private/Units/Execution/RigUnit_Collection.h b/Engine/Plugins/Animation/ControlRig/Source/ControlRig/Private/Units/Execution/RigUnit_Collection.h
--- a/Engine/Plugins/Animation/ControlRig/Source/ControlRig/Private/Units/Execution/RigUnit_Collection.h
+++ b/Engine/Plugins/Animation/ControlRig/Source/ControlRig/Private/Units/Execution/RigUnit_Collection.h
@@ -603,19 +603,27 @@ struct CONTROLRIG_API FRigUnit_CollectionLoop : public FRigUnit_CollectionBaseMu
 
 	FRigUnit_CollectionLoop()
 	{
+		BlockToRun = NAME_None;
 		Count = 0;
 		Index = 0;
 		Ratio = 0.f;
-		Continue = false;
 	}
 
 	// FRigVMStruct overrides
-	FORCEINLINE virtual bool IsForLoop() const override { return true; }
+	virtual const TArray<FName>& GetControlFlowBlocks_Impl() const override
+	{
+		static const TArray<FName> Blocks = {ExecuteContextName, ForLoopCompletedPinName};
+		return Blocks;
+	}
+	virtual const bool IsControlFlowBlockSliced(const FName& InBlockName) const { return InBlockName == ExecuteContextName; }
 	FORCEINLINE virtual int32 GetNumSlices() const override { return Count; }
 
 	RIGVM_METHOD()
 	virtual void Execute(const FRigUnitContext& Context) override;
 
+	UPROPERTY(meta = (Singleton))
+	FName BlockToRun;
+
 	UPROPERTY(meta = (Input))
 	FRigElementKeyCollection Collection;
 
@@ -636,9 +644,6 @@ struct CONTROLRIG_API FRigUnit_CollectionLoop : public FRigUnit_CollectionBaseMu
 	UPROPERTY(meta = (Singleton, Output))
 	float Ratio;
 
-	UPROPERTY(meta = (Singleton))
-	bool Continue;
-
 	UPROPERTY(meta = (Output))
 	FControlRigExecuteContext Completed;
 
diff --git a/Engine/Plugins/Animation/ControlRig/Source/ControlRig/Private/Units/Execution/RigUnit_ForLoop.cpp b/Engine/Plugins/Animation/ControlRig/Source/ControlRig/Private/Units/Execution/RigUnit_ForLoop.cpp
--- a/Engine/Plugins/Animation/ControlRig/Source/ControlRig/Private/Units/Execution/RigUnit_ForLoop.cpp
+++ b/Engine/Plugins/Animation/ControlRig/Source/ControlRig/Private/Units/Execution/RigUnit_ForLoop.cpp
@@ -8,7 +8,22 @@
 FRigUnit_ForLoopCount_Execute()
 {
     DECLARE_SCOPE_HIERARCHICAL_COUNTER_RIGUNIT()
-   	Continue = Index < Count;
+
+	if(BlockToRun.IsNone())
+	{
+		Index = 0;
+		BlockToRun = ExecuteContextName;
+	}
+	else if(BlockToRun == ExecuteContextName)
+	{
+		Index++;
+	}
+
+	if(Index == Count)
+	{
+		BlockToRun = ControlFlowCompletedName;
+	}
+
 	Ratio = GetRatioFromIndex(Index, Count);
 }
 
diff --git a/Engine/Plugins/Animation/ControlRig/Source/ControlRig/Private/Units/Execution/RigUnit_ForLoop.h b/Engine/Plugins/Animation/ControlRig/Source/ControlRig/Private/Units/Execution/RigUnit_ForLoop.h
--- a/Engine/Plugins/Animation/ControlRig/Source/ControlRig/Private/Units/Execution/RigUnit_ForLoop.h
+++ b/Engine/Plugins/Animation/ControlRig/Source/ControlRig/Private/Units/Execution/RigUnit_ForLoop.h
@@ -16,19 +16,27 @@ struct CONTROLRIG_API FRigUnit_ForLoopCount : public FRigUnitMutable
 
 	FRigUnit_ForLoopCount()
 	{
+		BlockToRun = NAME_None;
 		Count = 1;
 		Index = 0;
 		Ratio = 0.f;
-		Continue = false;
 	}
 
 	// FRigVMStruct overrides
-	FORCEINLINE virtual bool IsForLoop() const override { return true; }
+	virtual const TArray<FName>& GetControlFlowBlocks_Impl() const override
+	{
+		static const TArray<FName> Blocks = {ExecuteContextName, ForLoopCompletedPinName};
+		return Blocks;
+	}
+	virtual const bool IsControlFlowBlockSliced(const FName& InBlockName) const { return InBlockName == ExecuteContextName; }
 	FORCEINLINE virtual int32 GetNumSlices() const override { return Count; }
 
 	RIGVM_METHOD()
 	virtual void Execute(const FRigUnitContext& Context) override;
 
+	UPROPERTY(meta = (Singleton))
+	FName BlockToRun;
+
 	UPROPERTY(meta = (Singleton, Input))
 	int32 Count;
 
@@ -38,9 +46,6 @@ struct CONTROLRIG_API FRigUnit_ForLoopCount : public FRigUnitMutable
 	UPROPERTY(meta = (Singleton, Output))
 	float Ratio;
 
-	UPROPERTY(meta = (Singleton))
-	bool Continue;
-
 	UPROPERTY(EditAnywhere, Transient, Category = "ForLoop", meta = (Output))
 	FControlRigExecuteContext Completed;
 };
diff --git a/Engine/Plugins/Animation/ControlRig/Source/ControlRig/Private/Units/Execution/RigUnit_Hierarchy.cpp b/Engine/Plugins/Animation/ControlRig/Source/ControlRig/Private/Units/Execution/RigUnit_Hierarchy.cpp
--- a/Engine/Plugins/Animation/ControlRig/Source/ControlRig/Private/Units/Execution/RigUnit_Hierarchy.cpp
+++ b/Engine/Plugins/Animation/ControlRig/Source/ControlRig/Private/Units/Execution/RigUnit_Hierarchy.cpp
@@ -581,23 +581,34 @@ FRigUnit_PoseGetCurve_Execute()
 FRigUnit_PoseLoop_Execute()
 {
 	DECLARE_SCOPE_HIERARCHICAL_COUNTER_RIGUNIT()
-	Count = Pose.Num();
-	Continue = Pose.IsValidIndex(Index);
-	Ratio = GetRatioFromIndex(Index, Count);
 
-	if(Continue)
+	if(BlockToRun.IsNone())
 	{
-		const FRigPoseElement& PoseElement = Pose.Elements[Index];
-		Item = PoseElement.Index.GetKey();
-		GlobalTransform = PoseElement.GlobalTransform;
-		LocalTransform = PoseElement.LocalTransform;
-		CurveValue = PoseElement.CurveValue;
+		Count = Pose.Num();
+		Index = 0;
+		BlockToRun = ExecuteContextName;
 	}
-	else
+	else if(BlockToRun == ExecuteContextName)
+	{
+		Index++;
+	}
+
+	if(Index == Count)
 	{
 		Item = FRigElementKey();
 		GlobalTransform = LocalTransform = FTransform::Identity;
 		CurveValue = 0.f;
+		BlockToRun = ControlFlowCompletedName;
 	}
+	else
+	{
+		const FRigPoseElement& PoseElement = Pose.Elements[Index];
+		Item = PoseElement.Index.GetKey();
+		GlobalTransform = PoseElement.GlobalTransform;
+		LocalTransform = PoseElement.LocalTransform;
+		CurveValue = PoseElement.CurveValue;
+	}
+
+	Ratio = GetRatioFromIndex(Index, Count);
 }
 
diff --git a/Engine/Plugins/Animation/ControlRig/Source/ControlRig/Private/Units/Execution/RigUnit_Hierarchy.h b/Engine/Plugins/Animation/ControlRig/Source/ControlRig/Private/Units/Execution/RigUnit_Hierarchy.h
--- a/Engine/Plugins/Animation/ControlRig/Source/ControlRig/Private/Units/Execution/RigUnit_Hierarchy.h
+++ b/Engine/Plugins/Animation/ControlRig/Source/ControlRig/Private/Units/Execution/RigUnit_Hierarchy.h
@@ -701,6 +701,7 @@ struct CONTROLRIG_API FRigUnit_PoseLoop : public FRigUnit_HierarchyBaseMutable
 
 	FRigUnit_PoseLoop()
 	{
+		BlockToRun = NAME_None;
 		Pose = FRigPose();
 		Item = FRigElementKey();
 		GlobalTransform = LocalTransform = FTransform::Identity;
@@ -708,16 +709,23 @@ struct CONTROLRIG_API FRigUnit_PoseLoop : public FRigUnit_HierarchyBaseMutable
 		Count = 0;
 		Index = 0;
 		Ratio = 0.f;
-		Continue = false;
 	}
 
 	// FRigVMStruct overrides
-	FORCEINLINE virtual bool IsForLoop() const override { return true; }
+	virtual const TArray<FName>& GetControlFlowBlocks_Impl() const override
+	{
+		static const TArray<FName> Blocks = {ExecuteContextName, ForLoopCompletedPinName};
+		return Blocks;
+	}
+	virtual const bool IsControlFlowBlockSliced(const FName& InBlockName) const { return InBlockName == ExecuteContextName; }
 	FORCEINLINE virtual int32 GetNumSlices() const override { return Count; }
 
 	RIGVM_METHOD()
 	virtual void Execute(const FRigUnitContext& Context) override;
 
+	UPROPERTY(meta = (Singleton))
+	FName BlockToRun;
+
 	UPROPERTY(meta = (Input))
 	FRigPose Pose;
 
@@ -747,9 +755,6 @@ struct CONTROLRIG_API FRigUnit_PoseLoop : public FRigUnit_HierarchyBaseMutable
 	UPROPERTY(meta = (Singleton, Output))
 	float Ratio;
 
-	UPROPERTY(meta = (Singleton))
-	bool Continue;
-
 	UPROPERTY(meta = (Output))
 	FControlRigExecuteContext Completed;
 };
diff --git a/Engine/Plugins/Animation/ControlRig/Source/ControlRigEditor/Private/ControlRigEditorModule.cpp b/Engine/Plugins/Animation/ControlRig/Source/ControlRigEditor/Private/ControlRigEditorModule.cpp
--- a/Engine/Plugins/Animation/ControlRig/Source/ControlRigEditor/Private/ControlRigEditorModule.cpp
+++ b/Engine/Plugins/Animation/ControlRig/Source/ControlRigEditor/Private/ControlRigEditorModule.cpp
@@ -1095,6 +1095,11 @@ void FControlRigEditorModule::GetTypeActions(UControlRigBlueprint* CRB, FBluepri
 		Struct->GetStringMetaDataHierarchical(FRigVMStruct::CategoryMetaName, &CategoryMetadata);
 		Struct->GetStringMetaDataHierarchical(FRigVMStruct::DisplayNameMetaName, &DisplayNameMetadata);
 		Struct->GetStringMetaDataHierarchical(FRigVMStruct::MenuDescSuffixMetaName, &MenuDescSuffixMetadata);
+
+		if(DisplayNameMetadata.IsEmpty())
+		{
+			DisplayNameMetadata = Struct->GetDisplayNameText().ToString();
+		}
 		if (!MenuDescSuffixMetadata.IsEmpty())
 		{
 			MenuDescSuffixMetadata = TEXT(" ") + MenuDescSuffixMetadata;
diff --git a/Engine/Plugins/Animation/ControlRig/Source/ControlRigEditor/Private/Editor/SControlRigStackView.cpp b/Engine/Plugins/Animation/ControlRig/Source/ControlRigEditor/Private/Editor/SControlRigStackView.cpp
--- a/Engine/Plugins/Animation/ControlRig/Source/ControlRigEditor/Private/Editor/SControlRigStackView.cpp
+++ b/Engine/Plugins/Animation/ControlRig/Source/ControlRigEditor/Private/Editor/SControlRigStackView.cpp
@@ -628,6 +628,12 @@ void SControlRigStackView::PopulateStackView(URigVM* InVM)
 						Label = FString::Printf(TEXT("Run %s Event"), *Op.EntryName.ToString());
 						break;
 					}
+					case ERigVMOpCode::JumpToBranch:
+					{
+						const FRigVMJumpToBranchOp& Op = ByteCode.GetOpAt<FRigVMJumpToBranchOp>(Instructions[InstructionIndex]);
+						Label = TEXT("Jump To Branch");
+						break;
+					}
 					case ERigVMOpCode::Exit:
 					{
 						Label = TEXT("Exit");
diff --git a/Engine/Source/Developer/RigVMDeveloper/Private/RigVMCompiler/RigVMAST.cpp b/Engine/Source/Developer/RigVMDeveloper/Private/RigVMCompiler/RigVMAST.cpp
--- a/Engine/Source/Developer/RigVMDeveloper/Private/RigVMCompiler/RigVMAST.cpp
+++ b/Engine/Source/Developer/RigVMDeveloper/Private/RigVMCompiler/RigVMAST.cpp
@@ -600,17 +600,14 @@ bool FRigVMVarExprAST::SupportsSoftLinks() const
 {
 	if (URigVMUnitNode* UnitNode = Cast<URigVMUnitNode>(GetPin()->GetNode()))
 	{
-		if (UnitNode->IsLoopNode())
+		if (UnitNode->IsControlFlowNode())
 		{
-			if (GetPin()->GetFName() != FRigVMStruct::ExecuteContextName &&
-				GetPin()->GetFName() != FRigVMStruct::ForLoopCompletedPinName)
-			{
-				return true;
-			}
+			return !UnitNode->GetControlFlowBlocks().Contains(GetPin()->GetFName());
 		}
 	}
 	else if(URigVMArrayNode* ArrayNode = Cast<URigVMArrayNode>(GetPin()->GetNode()))
 	{
+		// this is still using is loop node - but that's to be deprecated soon
 		if(ArrayNode->IsLoopNode())
 		{
 			if (GetPin()->GetFName() != FRigVMStruct::ExecuteContextName &&
@@ -971,7 +968,7 @@ FRigVMExprAST* FRigVMParserAST::TraverseMutableNode(const FRigVMASTProxy& InNode
 					FRigVMASTProxy SourcePinProxy = InNodeProxy.GetSibling(SourcePin);
 
 					FRigVMExprAST* ParentExpr = InParentExpr;
-					if (NodeExpr->IsA(FRigVMExprAST::Branch) || Node->IsLoopNode())
+					if (NodeExpr->IsA(FRigVMExprAST::Branch) || Node->IsControlFlowNode())
 					{
 						if (FRigVMExprAST** PinExpr = SubjectToExpression.Find(SourcePinProxy))
 						{
@@ -1852,7 +1849,7 @@ void FRigVMParserAST::FoldAssignments()
 		// if this node is a loop node - let's skip the folding
 		if (URigVMUnitNode* UnitNode = Cast<URigVMUnitNode>(TargetPin->GetNode()))
 		{
-			if (UnitNode->IsLoopNode())
+			if (UnitNode->IsControlFlowNode())
 			{
 				continue;
 			}
@@ -1870,6 +1867,7 @@ void FRigVMParserAST::FoldAssignments()
 		// if this node is an array iterator node - let's skip the folding
 		if (URigVMArrayNode* ArrayNode = Cast<URigVMArrayNode>(TargetPin->GetNode()))
 		{
+			// still using IsLoopNode - but to be deprecated soon 
 			if (ArrayNode->IsLoopNode())
 			{
 				continue;
diff --git a/Engine/Source/Developer/RigVMDeveloper/Private/RigVMCompiler/RigVMCompiler.cpp b/Engine/Source/Developer/RigVMDeveloper/Private/RigVMCompiler/RigVMCompiler.cpp
--- a/Engine/Source/Developer/RigVMDeveloper/Private/RigVMCompiler/RigVMCompiler.cpp
+++ b/Engine/Source/Developer/RigVMDeveloper/Private/RigVMCompiler/RigVMCompiler.cpp
@@ -1019,15 +1019,6 @@ void URigVMCompiler::TraverseExpression(const FRigVMExprAST* InExpr, FRigVMCompi
 		case FRigVMExprAST::EType::CallExtern:
 		{
 			const FRigVMCallExternExprAST* CallExternExpr = InExpr->To<FRigVMCallExternExprAST>();
-			if (URigVMUnitNode* UnitNode = Cast<URigVMUnitNode>(CallExternExpr->GetNode()))
-			{
-				if (UnitNode->IsLoopNode())
-				{
-					TraverseForLoop(CallExternExpr, WorkData);
-					break;
-				}
-			}
-
 			TraverseCallExtern(CallExternExpr, WorkData);
 			break;
 		}
@@ -1374,7 +1365,8 @@ int32 URigVMCompiler::TraverseCallExtern(const FRigVMCallExternExprAST* InExpr,
 		}
 	}
 
-	int32 InstructionIndex = INDEX_NONE;
+	int32 CallExternInstructionIndex = INDEX_NONE;
+	const FRigVMCallstack Callstack = InExpr->GetProxy().GetCallstack();
 
 	if (WorkData.bSetupMemory)
 	{
@@ -1397,7 +1389,10 @@ int32 URigVMCompiler::TraverseCallExtern(const FRigVMCallExternExprAST* InExpr,
 		}
 			
 		check(Function);
-		
+
+		FRigVMOperand CountOperand;
+		FRigVMOperand IndexOperand;
+		FRigVMOperand BlockToRunOperand;
 		for(const FRigVMFunctionArgument& Argument : Function->GetArguments())
 		{
 			const FRigVMExprAST* ChildExpr = InExpr->FindExprWithPinName(Argument.Name);
@@ -1419,6 +1414,36 @@ int32 URigVMCompiler::TraverseCallExtern(const FRigVMCallExternExprAST* InExpr,
 			{
 				break;
 			}
+
+			if(Argument.Name == FRigVMStruct::ControlFlowBlockToRunName)
+			{
+				BlockToRunOperand = Operands.Last();
+			}
+			else if(Argument.Name == FRigVMStruct::ControlFlowCountName)
+			{
+				CountOperand = Operands.Last();
+			}
+			else if(Argument.Name == FRigVMStruct::ControlFlowIndexName)
+			{
+				IndexOperand = Operands.Last();
+			}
+		}
+
+		// make sure to skip the output blocks while we are traversing this call extern
+		TArray<const FRigVMExprAST*> ExpressionsToSkip;
+		TArray<int32> BranchIndices;
+		if(Node->IsControlFlowNode())
+		{
+			const TArray<FName>& BlockNames = Node->GetControlFlowBlocks();
+			BranchIndices.Reserve(BlockNames.Num());
+			
+			for(const FName& BlockName : BlockNames)
+			{
+				const FRigVMVarExprAST* BlockExpr = InExpr->FindVarWithPinName(BlockName);
+				check(BlockExpr);
+				WorkData.ExprToSkip.AddUnique(BlockExpr);
+				BranchIndices.Add(WorkData.VM->GetByteCode().AddBranchInfo(FRigVMBranchInfo()));
+			}
 		}
 
 		// traverse all non-lazy children
@@ -1471,11 +1496,17 @@ int32 URigVMCompiler::TraverseCallExtern(const FRigVMCallExternExprAST* InExpr,
 			WorkData.VM->GetByteCode().GetOpAt<FRigVMJumpOp>(JumpToCallExternByte).InstructionIndex = InstructionsToJump;
 		}
 
+		if(Node->IsControlFlowNode())
+		{
+			check(BlockToRunOperand.IsValid());
+			WorkData.VM->GetByteCode().AddZeroOp(BlockToRunOperand);
+		}
+
 		// setup the instruction
 		const int32 FunctionIndex = WorkData.VM->AddRigVMFunction(Function->GetName());
 		check(FunctionIndex != INDEX_NONE);
 		WorkData.VM->GetByteCode().AddExecuteOp(FunctionIndex, Operands);
-		InstructionIndex = WorkData.VM->GetByteCode().GetNumInstructions() - 1;
+		CallExternInstructionIndex = WorkData.VM->GetByteCode().GetNumInstructions() - 1;
 
 		// setup the branch infos for this call extern instruction
 		if(const TArray<FRigVMBranchInfo>* BranchInfosPtr = WorkData.BranchInfos.Find(Node))
@@ -1483,8 +1514,8 @@ int32 URigVMCompiler::TraverseCallExtern(const FRigVMCallExternExprAST* InExpr,
 			const TArray<FRigVMBranchInfo>& BranchInfos = *BranchInfosPtr;
 			for(FRigVMBranchInfo BranchInfo : BranchInfos)
 			{
-				BranchInfo.InstructionIndex = InstructionIndex;
-				WorkData.VM->GetByteCode().AddBranchInfo(BranchInfo);
+				BranchInfo.InstructionIndex = CallExternInstructionIndex;
+				(void)WorkData.VM->GetByteCode().AddBranchInfo(BranchInfo);
 			}
 		}
 
@@ -1522,7 +1553,7 @@ int32 URigVMCompiler::TraverseCallExtern(const FRigVMCallExternExprAST* InExpr,
 		}
 
 		WorkData.VM->GetByteCode().SetOperandsForInstruction(
-			InstructionIndex,
+			CallExternInstructionIndex,
 			FRigVMOperandArray(InputsOperands.GetData(), InputsOperands.Num()),
 			FRigVMOperandArray(OutputOperands.GetData(), OutputOperands.Num()));
 
@@ -1530,12 +1561,76 @@ int32 URigVMCompiler::TraverseCallExtern(const FRigVMCallExternExprAST* InExpr,
 		
 		if (Settings.SetupNodeInstructionIndex)
 		{
-			const FRigVMCallstack Callstack = InExpr->GetProxy().GetCallstack();
-			WorkData.VM->GetByteCode().SetSubject(InstructionIndex, Callstack.GetCallPath(), Callstack.GetStack());
+			WorkData.VM->GetByteCode().SetSubject(CallExternInstructionIndex, Callstack.GetCallPath(), Callstack.GetStack());
+		}
+
+		if(Node->IsControlFlowNode())
+		{
+			// add an operator to jump to the right branch
+			const int32 JumpToBranchInstructionIndex = WorkData.VM->GetByteCode().GetNumInstructions();
+
+			// use the index of the first branch info relating to this control flow node.
+			// branches are stored on the bytecode in order for each control flow node - so the
+			// VM needs to know which branch to start to look at then evaluating the JumpToBranchOp.
+			// Branches are stored in order - similar to this example representing two JumpBranchOps
+			// with BranchIndices [0, 1] and [2, 3]
+			// [
+			//    0 = {ExecuteContext, InstructionIndex 2, First 3, Last 5},
+			//    1 = {Completed, InstructionIndex 2, First 6, Last 12},
+			//    2 = {ExecuteContext, InstructionIndex 17, First 18, Last 21},
+			//    3 = {Completed, InstructionIndex 17, First 22, Last 28},
+			// ]
+			// The first index of the branch in the overall list of branches is stored in the operator (BranchIndices[0])
+			WorkData.VM->GetByteCode().AddJumpToBranchOp(BlockToRunOperand, BranchIndices[0]);
+
+			const TArray<FName>& BlockNames = Node->GetControlFlowBlocks();
+
+			// traverse all of the blocks now
+			for(int32 BlockIndex = 0; BlockIndex < BlockNames.Num(); BlockIndex++)
+			{
+				const FName& BlockName = BlockNames[BlockIndex];
+
+				FRigVMBranchInfo& BranchInfo = WorkData.VM->GetByteCode().BranchInfos[BranchIndices[BlockIndex]];
+				BranchInfo.Label = BlockName;
+				BranchInfo.InstructionIndex = JumpToBranchInstructionIndex;
+				BranchInfo.FirstInstruction = WorkData.VM->GetByteCode().GetNumInstructions();
+
+				// check if the block requires slicing or not.
+				// (do we want the private state of the nodes to be unique per run of the block)
+				if(Node->IsControlFlowBlockSliced(BlockName))
+				{
+					check(BlockName != FRigVMStruct::ControlFlowCompletedName);
+					check(CountOperand.IsValid());
+					check(IndexOperand.IsValid());
+					
+					WorkData.VM->GetByteCode().AddBeginBlockOp(CountOperand, IndexOperand);
+				}
+
+				// traverse the body of the block
+				const FRigVMVarExprAST* BlockExpr = InExpr->FindVarWithPinName(BlockName);
+				check(BlockExpr);
+				WorkData.ExprToSkip.Remove(BlockExpr);
+				TraverseExpression(BlockExpr, WorkData);
+
+				// end the block if necessary
+				if(Node->IsControlFlowBlockSliced(BlockName))
+				{
+					WorkData.VM->GetByteCode().AddEndBlockOp();
+				}
+
+				// if this is not the completed block - we need to jump back to the control flow instruction
+				if(BlockName != FRigVMStruct::ControlFlowCompletedName)
+				{
+					const int32 JumpToCallExternInstruction = WorkData.VM->GetByteCode().GetNumInstructions();
+					WorkData.VM->GetByteCode().AddJumpOp(ERigVMOpCode::JumpBackward, JumpToCallExternInstruction - CallExternInstructionIndex);
+				}
+
+				BranchInfo.LastInstruction = WorkData.VM->GetByteCode().GetNumInstructions() - 1;
+			}
 		}
 	}
 
-	return InstructionIndex;
+	return CallExternInstructionIndex;
 }
 
 int32 URigVMCompiler::TraverseInlineFunction(const FRigVMInlineFunctionExprAST* InExpr, FRigVMCompilerWorkData& WorkData)
@@ -1813,103 +1908,6 @@ int32 URigVMCompiler::TraverseInlineFunction(const FRigVMInlineFunctionExprAST*
 	return InstructionIndexEnd;
 }
 
-void URigVMCompiler::TraverseForLoop(const FRigVMCallExternExprAST* InExpr, FRigVMCompilerWorkData& WorkData)
-{
-	if (WorkData.bSetupMemory)
-	{
-		TraverseCallExtern(InExpr, WorkData);
-		return;
-	}
-
-	URigVMUnitNode* UnitNode = Cast<URigVMUnitNode>(InExpr->GetNode());
-	if(!ValidateNode(UnitNode))
-	{
-		return;
-	}
-
-	const FRigVMCallstack Callstack = InExpr->GetProxy().GetCallstack();
-
-	const FRigVMVarExprAST* CompletedExpr = InExpr->FindVarWithPinName(FRigVMStruct::ForLoopCompletedPinName);
-	check(CompletedExpr);
-	const FRigVMVarExprAST* ExecuteExpr = InExpr->FindVarWithPinName(FRigVMStruct::ExecuteContextName);
-	check(ExecuteExpr);
-	WorkData.ExprToSkip.AddUnique(CompletedExpr);
-	WorkData.ExprToSkip.AddUnique(ExecuteExpr);
-
-	// set the index to 0
-	const FRigVMVarExprAST* IndexExpr = InExpr->FindVarWithPinName(FRigVMStruct::ForLoopIndexPinName);
-	check(IndexExpr);
-	FRigVMOperand IndexOperand = WorkData.ExprToOperand.FindChecked(GetSourceVarExpr(IndexExpr));
-	WorkData.VM->GetByteCode().AddZeroOp(IndexOperand);
-	if (Settings.SetupNodeInstructionIndex)
-	{
-		WorkData.VM->GetByteCode().SetSubject(WorkData.VM->GetByteCode().GetNumInstructions() - 1, Callstack.GetCallPath(), Callstack.GetStack());
-	}
-
-	// call the for loop compute
-	int32 ForLoopInstructionIndex = TraverseCallExtern(InExpr, WorkData);
-	if (Settings.SetupNodeInstructionIndex)
-	{
-		WorkData.VM->GetByteCode().SetSubject(ForLoopInstructionIndex, Callstack.GetCallPath(), Callstack.GetStack());
-	}
-
-	// set up the jump forward (jump out of the loop)
-	const FRigVMVarExprAST* ContinueLoopExpr = InExpr->FindVarWithPinName(FRigVMStruct::ForLoopContinuePinName);
-	check(ContinueLoopExpr);
-	FRigVMOperand ContinueLoopOperand = WorkData.ExprToOperand.FindChecked(GetSourceVarExpr(ContinueLoopExpr));
-
-	uint64 JumpToEndByte = WorkData.VM->GetByteCode().AddJumpIfOp(ERigVMOpCode::JumpForwardIf, 0, ContinueLoopOperand, false);
-	int32 JumpToEndInstruction = WorkData.VM->GetByteCode().GetNumInstructions() - 1;
-	if (Settings.SetupNodeInstructionIndex)
-	{
-		WorkData.VM->GetByteCode().SetSubject(WorkData.VM->GetByteCode().GetNumInstructions() - 1, Callstack.GetCallPath(), Callstack.GetStack());
-	}
-
-	// begin the loop's block
-	const FRigVMVarExprAST* CountExpr = InExpr->FindVarWithPinName(FRigVMStruct::ForLoopCountPinName);
-	check(CountExpr);
-	FRigVMOperand CountOperand = WorkData.ExprToOperand.FindChecked(GetSourceVarExpr(CountExpr));
-	WorkData.VM->GetByteCode().AddBeginBlockOp(CountOperand, IndexOperand);
-	if (Settings.SetupNodeInstructionIndex)
-	{
-		WorkData.VM->GetByteCode().SetSubject(WorkData.VM->GetByteCode().GetNumInstructions() - 1, Callstack.GetCallPath(), Callstack.GetStack());
-	}
-
-	// traverse the body of the loop
-	WorkData.ExprToSkip.Remove(ExecuteExpr);
-	TraverseExpression(ExecuteExpr, WorkData);
-
-	// end the loop's block
-	WorkData.VM->GetByteCode().AddEndBlockOp();
-	if (Settings.SetupNodeInstructionIndex)
-	{
-		WorkData.VM->GetByteCode().SetSubject(WorkData.VM->GetByteCode().GetNumInstructions() - 1, Callstack.GetCallPath(), Callstack.GetStack());
-	}
-
-	// increment the index
-	WorkData.VM->GetByteCode().AddIncrementOp(IndexOperand);
-	if (Settings.SetupNodeInstructionIndex)
-	{
-		WorkData.VM->GetByteCode().SetSubject(WorkData.VM->GetByteCode().GetNumInstructions() - 1, Callstack.GetCallPath(), Callstack.GetStack());
-	}
-
-	// jump to the beginning of the loop
-	int32 JumpToStartInstruction = WorkData.VM->GetByteCode().GetNumInstructions();
-	WorkData.VM->GetByteCode().AddJumpOp(ERigVMOpCode::JumpBackward, JumpToStartInstruction - ForLoopInstructionIndex);
-	if (Settings.SetupNodeInstructionIndex)
-	{
-		WorkData.VM->GetByteCode().SetSubject(WorkData.VM->GetByteCode().GetNumInstructions() - 1, Callstack.GetCallPath(), Callstack.GetStack());
-	}
-
-	// update the jump operator with the right address
-	int32 InstructionsToEnd = WorkData.VM->GetByteCode().GetNumInstructions() - JumpToEndInstruction;
-	WorkData.VM->GetByteCode().GetOpAt<FRigVMJumpIfOp>(JumpToEndByte).InstructionIndex = InstructionsToEnd;
-
-	// now traverse everything else connected to the completed pin
-	WorkData.ExprToSkip.Remove(CompletedExpr);
-	TraverseExpression(CompletedExpr, WorkData);
-}
-
 void URigVMCompiler::TraverseNoOp(const FRigVMNoOpExprAST* InExpr, FRigVMCompilerWorkData& WorkData)
 {
 	TraverseChildren(InExpr, WorkData);
@@ -2773,7 +2771,7 @@ void URigVMCompiler::AddCopyOperator(const FRigVMCopyOp& InOp, const FRigVMAssig
 		bDelayCopyOperations = false;
 	}
 
-	// loop up a potentially delayed copy operation which needs to happen
+	// look up a potentially delayed copy operation which needs to happen
 	// just prior to this one and inject it as well.
 	if(!bDelayCopyOperations)
 	{
diff --git a/Engine/Source/Developer/RigVMDeveloper/Private/RigVMModel/Nodes/RigVMArrayNode.cpp b/Engine/Source/Developer/RigVMDeveloper/Private/RigVMModel/Nodes/RigVMArrayNode.cpp
--- a/Engine/Source/Developer/RigVMDeveloper/Private/RigVMModel/Nodes/RigVMArrayNode.cpp
+++ b/Engine/Source/Developer/RigVMDeveloper/Private/RigVMModel/Nodes/RigVMArrayNode.cpp
@@ -549,9 +549,14 @@ const FRigVMTemplate* URigVMArrayNode::GetTemplate() const
 	return CachedTemplate;
 }
 
-bool URigVMArrayNode::IsLoopNode() const
+const TArray<FName>& URigVMArrayNode::GetControlFlowBlocks() const
 {
-	return GetOpCode() == ERigVMOpCode::ArrayIterator;
+	if(GetOpCode() == ERigVMOpCode::ArrayIterator)
+	{
+		static const TArray<FName> LoopBlocks = {FRigVMStruct::ExecuteContextName, FRigVMStruct::ForLoopCompletedPinName};
+		return LoopBlocks;
+	}
+	return Super::GetControlFlowBlocks();
 }
 
 FText URigVMArrayNode::GetToolTipTextForPin(const URigVMPin* InPin) const
diff --git a/Engine/Source/Developer/RigVMDeveloper/Private/RigVMModel/Nodes/RigVMDispatchNode.cpp b/Engine/Source/Developer/RigVMDeveloper/Private/RigVMModel/Nodes/RigVMDispatchNode.cpp
--- a/Engine/Source/Developer/RigVMDeveloper/Private/RigVMModel/Nodes/RigVMDispatchNode.cpp
+++ b/Engine/Source/Developer/RigVMDeveloper/Private/RigVMModel/Nodes/RigVMDispatchNode.cpp
@@ -50,6 +50,24 @@ bool URigVMDispatchNode::IsDefinedAsVarying() const
 	return false;
 }
 
+const TArray<FName>& URigVMDispatchNode::GetControlFlowBlocks() const
+{
+	if (const FRigVMDispatchFactory* Factory = GetFactory())
+	{
+		return Factory->GetControlFlowBlocks();
+	}
+	return Super::GetControlFlowBlocks();
+}
+
+const bool URigVMDispatchNode::IsControlFlowBlockSliced(const FName& InBlockName) const
+{
+	if (const FRigVMDispatchFactory* Factory = GetFactory())
+	{
+		return Factory->IsControlFlowBlockSliced(InBlockName);
+	}
+	return Super::IsControlFlowBlockSliced(InBlockName);
+}
+
 FText URigVMDispatchNode::GetToolTipTextForPin(const URigVMPin* InPin) const
 {
 	if (const FRigVMDispatchFactory* Factory = GetFactory())
diff --git a/Engine/Source/Developer/RigVMDeveloper/Private/RigVMModel/Nodes/RigVMUnitNode.cpp b/Engine/Source/Developer/RigVMDeveloper/Private/RigVMModel/Nodes/RigVMUnitNode.cpp
--- a/Engine/Source/Developer/RigVMDeveloper/Private/RigVMModel/Nodes/RigVMUnitNode.cpp
+++ b/Engine/Source/Developer/RigVMDeveloper/Private/RigVMModel/Nodes/RigVMUnitNode.cpp
@@ -97,6 +97,28 @@ bool URigVMUnitNode::CanOnlyExistOnce() const
 	return Super::CanOnlyExistOnce();
 }
 
+const TArray<FName>& URigVMUnitNode::GetControlFlowBlocks() const
+{
+	const TSharedPtr<FStructOnScope> StructOnScope = ConstructStructInstance(true);
+	if (StructOnScope.IsValid())
+	{
+		const FRigVMStruct* StructMemory = (FRigVMStruct*)StructOnScope->GetStructMemory();
+		return StructMemory->GetControlFlowBlocks();
+	}
+	return Super::GetControlFlowBlocks();
+}
+
+const bool URigVMUnitNode::IsControlFlowBlockSliced(const FName& InBlockName) const
+{
+	const TSharedPtr<FStructOnScope> StructOnScope = ConstructStructInstance(true);
+	if (StructOnScope.IsValid())
+	{
+		const FRigVMStruct* StructMemory = (FRigVMStruct*)StructOnScope->GetStructMemory();
+		return StructMemory->IsControlFlowBlockSliced(InBlockName);
+	}
+	return Super::IsControlFlowBlockSliced(InBlockName);
+}
+
 FText URigVMUnitNode::GetToolTipTextForPin(const URigVMPin* InPin) const
 {
 	if(UScriptStruct* Struct = GetScriptStruct())
@@ -278,17 +300,6 @@ UScriptStruct* URigVMUnitNode::GetScriptStruct() const
 	return ScriptStruct_DEPRECATED;
 }
 
-bool URigVMUnitNode::IsLoopNode() const
-{
-	const TSharedPtr<FStructOnScope> StructOnScope = ConstructStructInstance(true);
-	if (StructOnScope.IsValid())
-	{
-		const FRigVMStruct* StructMemory = (FRigVMStruct*)StructOnScope->GetStructMemory();
-		return StructMemory->IsForLoop();
-	}
-	return false;
-}
-
 FName URigVMUnitNode::GetMethodName() const
 {
 	const FName ResolvedMethodName = Super::GetMethodName();
diff --git a/Engine/Source/Developer/RigVMDeveloper/Private/RigVMModel/RigVMNode.cpp b/Engine/Source/Developer/RigVMDeveloper/Private/RigVMModel/RigVMNode.cpp
--- a/Engine/Source/Developer/RigVMDeveloper/Private/RigVMModel/RigVMNode.cpp
+++ b/Engine/Source/Developer/RigVMDeveloper/Private/RigVMModel/RigVMNode.cpp
@@ -569,6 +569,37 @@ double URigVMNode::GetInstructionMicroSeconds(URigVM* InVM, const FRigVMASTProxy
 	return -1.0;
 }
 
+bool URigVMNode::IsLoopNode() const
+{
+	if(IsControlFlowNode())
+	{
+		static const TArray<FName> ExpectedLoopBlocks = {FRigVMStruct::ExecuteContextName, FRigVMStruct::ForLoopCompletedPinName};
+		const TArray<FName>& Blocks = GetControlFlowBlocks();
+		if(Blocks.Num() == ExpectedLoopBlocks.Num())
+		{
+			return Blocks[0] == ExpectedLoopBlocks[0] && Blocks[1] == ExpectedLoopBlocks[1];
+		}
+	}
+
+	return false;
+}
+
+bool URigVMNode::IsControlFlowNode() const
+{
+	return !GetControlFlowBlocks().IsEmpty();
+}
+
+const TArray<FName>& URigVMNode::GetControlFlowBlocks() const
+{
+	static const TArray<FName> EmptyArray;
+	return EmptyArray;
+}
+
+const bool URigVMNode::IsControlFlowBlockSliced(const FName& InBlockName) const
+{
+	return false;
+}
+
 TArray<FRigVMUserWorkflow> URigVMNode::GetSupportedWorkflows(ERigVMUserWorkflowType InType, const UObject* InSubject) const
 {
 	if(InSubject == nullptr)
diff --git a/Engine/Source/Developer/RigVMDeveloper/Public/RigVMCompiler/RigVMCompiler.h b/Engine/Source/Developer/RigVMDeveloper/Public/RigVMCompiler/RigVMCompiler.h
--- a/Engine/Source/Developer/RigVMDeveloper/Public/RigVMCompiler/RigVMCompiler.h
+++ b/Engine/Source/Developer/RigVMDeveloper/Public/RigVMCompiler/RigVMCompiler.h
@@ -203,7 +203,6 @@ private:
 	void TraverseEntry(const FRigVMEntryExprAST* InExpr, FRigVMCompilerWorkData& WorkData);
 	int32 TraverseCallExtern(const FRigVMCallExternExprAST* InExpr, FRigVMCompilerWorkData& WorkData);
 	int32 TraverseInlineFunction(const FRigVMInlineFunctionExprAST* InExpr, FRigVMCompilerWorkData& WorkData);
-	void TraverseForLoop(const FRigVMCallExternExprAST* InExpr, FRigVMCompilerWorkData& WorkData);
 	void TraverseNoOp(const FRigVMNoOpExprAST* InExpr, FRigVMCompilerWorkData& WorkData);
 	void TraverseVar(const FRigVMVarExprAST* InExpr, FRigVMCompilerWorkData& WorkData);
 	void TraverseLiteral(const FRigVMVarExprAST* InExpr, FRigVMCompilerWorkData& WorkData);
diff --git a/Engine/Source/Developer/RigVMDeveloper/Public/RigVMModel/Nodes/RigVMArrayNode.h b/Engine/Source/Developer/RigVMDeveloper/Public/RigVMModel/Nodes/RigVMArrayNode.h
--- a/Engine/Source/Developer/RigVMDeveloper/Public/RigVMModel/Nodes/RigVMArrayNode.h
+++ b/Engine/Source/Developer/RigVMDeveloper/Public/RigVMModel/Nodes/RigVMArrayNode.h
@@ -56,8 +56,8 @@ public:
 	// Override of node color
 	virtual bool IsDefinedAsVarying() const override { return true; }
 
-	// Override loop node
-	virtual bool IsLoopNode() const override;
+	// Override control flow
+	virtual const TArray<FName>& GetControlFlowBlocks() const override;
 
 protected:
 
diff --git a/Engine/Source/Developer/RigVMDeveloper/Public/RigVMModel/Nodes/RigVMDispatchNode.h b/Engine/Source/Developer/RigVMDeveloper/Public/RigVMModel/Nodes/RigVMDispatchNode.h
--- a/Engine/Source/Developer/RigVMDeveloper/Public/RigVMModel/Nodes/RigVMDispatchNode.h
+++ b/Engine/Source/Developer/RigVMDeveloper/Public/RigVMModel/Nodes/RigVMDispatchNode.h
@@ -39,6 +39,8 @@ public:
 	virtual FLinearColor GetNodeColor() const override;
 	virtual bool IsDefinedAsConstant() const override;
 	virtual bool IsDefinedAsVarying() const override;
+	virtual const TArray<FName>& GetControlFlowBlocks() const override;
+	virtual const bool IsControlFlowBlockSliced(const FName& InBlockName) const override;
 	virtual TArray<URigVMPin*> GetAggregateInputs() const override;
 	virtual TArray<URigVMPin*> GetAggregateOutputs() const override;
 	virtual FName GetNextAggregateName(const FName& InLastAggregatePinName) const override;
diff --git a/Engine/Source/Developer/RigVMDeveloper/Public/RigVMModel/Nodes/RigVMUnitNode.h b/Engine/Source/Developer/RigVMDeveloper/Public/RigVMModel/Nodes/RigVMUnitNode.h
--- a/Engine/Source/Developer/RigVMDeveloper/Public/RigVMModel/Nodes/RigVMUnitNode.h
+++ b/Engine/Source/Developer/RigVMDeveloper/Public/RigVMModel/Nodes/RigVMUnitNode.h
@@ -29,7 +29,8 @@ public:
 	virtual bool IsDefinedAsVarying() const override;
 	virtual FName GetEventName() const override;
 	virtual bool CanOnlyExistOnce() const override;
-	virtual bool IsLoopNode() const override;
+	virtual const TArray<FName>& GetControlFlowBlocks() const override;
+	virtual const bool IsControlFlowBlockSliced(const FName& InBlockName) const override;
 	virtual TArray<FRigVMUserWorkflow> GetSupportedWorkflows(ERigVMUserWorkflowType InType, const UObject* InSubject) const override;
 	virtual TArray<URigVMPin*> GetAggregateInputs() const override;
 	virtual TArray<URigVMPin*> GetAggregateOutputs() const override;
diff --git a/Engine/Source/Developer/RigVMDeveloper/Public/RigVMModel/RigVMNode.h b/Engine/Source/Developer/RigVMDeveloper/Public/RigVMModel/RigVMNode.h
--- a/Engine/Source/Developer/RigVMDeveloper/Public/RigVMModel/RigVMNode.h
+++ b/Engine/Source/Developer/RigVMDeveloper/Public/RigVMModel/RigVMNode.h
@@ -212,7 +212,17 @@ public:
 
 	// return true if this node is a loop node
 	UFUNCTION(BlueprintPure, Category = RigVMNode)
-	virtual bool IsLoopNode() const { return false; }
+	bool IsLoopNode() const;
+
+	// return true if this node is a control flow node
+	UFUNCTION(BlueprintPure, Category = RigVMNode)
+	bool IsControlFlowNode() const;
+
+	// returns the names of the control flow blocks of this node
+	virtual const TArray<FName>& GetControlFlowBlocks() const;
+
+	// returns true if a control flow block requires slicing
+	virtual const bool IsControlFlowBlockSliced(const FName& InBlockName) const;
 
 	// returns true if the node can be upgraded
 	UFUNCTION(BlueprintPure, Category = RigVMNode)
diff --git a/Engine/Source/Runtime/RigVM/Private/RigVMCore/RigVM.cpp b/Engine/Source/Runtime/RigVM/Private/RigVMCore/RigVM.cpp
--- a/Engine/Source/Runtime/RigVM/Private/RigVMCore/RigVM.cpp
+++ b/Engine/Source/Runtime/RigVM/Private/RigVMCore/RigVM.cpp
@@ -541,6 +541,12 @@ bool URigVM::ValidateAllOperandsDuringLoad()
 				CheckOperandValidity(Op.ArgF);
 				break;
 			}
+			case ERigVMOpCode::JumpToBranch:
+			{
+				const FRigVMJumpToBranchOp& Op = ByteCodeStorage.GetOpAt<FRigVMJumpToBranchOp>(ByteCodeInstruction);
+				CheckOperandValidity(Op.Arg);
+				break;
+			}
 			case ERigVMOpCode::Invalid:
 			{
 				ensure(false);
@@ -1248,6 +1254,13 @@ void URigVM::CacheMemoryHandlesIfRequired(TArrayView<URigVMMemoryStorage*> InMem
 				{
 					break;
 				}
+			case ERigVMOpCode::JumpToBranch:
+				{
+					const FRigVMJumpToBranchOp& Op = ByteCode.GetOpAt<FRigVMJumpToBranchOp>(Instructions[InstructionIndex]);
+					const FRigVMOperand& Arg = Op.Arg;
+					InOpFunc(InHandleBaseIndex, {InstructionIndex, 0}, Arg);
+					break;
+				}
 			case ERigVMOpCode::Invalid:
 			default:
 				{
@@ -1700,6 +1713,7 @@ bool URigVM::Initialize(TArrayView<URigVMMemoryStorage*> Memory, TArrayView<void
 			case ERigVMOpCode::ArrayIntersection:
 			case ERigVMOpCode::ArrayReverse:
 			case ERigVMOpCode::InvokeEntry:
+			case ERigVMOpCode::JumpToBranch:
 			{
 				break;
 			}
@@ -2056,6 +2070,12 @@ ERigVMExecuteResult URigVM::ExecuteInstructions(int32 InFirstInstruction, int32
 				Context.Factory = Factories[Op.FunctionIndex];
 				(*Functions[Op.FunctionIndex])(Context, Handles);
 
+				// todo: if this is a control flow node - we may need to jump
+				// this should likely be done with a control flow jump op - which can
+				// loop up the branch based on additional branch infos for it.
+				// we should also make sure to ignore the first branch - since that's the default
+				// we'll just automatically continue to...
+
 #if WITH_EDITOR
 
 				if(DebugMemoryStorageObject->Num() > 0)
@@ -2073,7 +2093,15 @@ ERigVMExecuteResult URigVM::ExecuteInstructions(int32 InFirstInstruction, int32
 			}
 			case ERigVMOpCode::Zero:
 			{
-				*((int32*)CachedMemoryHandles[FirstHandleForInstruction[ContextPublicData.InstructionIndex]].GetData()) = 0;
+				const FRigVMMemoryHandle& Handle = CachedMemoryHandles[FirstHandleForInstruction[ContextPublicData.InstructionIndex]];
+				if(Handle.GetProperty()->IsA<FIntProperty>())
+				{
+					*((int32*)CachedMemoryHandles[FirstHandleForInstruction[ContextPublicData.InstructionIndex]].GetData()) = 0;
+				}
+				else if(Handle.GetProperty()->IsA<FNameProperty>())
+				{
+					*((FName*)CachedMemoryHandles[FirstHandleForInstruction[ContextPublicData.InstructionIndex]].GetData()) = NAME_None;
+				}
 #if WITH_EDITOR
 				if(DebugMemoryStorageObject->Num() > 0)
 				{
@@ -2824,6 +2852,39 @@ ERigVMExecuteResult URigVM::ExecuteInstructions(int32 InFirstInstruction, int32
 				ContextPublicData.InstructionIndex++;
 				break;
 			}
+			case ERigVMOpCode::JumpToBranch:
+			{
+				const FRigVMJumpToBranchOp& Op = ByteCode.GetOpAt<FRigVMJumpToBranchOp>(Instruction);
+				// BranchLabel = Op.Arg
+				const FName& BranchLabel = *(FName*)CachedMemoryHandles[FirstHandleForInstruction[ContextPublicData.InstructionIndex]].GetData();
+
+				// iterate over the branches stored in the bytecode,
+				// starting at the first branch index stored in the operator.
+				// look over all branches matching this instruction index and
+				// find the one with the right label - then jump to the branch.
+				bool bBranchFound = false;
+				const TArray<FRigVMBranchInfo>& Branches = ByteCode.BranchInfos;
+				for(int32 BranchIndex = Op.FirstBranchInfoIndex // start at the first branch known to this jump op
+					; BranchIndex < Branches.Num(); BranchIndex++)
+				{
+					const FRigVMBranchInfo& Branch = Branches[BranchIndex];
+					if(Branch.InstructionIndex != ContextPublicData.InstructionIndex)
+					{
+						break;
+					}
+					if(Branch.Label == BranchLabel)
+					{
+						ContextPublicData.InstructionIndex = Branch.FirstInstruction;
+						bBranchFound = true;
+						break;
+					}
+				}
+				if (!bBranchFound)
+				{
+					ContextPublicData.InstructionIndex++;
+				}
+				break;
+			}
 			case ERigVMOpCode::Invalid:
 			{
 				ensure(false);
@@ -3283,6 +3344,12 @@ TArray<FString> URigVM::DumpByteCodeAsTextArray(const TArray<int32>& InInstructi
 				ResultLine = FString::Printf(TEXT("Invoke entry %s"), *Op.EntryName.ToString());
 				break;
 			}
+			case ERigVMOpCode::JumpToBranch:
+			{
+				const FRigVMJumpToBranchOp& Op = ByteCode.GetOpAt<FRigVMJumpToBranchOp>(Instructions[InstructionIndex]);
+				ResultLine = FString::Printf(TEXT("Jump To Branch %s"), *GetOperandLabel(Op.Arg, OperandFormatFunction));
+				break;
+			}
 			default:
 			{
 				ensure(false);
diff --git a/Engine/Source/Runtime/RigVM/Private/RigVMCore/RigVMByteCode.cpp b/Engine/Source/Runtime/RigVM/Private/RigVMCore/RigVMByteCode.cpp
--- a/Engine/Source/Runtime/RigVM/Private/RigVMCore/RigVMByteCode.cpp
+++ b/Engine/Source/Runtime/RigVM/Private/RigVMCore/RigVMByteCode.cpp
@@ -134,6 +134,13 @@ void FRigVMInvokeEntryOp::Serialize(FArchive& Ar)
 	}
 }
 
+void FRigVMJumpToBranchOp::Serialize(FArchive& Ar)
+{
+	Ar << OpCode;
+	Ar << Arg;
+	Ar << FirstBranchInfoIndex;
+}
+
 FRigVMInstructionArray::FRigVMInstructionArray()
 {
 }
@@ -431,6 +438,12 @@ void FRigVMByteCode::Save(FArchive& Ar) const
 				Ar << Op;
 				break;
 			}
+			case ERigVMOpCode::JumpToBranch:
+			{
+				FRigVMJumpToBranchOp Op = GetOpAt<FRigVMJumpToBranchOp>(Instruction.ByteCodeIndex);
+				Ar << Op;
+				break;
+			}
 			default:
 			{
 				ensure(false);
@@ -665,6 +678,13 @@ void FRigVMByteCode::Load(FArchive& Ar)
 				AddOp<FRigVMInvokeEntryOp>(Op);
 				break;
 			}
+			case ERigVMOpCode::JumpToBranch:
+			{
+				FRigVMJumpToBranchOp Op;
+				Ar << Op;
+				AddOp<FRigVMJumpToBranchOp>(Op);
+				break;
+			}
 			default:
 			{
 				ensure(false);
@@ -921,6 +941,10 @@ uint32 FRigVMByteCode::GetOperatorHash(const FRigVMInstruction& InInstruction) c
 		{
 			return GetTypeHash(GetOpAt<FRigVMInvokeEntryOp>(InInstruction));
 		}
+		case ERigVMOpCode::JumpToBranch:
+		{
+			return GetTypeHash(GetOpAt<FRigVMJumpToBranchOp>(InInstruction));
+		}
 		case ERigVMOpCode::Invalid:
 		{
 			ensure(false);
@@ -1128,6 +1152,10 @@ uint64 FRigVMByteCode::GetOpNumBytesAt(uint64 InByteCodeIndex, bool bIncludeOper
 		{
 			return (uint64)sizeof(FRigVMInvokeEntryOp);
 		}
+		case ERigVMOpCode::JumpToBranch:
+		{
+			return (uint64)sizeof(FRigVMJumpToBranchOp);
+		}
 		case ERigVMOpCode::Invalid:
 		{
 			ensure(false);
@@ -1496,6 +1524,28 @@ FString FRigVMByteCode::DumpToText() const
 				Line += FString::Printf(TEXT(", Entry '%s'"), *Op.EntryName.ToString());
 				break;
 			}
+			case ERigVMOpCode::JumpToBranch:
+			{
+				const FRigVMJumpToBranchOp& Op = GetOpAt<FRigVMJumpToBranchOp>(Instruction.ByteCodeIndex);
+				FString Arg;
+				FRigVMOperand::StaticStruct()->ExportText(Arg, &Op.Arg, &Op.Arg, nullptr, PPF_None, nullptr);
+				Line += TEXT(" BlockToRun ");
+				Line += Arg;
+				Line += TEXT(" for branches ");
+
+				TArray<FString> BranchInfoTexts;
+				for(int32 BranchIndex = Op.FirstBranchInfoIndex; BranchIndex < BranchInfos.Num(); BranchIndex++)
+				{
+					const FRigVMBranchInfo& BranchInfo = BranchInfos[BranchIndex];
+					if(BranchInfo.InstructionIndex != InstructionIndex)
+					{
+						break;
+					}
+					BranchInfoTexts.Add(FString::Printf(TEXT("%s (%d)"), *BranchInfo.Label.ToString(), (int32)BranchInfo.FirstInstruction));
+				}
+				Line += FString::Join(BranchInfoTexts, TEXT(", "));
+				break;
+			}
 			case ERigVMOpCode::Invalid:
 			{
 				ensure(false);
@@ -1614,15 +1664,21 @@ uint64 FRigVMByteCode::AddInvokeEntryOp(const FName& InEntryName)
 	return AddOp(FRigVMInvokeEntryOp(InEntryName));
 }
 
-void FRigVMByteCode::AddBranchInfo(const FRigVMBranchInfo& InBranchInfo)
+uint64 FRigVMByteCode::AddJumpToBranchOp(FRigVMOperand InBranchNameArg, int32 InFirstBranchInfoIndex)
+{
+	return AddOp(FRigVMJumpToBranchOp(InBranchNameArg, InFirstBranchInfoIndex));
+}
+
+int32 FRigVMByteCode::AddBranchInfo(const FRigVMBranchInfo& InBranchInfo)
 {
 	FRigVMBranchInfo BranchInfo = InBranchInfo;
 	BranchInfo.Index = BranchInfos.Num();
 	BranchInfos.Add(BranchInfo);
 	BranchInfoLookup.Reset();
+	return BranchInfo.Index;
 }
 
-void FRigVMByteCode::AddBranchInfo(const FName& InBranchLabel, int32 InInstructionIndex, int32 InArgumentIndex,
+int32 FRigVMByteCode::AddBranchInfo(const FName& InBranchLabel, int32 InInstructionIndex, int32 InArgumentIndex,
 	int32 InFirstBranchInstruction, int32 InLastBranchInstruction)
 {
 	FRigVMBranchInfo BranchInfo;
@@ -1631,7 +1687,7 @@ void FRigVMByteCode::AddBranchInfo(const FName& InBranchLabel, int32 InInstructi
 	BranchInfo.ArgumentIndex = InArgumentIndex;
 	BranchInfo.FirstInstruction = InFirstBranchInstruction;
 	BranchInfo.LastInstruction = InLastBranchInstruction;
-	AddBranchInfo(BranchInfo);
+	return AddBranchInfo(BranchInfo);
 }
 
 FRigVMOperandArray FRigVMByteCode::GetOperandsForOp(const FRigVMInstruction& InInstruction) const
@@ -2046,6 +2102,11 @@ uint64 FRigVMByteCode::GetOpAlignment(ERigVMOpCode InOpCode) const
 			static const uint64 Alignment = FRigVMInvokeEntryOp::StaticStruct()->GetCppStructOps()->GetAlignment();
 			return Alignment;
 		}
+		case ERigVMOpCode::JumpToBranch:
+		{
+			static const uint64 Alignment = FRigVMJumpToBranchOp::StaticStruct()->GetCppStructOps()->GetAlignment();
+			return Alignment;
+		}
 		case ERigVMOpCode::Invalid:
 		{
 			ensure(false);
diff --git a/Engine/Source/Runtime/RigVM/Private/RigVMCore/RigVMDispatchFactory.cpp b/Engine/Source/Runtime/RigVM/Private/RigVMCore/RigVMDispatchFactory.cpp
--- a/Engine/Source/Runtime/RigVM/Private/RigVMCore/RigVMDispatchFactory.cpp
+++ b/Engine/Source/Runtime/RigVM/Private/RigVMCore/RigVMDispatchFactory.cpp
@@ -126,6 +126,21 @@ FString FRigVMDispatchFactory::GetKeywords() const
 
 #endif
 
+const TArray<FName>& FRigVMDispatchFactory::GetControlFlowBlocks() const
+{
+	const TArray<FName>& Blocks = GetControlFlowBlocks_Impl();
+#if WITH_EDITOR
+	FRigVMStruct::ValidateControlFlowBlocks(Blocks);
+#endif
+	return Blocks;
+}
+
+const TArray<FName>& FRigVMDispatchFactory::GetControlFlowBlocks_Impl() const
+{
+	static const TArray<FName> EmptyArray;
+	return EmptyArray;
+}
+
 TArray<FRigVMExecuteArgument> FRigVMDispatchFactory::GetExecuteArguments() const
 {
 	TArray<FRigVMExecuteArgument> Arguments = GetExecuteArguments_Impl();
diff --git a/Engine/Source/Runtime/RigVM/Private/RigVMCore/RigVMRegistry.cpp b/Engine/Source/Runtime/RigVM/Private/RigVMCore/RigVMRegistry.cpp
--- a/Engine/Source/Runtime/RigVM/Private/RigVMCore/RigVMRegistry.cpp
+++ b/Engine/Source/Runtime/RigVM/Private/RigVMCore/RigVMRegistry.cpp
@@ -978,6 +978,15 @@ void FRigVMRegistry::Register(const TCHAR* InName, FRigVMFunctionPtr InFunctionP
 		return;
 	}
 
+#if WITH_EDITOR
+	FString StructureError;
+	if (!FRigVMStruct::ValidateStruct(InStruct, &StructureError))
+	{
+		UE_LOG(LogRigVM, Error, TEXT("Failed to validate struct '%s': %s"), *InStruct->GetName(), *StructureError);
+		return;
+	}
+#endif
+
 	const FRigVMFunction Function(InName, InFunctionPtr, InStruct, Functions.Num(), InArguments);
 	Functions.AddElement(Function);
 	FunctionNameToIndex.Add(InName, Function.Index);
diff --git a/Engine/Source/Runtime/RigVM/Private/RigVMCore/RigVMStruct.cpp b/Engine/Source/Runtime/RigVM/Private/RigVMCore/RigVMStruct.cpp
--- a/Engine/Source/Runtime/RigVM/Private/RigVMCore/RigVMStruct.cpp
+++ b/Engine/Source/Runtime/RigVM/Private/RigVMCore/RigVMStruct.cpp
@@ -307,6 +307,10 @@ const FName FRigVMStruct::ForLoopContinuePinName("Continue");
 const FName FRigVMStruct::ForLoopCompletedPinName("Completed");
 const FName FRigVMStruct::ForLoopIndexPinName("Index");
 const FName FRigVMStruct::ComputeLazilyMetaName("Lazy");
+const FName FRigVMStruct::ControlFlowBlockToRunName("BlockToRun");
+const FName FRigVMStruct::ControlFlowCompletedName("Completed");
+const FName FRigVMStruct::ControlFlowCountName("Count");
+const FName FRigVMStruct::ControlFlowIndexName("Index");
 
 float FRigVMStruct::GetRatioFromIndex(int32 InIndex, int32 InCount)
 {
@@ -317,6 +321,37 @@ float FRigVMStruct::GetRatioFromIndex(int32 InIndex, int32 InCount)
 	return ((float)FMath::Clamp<int32>(InIndex, 0, InCount - 1)) / ((float)(InCount - 1));
 }
 
+bool FRigVMStruct::IsForLoop() const
+{
+	UE_LOG(LogRigVM, Warning, TEXT("FRigVMStruct::IsForLoop is deprecated. Please call FRigVMStruct::IsControlFlowNode instead."));
+	
+	if(IsControlFlowNode())
+	{
+		static const TArray<FName> ExpectedLoopBlocks = {ExecuteContextName, ForLoopCompletedPinName};
+		const TArray<FName>& Blocks = GetControlFlowBlocks();
+		if(Blocks.Num() == ExpectedLoopBlocks.Num())
+		{
+			return Blocks[0] == ExpectedLoopBlocks[0] && Blocks[1] == ExpectedLoopBlocks[1];
+		}
+	}
+	
+	return false;
+}
+
+bool FRigVMStruct::IsControlFlowNode() const
+{
+	return !GetControlFlowBlocks().IsEmpty();
+}
+
+const TArray<FName>& FRigVMStruct::GetControlFlowBlocks() const
+{
+	const TArray<FName>& Blocks = GetControlFlowBlocks_Impl();
+#if WITH_EDITOR
+	ValidateControlFlowBlocks(Blocks);
+#endif
+	return Blocks;
+}
+
 TArray<FRigVMUserWorkflow> FRigVMStruct::GetWorkflows(ERigVMUserWorkflowType InType, const UObject* InSubject) const
 {
 	return GetSupportedWorkflows(InSubject).FilterByPredicate([InType](const FRigVMUserWorkflow& InWorkflow) -> bool
@@ -342,100 +377,88 @@ bool FRigVMStruct::ValidateStruct(UScriptStruct* InStruct, FString* OutErrorMess
 	FStructOnScope StructOnScope(InStruct);
 	FRigVMStruct* StructMemory = (FRigVMStruct*)StructOnScope.GetStructMemory();
 
-	if (StructMemory->IsForLoop())
+	if (StructMemory->IsControlFlowNode())
 	{
-		if (!CheckPinExists(InStruct, ForLoopCountPinName, TEXT("int32"), OutErrorMessage))
+		// first check that the Block To Run pin exists
+		if (!CheckPinExists(InStruct, ControlFlowBlockToRunName, TEXT("FName"), OutErrorMessage))
 		{
 			return false;
 		}
-		else
-		{
-			if (!CheckPinDirection(InStruct, ForLoopCountPinName, InputMetaName) &&
-				!CheckPinDirection(InStruct, ForLoopCountPinName, OutputMetaName) &&
-				!CheckPinDirection(InStruct, ForLoopCountPinName, HiddenMetaName))
-			{
-				if (OutErrorMessage)
-				{
-					*OutErrorMessage = FString::Printf(TEXT("The '%s' pin needs to be either hidden, an input or an output."), *ForLoopCountPinName.ToString());
-				}
-				return false;
-			}
-			if (!CheckMetadata(InStruct, ForLoopCountPinName, SingletonMetaName, OutErrorMessage))
-			{
-				return false;
-			}
-		}
 
-		if (!CheckPinExists(InStruct, ForLoopContinuePinName, TEXT("bool"), OutErrorMessage))
+		// check that the block to run name is marked up as a singleton
+		if (!CheckMetadata(InStruct, ControlFlowBlockToRunName, SingletonMetaName, OutErrorMessage))
 		{
 			return false;
 		}
-		else
+
+		if (!CheckPinDirection(InStruct, ControlFlowBlockToRunName, HiddenMetaName))
 		{
-			if (!CheckPinDirection(InStruct, ForLoopContinuePinName, HiddenMetaName))
-			{
-				if (OutErrorMessage)
-				{
-					*OutErrorMessage = FString::Printf(TEXT("The '%s' pin needs to be hidden."), *ForLoopContinuePinName.ToString());
-				}
-				return false;
-			}
-			if (!CheckMetadata(InStruct, ForLoopContinuePinName, SingletonMetaName, OutErrorMessage))
+			if (OutErrorMessage)
 			{
-				return false;
+				*OutErrorMessage = FString::Printf(TEXT("The '%s' pin needs to be hidden (not an input nor an output)."), *ControlFlowBlockToRunName.ToString());
 			}
-		}
-
-		if (!CheckPinExists(InStruct, ForLoopIndexPinName, TEXT("int32"), OutErrorMessage))
-		{
 			return false;
 		}
-		else
+
+		// the call to GetControlFlowBlocks will already ensure that the block names
+		// themselves are valid.
+		const TArray<FName>& Blocks = StructMemory->GetControlFlowBlocks();
+		for(const FName& Block : Blocks)
 		{
-			if (!CheckPinDirection(InStruct, ForLoopIndexPinName, HiddenMetaName) &&
-				!CheckPinDirection(InStruct, ForLoopIndexPinName, OutputMetaName))
+			if (!CheckPinExists(InStruct, Block, FRigVMExecuteContext::StaticStruct()->GetStructCPPName(), OutErrorMessage))
 			{
-				if (OutErrorMessage)
-				{
-					*OutErrorMessage = FString::Printf(TEXT("The '%s' pin needs to be hidden or an output."), *ForLoopIndexPinName.ToString());
-				}
 				return false;
 			}
-			if (!CheckMetadata(InStruct, ForLoopContinuePinName, SingletonMetaName, OutErrorMessage))
-			{
-				return false;
-			}
-		}
-		
-		if (!CheckPinExists(InStruct, ExecuteContextName, FString(), OutErrorMessage))
-		{
-			return false;
-		}
-		else
-		{
-			if (!CheckPinDirection(InStruct, ExecuteContextName, IOMetaName))
+			if (!CheckPinDirection(InStruct, Block, OutputMetaName) &&
+				!CheckPinDirection(InStruct, Block, IOMetaName))
 			{
 				if (OutErrorMessage)
 				{
-					*OutErrorMessage = FString::Printf(TEXT("The '%s' pin needs to be IO."), *ExecuteContextName.ToString());
+					*OutErrorMessage = FString::Printf(TEXT("The '%s' pin needs to be either an IO or an output."), *Block.ToString());
 				}
 				return false;
 			}
-		}
 
-		if (!CheckPinExists(InStruct, ForLoopCompletedPinName, FString(), OutErrorMessage))
-		{
-			return false;
-		}
-		else
-		{
-			if (!CheckPinDirection(InStruct, ForLoopCompletedPinName, OutputMetaName))
+			if(StructMemory->IsControlFlowBlockSliced(Block))
 			{
-				if (OutErrorMessage)
+				if(Block == ControlFlowCompletedName)
+				{
+					if (OutErrorMessage)
+					{
+						*OutErrorMessage = FString::Printf(TEXT("The '%s' control flow block cannot be sliced."), *Block.ToString());
+					}
+					return false;
+				}
+				if (!CheckPinExists(InStruct, ControlFlowCountName, TEXT("int32"), OutErrorMessage))
 				{
-					*OutErrorMessage = FString::Printf(TEXT("The '%s' pin needs to be an output."), *ForLoopCompletedPinName.ToString());
+					return false;
+				}
+				if (!CheckPinDirection(InStruct, ControlFlowCountName, OutputMetaName) &&
+					!CheckPinDirection(InStruct, ControlFlowCountName, IOMetaName) &&
+					!CheckPinDirection(InStruct, ControlFlowCountName, HiddenMetaName) &&
+					!CheckPinDirection(InStruct, ControlFlowCountName, InputMetaName))
+				{
+					if (OutErrorMessage)
+					{
+						*OutErrorMessage = FString::Printf(TEXT("The '%s' pin needs to be either an IO, input, output or hidden."), *ControlFlowCountName.ToString());
+					}
+					return false;
+				}
+
+				if (!CheckPinExists(InStruct, ControlFlowIndexName, TEXT("int32"), OutErrorMessage))
+				{
+					return false;
+				}
+				if (!CheckPinDirection(InStruct, ControlFlowIndexName, OutputMetaName) &&
+					!CheckPinDirection(InStruct, ControlFlowIndexName, IOMetaName) &&
+					!CheckPinDirection(InStruct, ControlFlowIndexName, HiddenMetaName))
+				{
+					if (OutErrorMessage)
+					{
+						*OutErrorMessage = FString::Printf(TEXT("The '%s' pin needs to be either an IO, an output or hidden."), *ControlFlowIndexName.ToString());
+					}
+					return false;
 				}
-				return false;
 			}
 		}
 	}
@@ -462,13 +485,23 @@ bool FRigVMStruct::CheckPinDirection(UScriptStruct* InStruct, const FName& PinNa
 
 bool FRigVMStruct::CheckPinType(UScriptStruct* InStruct, const FName& PinName, const FString& ExpectedType, FString* OutErrorMessage)
 {
-	if (FProperty* Property = InStruct->FindPropertyByName(PinName))
+	if (const FProperty* Property = InStruct->FindPropertyByName(PinName))
 	{
 		if (Property->GetCPPType() != ExpectedType)
 		{
+			if(ExpectedType == FRigVMExecuteContext::StaticStruct()->GetStructCPPName())
+			{
+				if(const FStructProperty* StructProperty = CastField<FStructProperty>(Property))
+				{
+					if(StructProperty->Struct->IsChildOf(FRigVMExecuteContext::StaticStruct()))
+					{
+						return true;
+					}
+				}
+			}
 			if (OutErrorMessage)
 			{
-				*OutErrorMessage = FString::Printf(TEXT("The '%s' property needs to be of type '%s'."), *ForLoopCountPinName.ToString(), *ExpectedType);;
+				*OutErrorMessage = FString::Printf(TEXT("The '%s' property needs to be of type '%s'."), *PinName.ToString(), *ExpectedType);;
 			}
 			return false;
 		}
@@ -479,16 +512,12 @@ bool FRigVMStruct::CheckPinType(UScriptStruct* InStruct, const FName& PinName, c
 
 bool FRigVMStruct::CheckPinExists(UScriptStruct* InStruct, const FName& PinName, const FString& ExpectedType, FString* OutErrorMessage)
 {
-	if (FProperty* Property = InStruct->FindPropertyByName(PinName))
+	if (InStruct->FindPropertyByName(PinName))
 	{
 		if (!ExpectedType.IsEmpty())
 		{
-			if (Property->GetCPPType() != ExpectedType)
+			if(!CheckPinType(InStruct, PinName, ExpectedType, OutErrorMessage))
 			{
-				if (OutErrorMessage)
-				{
-					*OutErrorMessage = FString::Printf(TEXT("The '%s' property needs to be of type '%s'."), *ForLoopCountPinName.ToString(), *ExpectedType);;
-				}
 				return false;
 			}
 		}
@@ -499,11 +528,11 @@ bool FRigVMStruct::CheckPinExists(UScriptStruct* InStruct, const FName& PinName,
 		{
 			if (ExpectedType.IsEmpty())
 			{
-				*OutErrorMessage = FString::Printf(TEXT("Struct requires a '%s' property."), *ForLoopCountPinName.ToString());;
+				*OutErrorMessage = FString::Printf(TEXT("Struct requires a '%s' property."), *PinName.ToString());;
 			}
 			else
 			{
-				*OutErrorMessage = FString::Printf(TEXT("Struct requires a '%s' property of type '%s'."), *ForLoopCountPinName.ToString(), *ExpectedType);;
+				*OutErrorMessage = FString::Printf(TEXT("Struct requires a '%s' property of type '%s'."), *PinName.ToString(), *ExpectedType);;
 			}
 		}
 		return false;
@@ -811,3 +840,29 @@ bool FRigVMStruct::ApplyUpgradeInfo(const FRigVMStructUpgradeInfo& InUpgradeInfo
 	return true;
 }
 
+const TArray<FName>& FRigVMStruct::GetControlFlowBlocks_Impl() const
+{
+	static const TArray<FName> EmptyArray;
+	return EmptyArray;
+}
+
+#if WITH_EDITOR
+
+void FRigVMStruct::ValidateControlFlowBlocks(const TArray<FName>& InBlocks)
+{
+	if(!InBlocks.IsEmpty())
+	{
+		check(InBlocks.Num() > 1);
+		check(InBlocks.Last() == ControlFlowCompletedName);
+
+		TArray<FName> UniqueBlocks;
+		UniqueBlocks.Reserve(InBlocks.Num());
+		for(const FName& Block : InBlocks)
+		{
+			UniqueBlocks.AddUnique(Block);
+		}
+		check(UniqueBlocks.Num() == InBlocks.Num());
+	}
+}
+
+#endif
\ No newline at end of file
diff --git a/Engine/Source/Runtime/RigVM/Public/RigVMCore/RigVMByteCode.h b/Engine/Source/Runtime/RigVM/Public/RigVMCore/RigVMByteCode.h
--- a/Engine/Source/Runtime/RigVM/Public/RigVMCore/RigVMByteCode.h
+++ b/Engine/Source/Runtime/RigVM/Public/RigVMCore/RigVMByteCode.h
@@ -194,6 +194,7 @@ enum class ERigVMOpCode : uint8
 	ArrayIntersection, // returns a new array containing elements found in both of the input arrays (ternary op, in array, in array, out result)
 	ArrayReverse, // returns the reverse of the input array (unary op, in out array)
 	InvokeEntry, // invokes an entry from the entry list
+	JumpToBranch, // jumps to a branch based on a name operand
 	Invalid,
 	FirstArrayOpCode = ArrayReset,
 	LastArrayOpCode = ArrayReverse,
@@ -281,7 +282,8 @@ struct RIGVM_API FRigVMUnaryOp : public FRigVMBaseOp
 			uint8(InOpCode) == uint8(ERigVMOpCode::JumpBackwardIf) ||
 			uint8(InOpCode) == uint8(ERigVMOpCode::ChangeType) ||
 			uint8(InOpCode) == uint8(ERigVMOpCode::ArrayReset) ||
-			uint8(InOpCode) == uint8(ERigVMOpCode::ArrayReverse)
+			uint8(InOpCode) == uint8(ERigVMOpCode::ArrayReverse) ||
+			uint8(InOpCode) == uint8(ERigVMOpCode::JumpToBranch)
 		);
 	}
 
@@ -752,7 +754,7 @@ struct RIGVM_API FRigVMJumpIfOp : public FRigVMUnaryOp
 
 	friend FORCEINLINE uint32 GetTypeHash(const FRigVMJumpIfOp& Op)
 	{
-		uint32 Hash = GetTypeHash((const FRigVMBaseOp&)Op);
+		uint32 Hash = GetTypeHash((const FRigVMUnaryOp&)Op);
 		Hash = HashCombine(Hash, GetTypeHash(Op.InstructionIndex));
 		Hash = HashCombine(Hash, GetTypeHash(Op.Condition));
 		return Hash;
@@ -811,6 +813,42 @@ struct RIGVM_API FRigVMInvokeEntryOp : public FRigVMBaseOp
 	}
 };
 
+// jump into a branch based on a name argument
+USTRUCT()
+struct RIGVM_API FRigVMJumpToBranchOp : public FRigVMUnaryOp
+{
+	GENERATED_USTRUCT_BODY()
+
+	FRigVMJumpToBranchOp()
+		: FRigVMUnaryOp()
+		, FirstBranchInfoIndex(INDEX_NONE)
+	{
+	}
+
+	FRigVMJumpToBranchOp(FRigVMOperand InBranchNameArg, int32 InFirstBranchInfoIndex)
+		: FRigVMUnaryOp(ERigVMOpCode::JumpToBranch, InBranchNameArg)
+		, FirstBranchInfoIndex(InFirstBranchInfoIndex)
+	{
+	}
+
+	int32 FirstBranchInfoIndex;
+
+	friend FORCEINLINE uint32 GetTypeHash(const FRigVMJumpToBranchOp& Op)
+	{
+		uint32 Hash = GetTypeHash((const FRigVMBaseOp&)Op);
+		Hash = HashCombine(Hash, GetTypeHash(Op.FirstBranchInfoIndex));
+		return Hash;
+	}
+
+	void Serialize(FArchive& Ar);
+	FORCEINLINE friend FArchive& operator<<(FArchive& Ar, FRigVMJumpToBranchOp& P)
+	{
+		P.Serialize(Ar);
+		return Ar;
+	}
+};
+
+
 /**
  * The FRigVMInstruction represents
  * a single instruction within the VM.
@@ -1055,9 +1093,12 @@ public:
 	// adds an invoke entry operator
 	uint64 AddInvokeEntryOp(const FName& InEntryName);
 
+	// adds a jump to branch operator
+	uint64 AddJumpToBranchOp(FRigVMOperand InBranchNameArg, int32 InFirstBranchInfoIndex);
+
 	// adds information about a branch for an instruction's argument
-	void AddBranchInfo(const FRigVMBranchInfo& InBranchInfo);
-	void AddBranchInfo(const FName& InBranchLabel, int32 InInstructionIndex, int32 InArgumentIndex, int32 InFirstBranchInstruction, int32 InLastBranchInstruction);
+	int32 AddBranchInfo(const FRigVMBranchInfo& InBranchInfo);
+	int32 AddBranchInfo(const FName& InBranchLabel, int32 InInstructionIndex, int32 InArgumentIndex, int32 InFirstBranchInstruction, int32 InLastBranchInstruction);
 
 	// returns an instruction array for iterating over all operators
 	FORCEINLINE FRigVMInstructionArray GetInstructions() const
@@ -1290,4 +1331,5 @@ private:
 	friend class FRigVMByteCodeTest;
 	friend class FRigVMInvokeEntryTest;
 	friend class URigVM;
+	friend struct FRigVMCodeGenerator;
 };
diff --git a/Engine/Source/Runtime/RigVM/Public/RigVMCore/RigVMDispatchFactory.h b/Engine/Source/Runtime/RigVM/Public/RigVMCore/RigVMDispatchFactory.h
--- a/Engine/Source/Runtime/RigVM/Public/RigVMCore/RigVMDispatchFactory.h
+++ b/Engine/Source/Runtime/RigVM/Public/RigVMCore/RigVMDispatchFactory.h
@@ -88,6 +88,15 @@ public:
 
 #endif
 
+	// returns true if the dispatch is a control flow dispatch
+	FORCEINLINE bool IsControlFlowDispatch() const { return !GetControlFlowBlocks().IsEmpty(); }
+
+	// returns the control flow blocks of this dispatch
+	const TArray<FName>& GetControlFlowBlocks() const;
+
+	// returns true if a given control flow block needs to be sliced
+	virtual const bool IsControlFlowBlockSliced(const FName& InBlockName) const { return false; }
+
 	// Returns the execute context support for this dispatch factory
 	FORCEINLINE virtual UScriptStruct* GetExecuteContextStruct() const { return FRigVMExecuteContext::StaticStruct(); }
 
@@ -167,6 +176,8 @@ protected:
 	}
 #endif
 
+	virtual const TArray<FName>& GetControlFlowBlocks_Impl() const;
+
 	static const FString DispatchPrefix;
 
 	UScriptStruct* FactoryScriptStruct;
diff --git a/Engine/Source/Runtime/RigVM/Public/RigVMCore/RigVMExecuteContext.h b/Engine/Source/Runtime/RigVM/Public/RigVMCore/RigVMExecuteContext.h
--- a/Engine/Source/Runtime/RigVM/Public/RigVMCore/RigVMExecuteContext.h
+++ b/Engine/Source/Runtime/RigVM/Public/RigVMCore/RigVMExecuteContext.h
@@ -224,6 +224,8 @@ struct RIGVM_API FRigVMExecuteContext
 		: EventName(NAME_None)
 		, FunctionName(NAME_None)
 		, InstructionIndex(0)
+		, DeltaTime(0.0)
+		, AbsoluteTime(0.0)
 		, RuntimeSettings()
 		, NameCache(nullptr)
 	{
@@ -266,6 +268,12 @@ struct RIGVM_API FRigVMExecuteContext
 	
 	FORCEINLINE FName GetEventName() const { return EventName; }
 
+	FORCEINLINE double GetDeltaTime() const { return DeltaTime; }
+	FORCEINLINE void SetDeltaTime(double InDeltaTime) { DeltaTime = InDeltaTime; }
+
+	FORCEINLINE double GetAbsoluteTime() const { return AbsoluteTime; } 
+	FORCEINLINE void SetAbsoluteTime(double InAbsoluteTime) { AbsoluteTime = InAbsoluteTime; }
+	
 	FORCEINLINE FRigVMNameCache* GetNameCache() const { return NameCache; }
 
 	virtual void Initialize()
@@ -299,6 +307,10 @@ protected:
 	
 	uint16 InstructionIndex;
 
+	double DeltaTime;
+
+	double AbsoluteTime;
+
 	FRigVMRuntimeSettings RuntimeSettings;
 
 	mutable FRigVMNameCache* NameCache;
diff --git a/Engine/Source/Runtime/RigVM/Public/RigVMCore/RigVMStruct.h b/Engine/Source/Runtime/RigVM/Public/RigVMCore/RigVMStruct.h
--- a/Engine/Source/Runtime/RigVM/Public/RigVMCore/RigVMStruct.h
+++ b/Engine/Source/Runtime/RigVM/Public/RigVMCore/RigVMStruct.h
@@ -170,9 +170,12 @@ struct RIGVM_API FRigVMStruct
 
 public:
 
-	// loop related
-	FORCEINLINE virtual bool IsForLoop() const { return false; }
+	// control flow related
+	bool IsForLoop() const;
+	bool IsControlFlowNode() const; 
 	FORCEINLINE virtual int32 GetNumSlices() const { return 1; }
+	const TArray<FName>& GetControlFlowBlocks() const;
+	virtual const bool IsControlFlowBlockSliced(const FName& InBlockName) const { return false; }
 
 	// node creation
 	FORCEINLINE virtual void OnUnitNodeCreated(FRigVMUnitNodeCreatedContext& InContext) const {}
@@ -250,15 +253,25 @@ public:
 	static const FName ForLoopCompletedPinName;
 	static const FName ForLoopIndexPinName;
 	static const FName ComputeLazilyMetaName;
+	static const FName ControlFlowBlockToRunName;
+	static const FName ControlFlowCompletedName;
+	static const FName ControlFlowCountName;
+	static const FName ControlFlowIndexName;
 
 protected:
 
 	static float GetRatioFromIndex(int32 InIndex, int32 InCount);
 	TMap<FName, FString> GetDefaultValues(UScriptStruct* InScriptStruct) const;
 	bool ApplyUpgradeInfo(const FRigVMStructUpgradeInfo& InUpgradeInfo);
-	FORCEINLINE virtual TArray<FRigVMUserWorkflow> GetSupportedWorkflows(const UObject* InSubject) const { return TArray<FRigVMUserWorkflow>(); } 
+	FORCEINLINE virtual TArray<FRigVMUserWorkflow> GetSupportedWorkflows(const UObject* InSubject) const { return TArray<FRigVMUserWorkflow>(); }
+	virtual const TArray<FName>& GetControlFlowBlocks_Impl() const;
+
+#if WITH_EDITOR
+	static void ValidateControlFlowBlocks(const TArray<FName>& InBlocks);
+#endif
 
 	friend struct FRigVMStructUpgradeInfo;
 	friend class FRigVMGraphStructUpgradeInfoTest;
 	friend class URigVMController;
+	friend struct FRigVMDispatchFactory;
 };
