Control Rig: Fix saving "FText which has been initialized from FString" at cook time warning

diff --git a/Engine/Source/Developer/RigVMDeveloper/Private/RigVMModel/Nodes/RigVMLibraryNode.cpp b/Engine/Source/Developer/RigVMDeveloper/Private/RigVMModel/Nodes/RigVMLibraryNode.cpp
--- a/Engine/Source/Developer/RigVMDeveloper/Private/RigVMModel/Nodes/RigVMLibraryNode.cpp
+++ b/Engine/Source/Developer/RigVMDeveloper/Private/RigVMModel/Nodes/RigVMLibraryNode.cpp
@@ -245,8 +245,11 @@ FRigVMGraphFunctionHeader URigVMLibraryNode::GetFunctionHeader(IRigVMGraphFuncti
     
     Header.Name = GetFName();
     Header.Category = GetNodeCategory();
-    Header.Keywords = GetNodeKeywords();						
-    Header.Tooltip = GetToolTipText();
+    Header.Keywords = GetNodeKeywords();
+	
+	// Avoid "initialized from FString" warning while cooking. The graph function tooltip does not have useful information anyway.
+	//Header.Tooltip = GetToolTipText();
+	
     Header.NodeColor = GetNodeColor();
     Header.NodeTitle = GetNodeTitle();
     for(URigVMPin* Pin : GetPins())
@@ -266,7 +269,11 @@ FRigVMGraphFunctionHeader URigVMLibraryNode::GetFunctionHeader(IRigVMGraphFuncti
     	for (int32 i=0; i<PinsToProcess.Num(); ++i)
     	{
     		const FString SubPinPath = PinsToProcess[i]->GetSubPinPath(Pin, false);
-    		Arg.PathToTooltip.Add(SubPinPath, GetToolTipTextForPin(PinsToProcess[i]));
+
+    		// Avoid "initialized from FString" warning while cooking. The graph function tooltip does not have useful information anyway.
+    		//Arg.PathToTooltip.Add(SubPinPath, GetToolTipTextForPin(PinsToProcess[i]));
+    		Arg.PathToTooltip.Add(SubPinPath, FText());
+    		
     		PinsToProcess.Append(PinsToProcess[i]->GetSubPins());
     	}
     	
