Control Rig: Fix RigVMFunctionLibrary storing soft references to non-existing paths

diff --git a/Engine/Source/Developer/RigVMDeveloper/Public/RigVMModel/RigVMFunctionLibrary.h b/Engine/Source/Developer/RigVMDeveloper/Public/RigVMModel/RigVMFunctionLibrary.h
--- a/Engine/Source/Developer/RigVMDeveloper/Public/RigVMModel/RigVMFunctionLibrary.h
+++ b/Engine/Source/Developer/RigVMDeveloper/Public/RigVMModel/RigVMFunctionLibrary.h
@@ -87,7 +87,7 @@ private:
 	// The source pathname is the full path of the source function that was localized
 	// to the local copy stored in the value of the pair.
 	UPROPERTY()
-	TMap< FSoftObjectPath, TObjectPtr<URigVMLibraryNode> > LocalizedFunctions;
+	TMap< FString, TObjectPtr<URigVMLibraryNode> > LocalizedFunctions;
 
 	friend class URigVMController;
 	friend class URigVMCompiler;
