RigVM: Upgrade nativization to support control flow, dispatches and lazy eval

diff --git a/Engine/Plugins/Animation/ControlRig/Source/ControlRig/Private/ControlRig.cpp b/Engine/Plugins/Animation/ControlRig/Source/ControlRig/Private/ControlRig.cpp
--- a/Engine/Plugins/Animation/ControlRig/Source/ControlRig/Private/ControlRig.cpp
+++ b/Engine/Plugins/Animation/ControlRig/Source/ControlRig/Private/ControlRig.cpp
@@ -810,6 +810,7 @@ bool UControlRig::Execute(const EControlRigState InState, const FName& InEventNa
 	if (bRequiresInitExecution)
 	{
 		bRequiresInitExecution = false;
+		VM->SetContextPublicDataStruct(FControlRigExecuteContext::StaticStruct());
 
 		if (!bIsInitializingMemory)
 		{
@@ -3382,6 +3383,25 @@ void UControlRig::SwapVMToNativizedIfRequired(UClass* InNativizedClass)
 #endif
 		}
 	}
+	
+#if WITH_EDITOR
+
+	// if we are a nativized VM,
+	// let's set the bytecode for UI purposes.
+	// this is only used for traversing node from execute stack to node and back etc. 
+	// since the hash between nativized VM and current matches we assume the bytecode is identical as well.
+	if(VM->IsNativized())
+	{
+		if(URigVMNativized* NativizedVM = Cast<URigVMNativized>(VM))
+		{
+			UControlRig* CDO = GetClass()->GetDefaultObject<UControlRig>();
+			if (CDO && CDO->VM)
+			{
+				NativizedVM->SetByteCode(CDO->VM->GetByteCode());
+			}
+		}
+	}
+#endif
 }
 
 bool UControlRig::AreNativizedVMsDisabled()
diff --git a/Engine/Plugins/Animation/ControlRig/Source/ControlRigEditor/Private/Editor/ControlRigEditor.cpp b/Engine/Plugins/Animation/ControlRig/Source/ControlRigEditor/Private/Editor/ControlRigEditor.cpp
--- a/Engine/Plugins/Animation/ControlRig/Source/ControlRigEditor/Private/Editor/ControlRigEditor.cpp
+++ b/Engine/Plugins/Animation/ControlRig/Source/ControlRigEditor/Private/Editor/ControlRigEditor.cpp
@@ -3105,6 +3105,14 @@ FGraphAppearanceInfo FControlRigEditor::GetGraphAppearance(UEdGraph* InGraph) co
 	if (GetBlueprintObj()->IsA(UControlRigBlueprint::StaticClass()))
 	{
 		AppearanceInfo.CornerText = LOCTEXT("AppearanceCornerText_ControlRig", "RIG");
+
+		if(ControlRig && ControlRig->VM && ControlRig->VM->IsNativized())
+		{
+			AppearanceInfo.InstructionFade = 1;
+			AppearanceInfo.InstructionText = FText::FromString(
+				FString::Printf(TEXT("This graph runs a nativized VM (U%s)."), *ControlRig->VM->GetNativizedClass()->GetName())
+			);
+		}
 	}
 
 	return AppearanceInfo;
diff --git a/Engine/Source/Developer/RigVMDeveloper/Private/RigVMCompiler/RigVMCodeGenerator.cpp b/Engine/Source/Developer/RigVMDeveloper/Private/RigVMCompiler/RigVMCodeGenerator.cpp
--- a/Engine/Source/Developer/RigVMDeveloper/Private/RigVMCompiler/RigVMCodeGenerator.cpp
+++ b/Engine/Source/Developer/RigVMDeveloper/Private/RigVMCompiler/RigVMCodeGenerator.cpp
@@ -11,6 +11,14 @@ static constexpr TCHAR RigVM_CommaSeparator[] = TEXT(", ");
 static constexpr TCHAR RigVM_NewLineFormat[] = TEXT("\r\n");
 static constexpr TCHAR RigVM_IncludeBracketFormat[] = TEXT("#include <{0}>");
 static constexpr TCHAR RigVM_IncludeQuoteFormat[] = TEXT("#include \"{0}.h\"");
+static constexpr TCHAR RigVM_DispatchKeyFormat[] = TEXT("{0}_{1}");
+static constexpr TCHAR RigVM_DispatchDeclarationFormat[] = TEXT("\tFORCEINLINE bool {0}({1})\r\n\t{\r\n\t\tstatic const FRigVMFunction* Dispatch = FRigVMRegistry::Get().FindFunction(TEXT(\"{2}\"));\r\n\t\tif(Dispatch == nullptr) return false;");
+static constexpr TCHAR RigVM_UPropertyDeclareFormat[] = TEXT("\tUPROPERTY()\r\n\t{0} {1};");
+static constexpr TCHAR RigVM_UPropertyMemberFormat[] = TEXT("\tstatic const FProperty* {0}_Ptr;");
+static constexpr TCHAR RigVM_UPropertyMember2Format[] = TEXT("const FProperty* U{0}::{1}_Ptr = nullptr;");
+static constexpr TCHAR RigVM_UPropertyDefineFormat[] = TEXT("\tif({1}_Ptr == nullptr)\r\n\t{\r\n\t\t{1}_Ptr = StaticClass()->FindPropertyByName(GET_MEMBER_NAME_CHECKED(U{0}, {1}));\r\n\t}");
+static constexpr TCHAR RigVM_InvokeDispatchPtrFormat[] = TEXT("\t\t(*Dispatch->FunctionPtr)({0});");
+static constexpr TCHAR RigVM_InvokeDispatchFormat[] = TEXT("\t{0}({1});");
 static constexpr TCHAR RigVM_WrappedArrayTypeFormat[] = TEXT("struct {0}_API {1}\r\n{\r\n\tTArray<{2}> Array;\r\n};");
 static constexpr TCHAR RigVM_WrappedTypeNameFormat[] = TEXT("{0}Array_{1}");
 static constexpr TCHAR RigVM_DeclareExternalVariableFormat[] = TEXT("\t{0}* {1} = nullptr;");
@@ -18,6 +26,8 @@ static constexpr TCHAR RigVM_UpdateExternalVariableFormat[] = TEXT("\t{0} = &Get
 static constexpr TCHAR RigVM_MemberPropertyFormat[] = TEXT("\t{0} {1};");
 static constexpr TCHAR RigVM_DeclareEntryNameFormat[] = TEXT("\tstatic const FName EntryName_{0};");
 static constexpr TCHAR RigVM_DefineEntryNameFormat[] = TEXT("const FName U{0}::EntryName_{1} = TEXT(\"{2}\");");
+static constexpr TCHAR RigVM_DeclareBlockNameFormat[] = TEXT("\tstatic const FName BlockName_{0};");
+static constexpr TCHAR RigVM_DefineBlockNameFormat[] = TEXT("const FName U{0}::BlockName_{1} = TEXT(\"{2}\");");
 static constexpr TCHAR RigVM_DefineConstFormatNoDefault[] = TEXT("\tstatic const {0} {1};");
 static constexpr TCHAR RigVM_DefineStructConstFormat[] = TEXT("\tstatic const {0} {1} = URigVMNativized::GetStructConstant<{0}>(TEXT(\"{2}\"));");
 static constexpr TCHAR RigVM_DefineStructArrayConstFormat[] = TEXT("\tstatic const {0} {1} = URigVMNativized::GetStructArrayConstant<{2}>(TEXT(\"{3}\"));");
@@ -31,7 +41,8 @@ static constexpr TCHAR RigVM_CurlyBracesFormat[] = TEXT("{{0}}");
 static constexpr TCHAR RigVM_BracesFormat[] = TEXT("({0})");
 static constexpr TCHAR RigVM_TemplateOneArgFormat[] = TEXT("<{0}>");
 static constexpr TCHAR RigVM_CallExternOpFormat[] = TEXT("\t{0}::Static{1}({2});");
-static constexpr TCHAR RigVM_ZeroOpFormat[] = TEXT("\t{0} = 0;");
+static constexpr TCHAR RigVM_ZeroOpIntFormat[] = TEXT("\t{0} = 0;");
+static constexpr TCHAR RigVM_ZeroOpNameFormat[] = TEXT("\t{0} = NAME_None;");
 static constexpr TCHAR RigVM_BoolFalseOpFormat[] = TEXT("\t{0} = false;");
 static constexpr TCHAR RigVM_BoolTrueFormat[] = TEXT("\t{0} = true;");
 static constexpr TCHAR RigVM_CopyUnrelatedArraysFormat[] = TEXT("\tCopyUnrelatedArrays<{0}, {1}>({2}, {3});");
@@ -41,13 +52,13 @@ static constexpr TCHAR RigVM_IncrementOpFormat[] = TEXT("\t{0}++;");
 static constexpr TCHAR RigVM_DecrementOpFormat[] = TEXT("\t{0}--;");
 static constexpr TCHAR RigVM_EqualsOpFormat[] = TEXT("\t{0} = {1} == {2};");
 static constexpr TCHAR RigVM_InvokeEntryFormat[] = TEXT("\tif (InEntryName == EntryName_{0}) return ExecuteEntry_{0}({1});");
-static constexpr TCHAR RigVM_InvokeEntryByNameFormat[] = TEXT("\treturn InvokeEntryByName(InEntryName{0});");
-static constexpr TCHAR RigVM_InvokeEntryByNameFormat2[] = TEXT("\tif(!InvokeEntryByName({0}{1})) return false;");
-static constexpr TCHAR RigVM_CanExecuteEntryFormat[] = TEXT("\tif(!CanExecuteEntry(InEntryName)) { return false; }");
+static constexpr TCHAR RigVM_InvokeEntryByNameFormat[] = TEXT("\tERigVMExecuteResult EntryResult = InvokeEntryByName(InEntryName{0});\r\n\tSetInstructionIndex(INDEX_NONE);\r\n\tStopProfiling();\r\n\treturn ExecuteResult;");
+static constexpr TCHAR RigVM_InvokeEntryByNameFormat2[] = TEXT("\tif(!InvokeEntryByName({0}{1})) return ERigVMExecuteResult::Failed;");
+static constexpr TCHAR RigVM_CanExecuteEntryFormat[] = TEXT("\tif(!CanExecuteEntry(InEntryName, false)) { return ERigVMExecuteResult::Failed; }");
 static constexpr TCHAR RigVM_EntryExecuteGuardFormat[] = TEXT("\tFEntryExecuteGuard EntryExecuteGuard(EntriesBeingExecuted, FindEntry(InEntryName));");
 static constexpr TCHAR RigVM_PublicContextGuardFormat[] = TEXT("\tTGuardValue<{0}> PublicContextGuard(Context.GetPublicData<{0}>(), PublicContext);");
 static constexpr TCHAR RigVM_EntryNameFormat[] = TEXT("EntryName_{0}");
-static constexpr TCHAR RigVM_SetExecuteContextStructFormmt[] = TEXT("\tSetContextPublicDataStruct({0}::StaticStruct());");
+static constexpr TCHAR RigVM_SetExecuteContextStructFormat[] = TEXT("\tSetContextPublicDataStruct({0}::StaticStruct());");
 static constexpr TCHAR RigVM_UpdateContextFormat[] = TEXT("\t{0}& PublicContext = UpdateContext<{0}>(AdditionalArguments, {1}, InEntryName);");
 static constexpr TCHAR RigVM_TrueFormat[] = TEXT("true");  
 static constexpr TCHAR RigVM_FalseFormat[] = TEXT("false");
@@ -58,17 +69,20 @@ static constexpr TCHAR RigVM_EnumTypeSuffixFormat[] = TEXT("::Type");
 static constexpr TCHAR RigVM_IsValidArraySizeFormat[] = TEXT("IsValidArraySize({0})");
 static constexpr TCHAR RigVM_IsValidArrayIndexFormat[] = TEXT("IsValidArrayIndex<{0}>(TemporaryArrayIndex, {1})");
 static constexpr TCHAR RigVM_TemporaryArrayIndexFormat[] = TEXT("\tTemporaryArrayIndex = {0};");
+static constexpr TCHAR RigVM_StartProfilingFormat[] = TEXT("\tStartProfiling");
 static constexpr TCHAR RigVM_ExecuteReachedExitFormat[] = TEXT("\tBroadcastExecutionReachedExit();");
 static constexpr TCHAR RigVM_AdditionalArgumentNameFormat[] = TEXT("AdditionalArgument_{0}_{1}");
 static constexpr TCHAR RigVM_FullOpaqueArgumentFormat[] = TEXT("\t{0} {1} = *({2}*)Context.OpaqueArguments[{3}];");
 static constexpr TCHAR RigVM_ParameterOpaqueArgumentFormat[] = TEXT("{0} {1}");
 static constexpr TCHAR RigVM_InstructionLabelFormat[] = TEXT("\tInstruction{0}Label:");
 static constexpr TCHAR RigVM_SetInstructionIndexFormat[] = TEXT("\tSetInstructionIndex({0});");
+static constexpr TCHAR RigVM_ContextFormat[] = TEXT("Context");
 static constexpr TCHAR RigVM_ContextPublicFormat[] = TEXT("PublicContext");
 static constexpr TCHAR RigVM_ContextPublicParameterFormat[] = TEXT("{0}& PublicContext");
 static constexpr TCHAR RigVM_NotEqualsOpFormat[] = TEXT("\t{0} = {1} != {2};");
 static constexpr TCHAR RigVM_JumpOpFormat[] = TEXT("\tgoto Instruction{0}Label;");
 static constexpr TCHAR RigVM_JumpIfOpFormat[] = TEXT("\tif ({0} == {1}) { goto Instruction{2}Label; }");
+static constexpr TCHAR RigVM_JumpToBranchFormat[] = TEXT("\tif ({0} == BlockName_{1}) { goto Instruction{2}Label; }");
 static constexpr TCHAR RigVM_BeginBlockOpFormat[] = TEXT("\tBeginSlice({0}, {1});");
 static constexpr TCHAR RigVM_EndBlockOpFormat[] = TEXT("\tEndSlice();");
 static constexpr TCHAR RigVM_ArrayResetOpFormat[] = TEXT("\t{0}.Reset();");
@@ -100,6 +114,8 @@ static constexpr TCHAR RigVM_GetVMHashFormat[] = TEXT("\tvirtual uint32 GetVMHas
 static constexpr TCHAR RigVM_GetEntryNamesFormat[] = TEXT("\tvirtual const TArray<FName>& GetEntryNames() const override\r\n\t{\r\n\t\tstatic const TArray<FName> StaticEntryNames = { {0} };\r\n\t\treturn StaticEntryNames;\r\n\t}");
 static constexpr TCHAR RigVM_DeclareUpdateExternalVariablesFormat[] = TEXT("\tvirtual void UpdateExternalVariables() override;");
 static constexpr TCHAR RigVM_DeclareInvokeEntryByNameFormat[] = TEXT("\tERigVMExecuteResult InvokeEntryByName(const FName& InEntryName{0});");
+static constexpr TCHAR RigVM_DeclareInitializeFormat[] = TEXT("\tvirtual bool Initialize(TArrayView<URigVMMemoryStorage*> Memory, TArrayView<void*> AdditionalArguments, bool bInitializeMemory = true) override;");
+static constexpr TCHAR RigVM_DefineInitializeFormat[] = TEXT("bool U{0}::Initialize(TArrayView<URigVMMemoryStorage*> Memory, TArrayView<void*> AdditionalArguments, bool bInitializeMemory)\r\n{");
 static constexpr TCHAR RigVM_DeclareExecuteFormat[] = TEXT("\tvirtual ERigVMExecuteResult Execute(TArrayView<URigVMMemoryStorage*> Memory, TArrayView<void*> AdditionalArguments, const FName& InEntryName) override;");
 static constexpr TCHAR RigVM_DefineUpdateExternalVariablesFormat[] = TEXT("void U{0}::UpdateExternalVariables()\r\n{");
 static constexpr TCHAR RigVM_DefineInvokeEntryByNameFormat[] = TEXT("ERigVMExecuteResult U{0}::InvokeEntryByName(const FName& InEntryName{1})\r\n{");
@@ -110,6 +126,7 @@ static constexpr TCHAR RigVM_DeclareExecuteGroupFormat[] = TEXT("\tERigVMExecute
 static constexpr TCHAR RigVM_DefineExecuteGroupFormat[] = TEXT("ERigVMExecuteResult U{0}::ExecuteGroup_{1}_{2}({3})\r\n{");
 static constexpr TCHAR RigVM_InvokeExecuteGroupFormat[] = TEXT("\tif(ExecuteGroup_{0}_{1}({2}) != ERigVMExecuteResult::Succeeded) return ERigVMExecuteResult::Failed;");
 static constexpr TCHAR RigVM_RigVMCoreIncludeFormat[] = TEXT("RigVMCore/RigVMCore.h");
+static constexpr TCHAR RigVM_RigVMModuleIncludeFormat[] = TEXT("RigVMModule.h");
 static constexpr TCHAR RigVM_RigVMCoreLibraryFormat[] = TEXT("RigVM");
 static constexpr TCHAR RigVM_JoinFilePathFormat[] = TEXT("{0}/{1}");
 static constexpr TCHAR RigVM_GetOperandSliceFormat[] = TEXT("GetOperandSlice<{0}>({1}){2}");
@@ -117,6 +134,11 @@ static constexpr TCHAR RigVM_ExternalVariableFormat[] = TEXT("(*External_{0})");
 static constexpr TCHAR RigVM_JoinSegmentPathFormat[] = TEXT("{0}.{1}");
 static constexpr TCHAR RigVM_GetArrayElementSafeFormat[] = TEXT("GetArrayElementSafe<{0}>({1}, {2})");
 static constexpr TCHAR RigVM_InvokeEntryOpFormat[] = TEXT("\tif(InvokeEntryByName(EntryName_{0}) != ERigVMExecuteResult::Succeeded) return ERigVMExecuteResult::Failed;");
+static constexpr TCHAR RigVM_LazyEvalValueName[] =  TEXT("LazyValue_{0}_{1}");
+static constexpr TCHAR RigVM_LazyEvalLambdaDefine[] =  TEXT("\tconst TRigVMLazyValue<{2}> LazyValue_{0}_{1} = GetLazyValue<{2}>({0}, {3}, {4}_Ptr,\r\n\t\t[&]() -> ERigVMExecuteResult\r\n\t\t{");
+static constexpr TCHAR RigVM_LazyEvalLambdaReturn[] =  TEXT("\t\t\treturn ERigVMExecuteResult::Succeeded;\r\n\t\t}\r\n\t);");
+static constexpr TCHAR RigVM_LazyMemoryHandleInitFormat[] = TEXT("\tAllocateLazyMemoryHandles({0});");
+static constexpr TCHAR RigVM_SetupInstructionTrackingFormat[] = TEXT("\tSetupInstructionTracking({0});");
 
 FString FRigVMCodeGenerator::DumpIncludes(bool bLog)
 {
@@ -148,8 +170,13 @@ FString FRigVMCodeGenerator::DumpWrappedArrayTypes(bool bLog)
 
 FString FRigVMCodeGenerator::DumpExternalVariables(bool bForHeader, bool bLog)
 {
-
 	FStringArray Lines;
+
+	if(bForHeader)
+	{
+		Lines.Emplace();
+	}
+	
 	for(int32 ExternalVariableIndex = 0; ExternalVariableIndex < VM->GetExternalVariables().Num(); ExternalVariableIndex++)
 	{
 		const FRigVMOperand ExternalVarOperand(ERigVMMemoryType::External, ExternalVariableIndex, INDEX_NONE);
@@ -196,18 +223,44 @@ FString FRigVMCodeGenerator::DumpEntries(bool bForHeader, bool bLog)
 	return DumpLines(Lines, bLog);
 }
 
-FString FRigVMCodeGenerator::DumpProperties(bool bForHeader, int32 InOperationGroup, bool bLog)
+FString FRigVMCodeGenerator::DumpBlockNames(bool bForHeader, bool bLog)
+{
+	FStringArray Lines;
+
+	const FRigVMByteCode& ByteCode = VM->GetByteCode();
+	for(const FRigVMBranchInfo& BranchInfo : ByteCode.BranchInfos)
+	{
+		if(!BranchInfo.IsOutputBranch())
+		{
+			continue;
+		}
+		
+		const FString BlockName = BranchInfo.Label.ToString();
+		if(bForHeader)
+		{
+			Lines.AddUnique(Format(RigVM_DeclareBlockNameFormat, *BlockName));
+		}
+		else
+		{
+			Lines.AddUnique(Format(RigVM_DefineBlockNameFormat, *ClassName, *BlockName, *BlockName));
+		}
+	}
+
+	return DumpLines(Lines, bLog);
+}
+
+FString FRigVMCodeGenerator::DumpProperties(bool bForHeader, int32 InInstructionGroup, bool bLog)
 {
 	if(bForHeader)
 	{
 		// for headers we show all properties
-		check(InOperationGroup == INDEX_NONE);
+		check(InInstructionGroup == INDEX_NONE);
 	}
 
 	FStringArray Lines;
 	for(int32 Index = 0; Index < Properties.Num(); Index++)
 	{
-		if(!IsPropertyPartOfGroup(Index, InOperationGroup))
+		if(!IsPropertyPartOfGroup(Index, InInstructionGroup))
 		{
 			continue;
 		}
@@ -445,56 +498,111 @@ FString FRigVMCodeGenerator::DumpProperties(bool bForHeader, int32 InOperationGr
 	return DumpLines(Lines, bLog);
 }
 
-FString FRigVMCodeGenerator::DumpInstructions(int32 InOperationGroup, bool bLog)
+FString FRigVMCodeGenerator::DumpDispatches(bool bLog)
 {
-	const FRigVMByteCode& ByteCode = VM->GetByteCode();
-	const TArray<FName>& Functions = VM->GetFunctionNames();
-	const FRigVMInstructionArray Operations = ByteCode.GetInstructions();
-
 	FStringArray Lines;
 
-	const FOperationGroup& Group = GetGroup(InOperationGroup);
-
-	auto GetEntryParameters = [&]() -> FString
+	if(!Dispatches.IsEmpty())
 	{
-		const int32 EntryIndex = ByteCode.NumEntries() - 1;
-		const FRigVMByteCodeEntry& Entry = ByteCode.GetEntry(EntryIndex);
-		const FString EntryName = Entry.GetSanitizedName();
-		FString Parameters;
-
-		// find the entry's group and provide the needed arguments
-		for(const FOperationGroup& EntryGroup : OperationGroups)
+		for(auto Pair : Dispatches)
 		{
-			if(EntryGroup.Depth <= 0 && EntryGroup.Entry == EntryName)
+			const FRigVMFunction* Function = Pair.Value;
+			const FString FunctionName = Function->GetName();
+			const FRigVMDispatchFactory* Factory = Function->Factory;
+
+			TArray<FString> InputArguments, OutputArguments, MemoryHandles;
+			OutputArguments.Add(RigVM_ContextFormat);
+
+			for(const FRigVMFunctionArgument& Argument : Function->Arguments)
 			{
-				TArray<FString> ArgumentNames = {RigVM_ContextPublicFormat};
-				for(const FOpaqueArgument& OpaqueArgument : EntryGroup.OpaqueArguments)
+				const TRigVMTypeIndex TypeIndex = FRigVMRegistry::Get().GetTypeIndexFromCPPType(Argument.Type);
+				const FString DispatchArgument = RequiredUProperties.FindChecked(TypeIndex).Get<1>();
+				const ERigVMPinDirection Direction = Factory->GetTemplate()->FindArgument(Argument.Name)->GetDirection();
+				const FString ConstPrefix = (Direction == ERigVMPinDirection::Visible || Direction == ERigVMPinDirection::Input) ? TEXT("const ") : FString();
+
+				if(Factory->IsLazyInputArgument(Argument.Name))
 				{
-					ArgumentNames.Add(OpaqueArgument.Get<1>());
+					InputArguments.Add(Format(TEXT("{0}TRigVMLazyValue<{1}>& {2}"), ConstPrefix, Argument.Type, Argument.Name));
+					MemoryHandles.Add(Format(TEXT("{0}.GetMemoryHandle()"), Argument.Name));
+				}
+				else
+				{
+					InputArguments.Add(Format(TEXT("{0}{1}& {2}"), ConstPrefix, Argument.Type, Argument.Name));
+					MemoryHandles.Add(Format(TEXT("{(uint8*)&{0}, {1}_Ptr, nullptr}"), Argument.Name, *DispatchArgument));
 				}
-				Parameters = FString::Join(ArgumentNames, RigVM_CommaSeparator);
-				break;
 			}
+
+			OutputArguments.Add(TEXT("MemoryHandles"));
+			
+			Lines.Emplace();
+			Lines.Add(Format(RigVM_DispatchDeclarationFormat, Pair.Key, FString::Join(InputArguments, RigVM_CommaSeparator), *FunctionName));
+			Lines.Add(Format(TEXT("\t\tTArray<FRigVMMemoryHandle> MemoryHandles = {\r\n\t\t\t{0}\r\n\t\t};"), FString::Join(MemoryHandles, TEXT(",\r\n\t\t\t"))));
+			Lines.Add(Format(RigVM_InvokeDispatchPtrFormat, FString::Join(OutputArguments, RigVM_CommaSeparator)));
+			Lines.Add(TEXT("\t\treturn true;\r\n\t}"));
 		}
+	}
+	
+	return DumpLines(Lines, bLog); 
+}
 
-		if(!Parameters.IsEmpty())
+FString FRigVMCodeGenerator::DumpRequiredUProperties(bool bLog)
+{
+	FStringArray Lines;
+
+	if(!RequiredUProperties.IsEmpty())
+	{
+		for(auto Pair : RequiredUProperties)
 		{
-			Parameters = RigVM_CommaSeparator + Parameters;
+			Lines.Emplace();
+			Lines.Add(Format(RigVM_UPropertyDeclareFormat, Pair.Value.Get<0>(), Pair.Value.Get<1>()));
 		}
-		
-		return Parameters;
-	};
+
+		Lines.Emplace();
+
+		for(auto Pair : RequiredUProperties)
+		{
+			Lines.Add(Format(RigVM_UPropertyMemberFormat, Pair.Value.Get<1>()));
+		}
+	}
+	
+	return DumpLines(Lines, bLog); 
+}
+
+
+FString FRigVMCodeGenerator::DumpInitialize(bool bLog)
+{
+	FStringArray Lines;
+	Lines.Add(Format(RigVM_LazyMemoryHandleInitFormat, VM->GetByteCode().BranchInfos.Num()));
+
+	for(auto Pair : RequiredUProperties)
+	{
+		Lines.Add(Format(RigVM_UPropertyDefineFormat, *ClassName, Pair.Value.Get<1>()));
+	}
+	
+	// we'll add workstate initialization here later
+	Lines.Add(TEXT("\treturn true;"));
+	return DumpLines(Lines, bLog); 
+}
+
+FString FRigVMCodeGenerator::DumpInstructions(int32 InInstructionGroup, bool bLog)
+{
+	const FRigVMByteCode& ByteCode = VM->GetByteCode();
+	const FRigVMInstructionArray Instructions = ByteCode.GetInstructions();
+
+	FStringArray Lines;
+
+	const FInstructionGroup& Group = GetGroup(InInstructionGroup);
 
 	if(Group.Entry.IsEmpty())
 	{
-		if(InOperationGroup == INDEX_NONE)
+		if(InInstructionGroup == INDEX_NONE)
 		{
-			Lines.Add(Format(RigVM_SetExecuteContextStructFormmt, ExecuteContextType));
-			Lines.Add(Format(RigVM_UpdateContextFormat, ExecuteContextType, Operations.Num()));
+			Lines.Add(Format(RigVM_SetExecuteContextStructFormat, ExecuteContextType));
+			Lines.Add(Format(RigVM_UpdateContextFormat, ExecuteContextType, Instructions.Num()));
 
 			// let's get the additional arguments for all sub groups
 			TArray<FString> OpaqueArgumentHit;
-			for(const FOperationGroup& EntryGroup : OperationGroups)
+			for(const FInstructionGroup& EntryGroup : InstructionGroups)
 			{
 				if(EntryGroup.Entry.IsEmpty() || EntryGroup.Depth > 0)
 				{
@@ -536,7 +644,7 @@ FString FRigVMCodeGenerator::DumpInstructions(int32 InOperationGroup, bool bLog)
 				FString OpaqueArguments;
 
 				// find the entry's group and provide the needed arguments
-				for(const FOperationGroup& EntryGroup : OperationGroups)
+				for(const FInstructionGroup& EntryGroup : InstructionGroups)
 				{
 					if(EntryGroup.Depth <= 0 && EntryGroup.Entry == EntryName)
 					{
@@ -554,7 +662,7 @@ FString FRigVMCodeGenerator::DumpInstructions(int32 InOperationGroup, bool bLog)
 		}
 	}
 
-	if(Group.Entry.IsEmpty() && OperationGroups.Num() > 0)
+	if(Group.Entry.IsEmpty() && InstructionGroups.Num() > 0)
 	{
 		Lines.Add(RigVM_ReturnFailedFormat);
 		return DumpLines(Lines, bLog);
@@ -562,423 +670,563 @@ FString FRigVMCodeGenerator::DumpInstructions(int32 InOperationGroup, bool bLog)
 
 	if(Group.ChildGroups.IsEmpty())
 	{
-		// loop over all of the operations once more
-		for(int32 OperationIndex = Group.First; OperationIndex <= Group.Last; OperationIndex++)
+		TArray<int32> InstructionIndices = GetInstructionIndicesFromRange(Group.First, Group.Last);
+
+		// dump all lambdas and remove those instructions from instructions to process
+		for(const FRigVMBranchInfo& BranchInfo : ByteCode.BranchInfos)
 		{
-			// inject a label if required
-			if (Group.RequiredLabels.Contains(OperationIndex))
+			if(BranchInfo.IsOutputBranch())
+			{
+				continue;
+			}
+			
+			if((int32)BranchInfo.FirstInstruction < Group.First ||
+				(int32)BranchInfo.LastInstruction > Group.Last)
 			{
-				Lines.Add(Format(RigVM_InstructionLabelFormat, OperationIndex));
+				continue;
 			}
 
-			Lines.Add(Format(RigVM_SetInstructionIndexFormat, OperationIndex));
+			InstructionIndices.RemoveAll([BranchInfo](int32 Index) {
+				return FMath::IsWithinInclusive<uint16>((uint16)Index, BranchInfo.FirstInstruction, BranchInfo.LastInstruction);
+			});
+
+			// find the operand this lazy eval lambda belongs to
+			const FRigVMOperandArray Operands = ByteCode.GetOperandsForOp(Instructions[BranchInfo.InstructionIndex]);
+			check(Operands.IsValidIndex(BranchInfo.ArgumentIndex));
+			const FRigVMOperand Operand = Operands[BranchInfo.ArgumentIndex];
+			const FString OperandName = GetOperandName(Operand, false, true);
+			const FString OperandCPPType = GetOperandCPPType(Operand);
+			const TRigVMTypeIndex OperandTypeIndex = FRigVMRegistry::Get().GetTypeIndexFromCPPType(OperandCPPType);
+			const FString PropertyName = RequiredUProperties.FindChecked(OperandTypeIndex).Get<1>();
+
+			// dump the instructions for the lambda wrapped with the lambda definition
+			Lines.Add(Format(RigVM_LazyEvalLambdaDefine, BranchInfo.Index, BranchInfo.Label.ToString(), *OperandCPPType, *OperandName, *PropertyName));
+			Lines.Add(DumpInstructions(TEXT("\t\t"), (int32)BranchInfo.FirstInstruction, (int32)BranchInfo.LastInstruction, Group, false));
+			Lines.Add(RigVM_LazyEvalLambdaReturn);
+			Lines.Emplace();
+
+			OverriddenOperatorNames.Add(OperandName, Format(RigVM_LazyEvalValueName, BranchInfo.Index, BranchInfo.Label.ToString()));
+		}
+		
+		// dump the remaining instruction indices
+		Lines.Add(DumpInstructions(FString(), InstructionIndices, Group, false));
+	}
+	else
+	{
+		// we have child groups - we need to invoke those
+		for(int32 ChildGroupIndex : Group.ChildGroups)
+		{
+			const FInstructionGroup& ChildGroup = InstructionGroups[ChildGroupIndex];
 
-			const FRigVMInstruction& Operation = Operations[OperationIndex];
-			switch(Operation.OpCode)
+			FString Parameters;
+			TArray<FString> ParameterArray = {RigVM_ContextPublicFormat};
+			if(!ChildGroup.OpaqueArguments.IsEmpty())
 			{
-				case ERigVMOpCode::Execute_0_Operands:
-				case ERigVMOpCode::Execute_1_Operands:
-				case ERigVMOpCode::Execute_2_Operands:
-				case ERigVMOpCode::Execute_3_Operands:
-				case ERigVMOpCode::Execute_4_Operands:
-				case ERigVMOpCode::Execute_5_Operands:
-				case ERigVMOpCode::Execute_6_Operands:
-				case ERigVMOpCode::Execute_7_Operands:
-				case ERigVMOpCode::Execute_8_Operands:
-				case ERigVMOpCode::Execute_9_Operands:
-				case ERigVMOpCode::Execute_10_Operands:
-				case ERigVMOpCode::Execute_11_Operands:
-				case ERigVMOpCode::Execute_12_Operands:
-				case ERigVMOpCode::Execute_13_Operands:
-				case ERigVMOpCode::Execute_14_Operands:
-				case ERigVMOpCode::Execute_15_Operands:
-				case ERigVMOpCode::Execute_16_Operands:
-				case ERigVMOpCode::Execute_17_Operands:
-				case ERigVMOpCode::Execute_18_Operands:
-				case ERigVMOpCode::Execute_19_Operands:
-				case ERigVMOpCode::Execute_20_Operands:
-				case ERigVMOpCode::Execute_21_Operands:
-				case ERigVMOpCode::Execute_22_Operands:
-				case ERigVMOpCode::Execute_23_Operands:
-				case ERigVMOpCode::Execute_24_Operands:
-				case ERigVMOpCode::Execute_25_Operands:
-				case ERigVMOpCode::Execute_26_Operands:
-				case ERigVMOpCode::Execute_27_Operands:
-				case ERigVMOpCode::Execute_28_Operands:
-				case ERigVMOpCode::Execute_29_Operands:
-				case ERigVMOpCode::Execute_30_Operands:
-				case ERigVMOpCode::Execute_31_Operands:
-				case ERigVMOpCode::Execute_32_Operands:
-				case ERigVMOpCode::Execute_33_Operands:
-				case ERigVMOpCode::Execute_34_Operands:
-				case ERigVMOpCode::Execute_35_Operands:
-				case ERigVMOpCode::Execute_36_Operands:
-				case ERigVMOpCode::Execute_37_Operands:
-				case ERigVMOpCode::Execute_38_Operands:
-				case ERigVMOpCode::Execute_39_Operands:
-				case ERigVMOpCode::Execute_40_Operands:
-				case ERigVMOpCode::Execute_41_Operands:
-				case ERigVMOpCode::Execute_42_Operands:
-				case ERigVMOpCode::Execute_43_Operands:
-				case ERigVMOpCode::Execute_44_Operands:
-				case ERigVMOpCode::Execute_45_Operands:
-				case ERigVMOpCode::Execute_46_Operands:
-				case ERigVMOpCode::Execute_47_Operands:
-				case ERigVMOpCode::Execute_48_Operands:
-				case ERigVMOpCode::Execute_49_Operands:
-				case ERigVMOpCode::Execute_50_Operands:
-				case ERigVMOpCode::Execute_51_Operands:
-				case ERigVMOpCode::Execute_52_Operands:
-				case ERigVMOpCode::Execute_53_Operands:
-				case ERigVMOpCode::Execute_54_Operands:
-				case ERigVMOpCode::Execute_55_Operands:
-				case ERigVMOpCode::Execute_56_Operands:
-				case ERigVMOpCode::Execute_57_Operands:
-				case ERigVMOpCode::Execute_58_Operands:
-				case ERigVMOpCode::Execute_59_Operands:
-				case ERigVMOpCode::Execute_60_Operands:
-				case ERigVMOpCode::Execute_61_Operands:
-				case ERigVMOpCode::Execute_62_Operands:
-				case ERigVMOpCode::Execute_63_Operands:
-				case ERigVMOpCode::Execute_64_Operands:
+				for(const FOpaqueArgument& OpaqueArgument : ChildGroup.OpaqueArguments)
 				{
-					const FRigVMExecuteOp& Op = ByteCode.GetOpAt<FRigVMExecuteOp>(Operation);
-					FRigVMOperandArray Operands = ByteCode.GetOperandsForExecuteOp(Operation);
+					ParameterArray.Add(OpaqueArgument.Get<1>());
+				}
+			}
+			Parameters = FString::Join(ParameterArray, RigVM_CommaSeparator);
+			Lines.Add(Format(RigVM_InvokeExecuteGroupFormat, *ChildGroup.Entry, ChildGroupIndex, *Parameters));
+		}
+	}
+	
+	Lines.Emplace();
+	if(Group.Depth <= 0)
+	{
+		Lines.Add(RigVM_ExecuteReachedExitFormat);
+	}
+	Lines.Add(RigVM_ReturnSucceededFormat);
 
-					const FRigVMFunction* Function = FRigVMRegistry::Get().FindFunction(*Functions[Op.FunctionIndex].ToString());
-					check(Function);
-						
-					FStringArray Arguments;
-					Arguments.Add(RigVM_ContextPublicFormat);
-					for(int32 OperandIndex = 0; OperandIndex < Operands.Num(); OperandIndex++)
-					{
-						const FRigVMOperand& Operand = Operands[OperandIndex];
-						bool bSliced = false;
-						const FRigVMPropertyDescription& Property = GetPropertyForOperand(Operand);
-						if (RigVMTypeUtils::IsArrayType(Property.CPPType))
-						{
-							const FRigVMFunctionArgument& FunctionArgument = Function->GetArguments()[OperandIndex];
-							bSliced = FunctionArgument.Type != Property.CPPType;
-						}
-						Arguments.Add(GetOperandName(Operand, bSliced));
-					}
+	return DumpLines(Lines, bLog);
+}
 
-					for(const FRigVMFunctionArgument& Argument : Function->GetArguments())
-					{
-						if (Function->IsAdditionalArgument(Argument))
-						{
-							const FOpaqueArgument* OpaqueArgument = Group.OpaqueArguments.FindByPredicate([Argument](const FOpaqueArgument& InArgument) -> bool
-							{
-								return InArgument.Get<0>() == Argument.Name;
-							});
+FString FRigVMCodeGenerator::DumpInstructions(const FString& InPrefix, int32 InFirstInstruction, int32 InLastInstruction, const FInstructionGroup& InGroup, bool bLog)
+{
+	return DumpInstructions(InPrefix, GetInstructionIndicesFromRange(InFirstInstruction, InLastInstruction), InGroup, bLog);
+}
 
-							check(OpaqueArgument);
-							Arguments.Add(OpaqueArgument->Get<1>());
-						}
-					}
+FString FRigVMCodeGenerator::DumpInstructions(const FString& InPrefix, const TArray<int32> InInstructionIndices, const FInstructionGroup& InGroup, bool bLog)
+{
+	const FRigVMByteCode& ByteCode = VM->GetByteCode();
+	const TArray<FName>& Functions = VM->GetFunctionNames();
+	const FRigVMInstructionArray Instructions = ByteCode.GetInstructions();
 
-					const FString JoinedArguments = FString::Join(Arguments, RigVM_CommaSeparator);
-					Lines.Add(Format(RigVM_CallExternOpFormat, *Function->Struct->GetStructCPPName(), *Function->GetMethodName().ToString(), *JoinedArguments));
-					break;
-				}
-				case ERigVMOpCode::Zero:
-				{
-					const FRigVMUnaryOp& Op = ByteCode.GetOpAt<FRigVMUnaryOp>(Operation);
-					Lines.Add(Format(RigVM_ZeroOpFormat, *GetOperandName(Op.Arg, false)));
-					break;
-				}
-				case ERigVMOpCode::BoolFalse:
-				{
-					const FRigVMUnaryOp& Op = ByteCode.GetOpAt<FRigVMUnaryOp>(Operation);
-					Lines.Add(Format(RigVM_BoolFalseOpFormat, *GetOperandName(Op.Arg, false)));
-					break;
-				}
-				case ERigVMOpCode::BoolTrue:
+	FString Prefix = InPrefix;
+	
+	TArray<FString> Lines;
+	for(int32 InstructionIndex : InInstructionIndices)
+	{
+		// inject a label if required
+		if (InGroup.RequiredLabels.Contains(InstructionIndex))
+		{
+			// check if the last line was a jump to this label
+			if(Lines.Last().Contains(Format(RigVM_JumpOpFormat, InstructionIndex)))
+			{
+				Lines.Pop();
+			}
+			else
+			{
+				Lines.Add(Prefix + Format(RigVM_InstructionLabelFormat, InstructionIndex));
+			}
+		}
+
+		while(!Lines.IsEmpty() && Lines.Last().Contains(TEXT("SetInstructionIndex")))
+		{
+			Lines.Pop();
+		}
+		Lines.Add(Prefix + Format(RigVM_SetInstructionIndexFormat, InstructionIndex));
+
+		const FRigVMInstruction& Instruction = Instructions[InstructionIndex];
+		switch(Instruction.OpCode)
+		{
+			case ERigVMOpCode::Execute_0_Operands:
+			case ERigVMOpCode::Execute_1_Operands:
+			case ERigVMOpCode::Execute_2_Operands:
+			case ERigVMOpCode::Execute_3_Operands:
+			case ERigVMOpCode::Execute_4_Operands:
+			case ERigVMOpCode::Execute_5_Operands:
+			case ERigVMOpCode::Execute_6_Operands:
+			case ERigVMOpCode::Execute_7_Operands:
+			case ERigVMOpCode::Execute_8_Operands:
+			case ERigVMOpCode::Execute_9_Operands:
+			case ERigVMOpCode::Execute_10_Operands:
+			case ERigVMOpCode::Execute_11_Operands:
+			case ERigVMOpCode::Execute_12_Operands:
+			case ERigVMOpCode::Execute_13_Operands:
+			case ERigVMOpCode::Execute_14_Operands:
+			case ERigVMOpCode::Execute_15_Operands:
+			case ERigVMOpCode::Execute_16_Operands:
+			case ERigVMOpCode::Execute_17_Operands:
+			case ERigVMOpCode::Execute_18_Operands:
+			case ERigVMOpCode::Execute_19_Operands:
+			case ERigVMOpCode::Execute_20_Operands:
+			case ERigVMOpCode::Execute_21_Operands:
+			case ERigVMOpCode::Execute_22_Operands:
+			case ERigVMOpCode::Execute_23_Operands:
+			case ERigVMOpCode::Execute_24_Operands:
+			case ERigVMOpCode::Execute_25_Operands:
+			case ERigVMOpCode::Execute_26_Operands:
+			case ERigVMOpCode::Execute_27_Operands:
+			case ERigVMOpCode::Execute_28_Operands:
+			case ERigVMOpCode::Execute_29_Operands:
+			case ERigVMOpCode::Execute_30_Operands:
+			case ERigVMOpCode::Execute_31_Operands:
+			case ERigVMOpCode::Execute_32_Operands:
+			case ERigVMOpCode::Execute_33_Operands:
+			case ERigVMOpCode::Execute_34_Operands:
+			case ERigVMOpCode::Execute_35_Operands:
+			case ERigVMOpCode::Execute_36_Operands:
+			case ERigVMOpCode::Execute_37_Operands:
+			case ERigVMOpCode::Execute_38_Operands:
+			case ERigVMOpCode::Execute_39_Operands:
+			case ERigVMOpCode::Execute_40_Operands:
+			case ERigVMOpCode::Execute_41_Operands:
+			case ERigVMOpCode::Execute_42_Operands:
+			case ERigVMOpCode::Execute_43_Operands:
+			case ERigVMOpCode::Execute_44_Operands:
+			case ERigVMOpCode::Execute_45_Operands:
+			case ERigVMOpCode::Execute_46_Operands:
+			case ERigVMOpCode::Execute_47_Operands:
+			case ERigVMOpCode::Execute_48_Operands:
+			case ERigVMOpCode::Execute_49_Operands:
+			case ERigVMOpCode::Execute_50_Operands:
+			case ERigVMOpCode::Execute_51_Operands:
+			case ERigVMOpCode::Execute_52_Operands:
+			case ERigVMOpCode::Execute_53_Operands:
+			case ERigVMOpCode::Execute_54_Operands:
+			case ERigVMOpCode::Execute_55_Operands:
+			case ERigVMOpCode::Execute_56_Operands:
+			case ERigVMOpCode::Execute_57_Operands:
+			case ERigVMOpCode::Execute_58_Operands:
+			case ERigVMOpCode::Execute_59_Operands:
+			case ERigVMOpCode::Execute_60_Operands:
+			case ERigVMOpCode::Execute_61_Operands:
+			case ERigVMOpCode::Execute_62_Operands:
+			case ERigVMOpCode::Execute_63_Operands:
+			case ERigVMOpCode::Execute_64_Operands:
+			{
+				const FRigVMExecuteOp& Op = ByteCode.GetOpAt<FRigVMExecuteOp>(Instruction);
+				FRigVMOperandArray Operands = ByteCode.GetOperandsForExecuteOp(Instruction);
+
+				const FRigVMFunction* Function = FRigVMRegistry::Get().FindFunction(*Functions[Op.FunctionIndex].ToString());
+				check(Function);
+					
+				FStringArray Arguments;
+				if(Function->Struct)
 				{
-					const FRigVMUnaryOp& Op = ByteCode.GetOpAt<FRigVMUnaryOp>(Operation);
-					Lines.Add(Format(RigVM_BoolTrueFormat, *GetOperandName(Op.Arg, false)));
-					break;
+					Arguments.Add(RigVM_ContextPublicFormat);
 				}
-				case ERigVMOpCode::Copy:
-				case ERigVMOpCode::ArrayClone:
+				for(int32 OperandIndex = 0; OperandIndex < Operands.Num(); OperandIndex++)
 				{
-					const FRigVMBinaryOp& Op = ByteCode.GetOpAt<FRigVMBinaryOp>(Operation);
-					const FString TargetOperand = GetOperandName(Op.ArgB, false, false);
-					const FString SourceOperand = GetOperandName(Op.ArgA, false, true);
-					const FString TargetCPPType = GetOperandCPPType(Op.ArgB);
-					const FString SourceCPPType = GetOperandCPPType(Op.ArgA);
-
-					if(RigVMTypeUtils::IsArrayType(TargetCPPType) &&
-						RigVMTypeUtils::IsArrayType(SourceCPPType) &&
-						TargetCPPType != SourceCPPType)
-					{
-						const FString TargetBaseCPPType = RigVMTypeUtils::BaseTypeFromArrayType(TargetCPPType);
-						const FString SourceBaseCPPType = RigVMTypeUtils::BaseTypeFromArrayType(SourceCPPType);
-						Lines.Add(Format(RigVM_CopyUnrelatedArraysFormat, *TargetBaseCPPType, *SourceBaseCPPType, *TargetOperand, *SourceOperand));
-					}
-					else
+					const FRigVMOperand& Operand = Operands[OperandIndex];
+					bool bSliced = false;
+					const FRigVMPropertyDescription& Property = GetPropertyForOperand(Operand);
+					if (RigVMTypeUtils::IsArrayType(Property.CPPType))
 					{
-						const FString CastPrefix = TargetCPPType != SourceCPPType ? Format(RigVM_BracesFormat, *TargetCPPType) : FString();  
-						
-						if(TargetOperand.EndsWith(TEXT("(")))
-						{
-							Lines.Add(Format(RigVM_CopyOpMethodFormat, *TargetOperand, *CastPrefix, *SourceOperand));
-						}
-						else
-						{
-							Lines.Add(Format(RigVM_CopyOpAssignFormat, *TargetOperand, *CastPrefix, *SourceOperand));
-						}
+						const FRigVMFunctionArgument& FunctionArgument = Function->GetArguments()[OperandIndex];
+						bSliced = FunctionArgument.Type != Property.CPPType;
 					}
-					break;
-				}
-				case ERigVMOpCode::Increment:
-				{
-					const FRigVMUnaryOp& Op = ByteCode.GetOpAt<FRigVMUnaryOp>(Operation);
-					Lines.Add(Format(RigVM_IncrementOpFormat, *GetOperandName(Op.Arg, false)));
-					break;
-				}
-				case ERigVMOpCode::Decrement:
-				{
-					const FRigVMUnaryOp& Op = ByteCode.GetOpAt<FRigVMUnaryOp>(Operation);
-					Lines.Add(Format(RigVM_DecrementOpFormat, *GetOperandName(Op.Arg, false)));
-					break;
-				}
-				case ERigVMOpCode::Equals:
-				{
-					const FRigVMComparisonOp& Op = ByteCode.GetOpAt<FRigVMComparisonOp>(Operation);
-					Lines.Add(Format(RigVM_EqualsOpFormat, *GetOperandName(Op.Result, false), *GetOperandName(Op.B, false), *GetOperandName(Op.B, false)));
-					break;
-				}
-				case ERigVMOpCode::NotEquals:
-				{
-					const FRigVMComparisonOp& Op = ByteCode.GetOpAt<FRigVMComparisonOp>(Operation);
-					Lines.Add(Format(RigVM_NotEqualsOpFormat, *GetOperandName(Op.Result, false), *GetOperandName(Op.B, false), *GetOperandName(Op.B, false)));
-					break;
+					Arguments.Add(GetOperandName(Operand, bSliced));
 				}
-				case ERigVMOpCode::JumpAbsolute:
-				{
-					const FRigVMJumpOp& Op = ByteCode.GetOpAt<FRigVMJumpOp>(Operation);
-					Lines.Add(Format(RigVM_JumpOpFormat, Op.InstructionIndex));
-					break;
-				}
-				case ERigVMOpCode::JumpForward:
-				{
-					const FRigVMJumpOp& Op = ByteCode.GetOpAt<FRigVMJumpOp>(Operation);
-					Lines.Add(Format(RigVM_JumpOpFormat, OperationIndex + Op.InstructionIndex));
-					break;
-				}
-				case ERigVMOpCode::JumpBackward:
-				{
-					const FRigVMJumpOp& Op = ByteCode.GetOpAt<FRigVMJumpOp>(Operation);
-					Lines.Add(Format(RigVM_JumpOpFormat, OperationIndex - Op.InstructionIndex));
-					break;
-				}
-				case ERigVMOpCode::JumpAbsoluteIf:
-				{
-					const FRigVMJumpIfOp& Op = ByteCode.GetOpAt<FRigVMJumpIfOp>(Operation);
-					const FString& Condition = Op.Condition ? RigVM_TrueFormat : RigVM_FalseFormat;
-					Lines.Add(Format(RigVM_JumpIfOpFormat, *GetOperandName(Op.Arg, false), *Condition, Op.InstructionIndex));
-					break;
-				}
-				case ERigVMOpCode::JumpForwardIf:
-				{
-					const FRigVMJumpIfOp& Op = ByteCode.GetOpAt<FRigVMJumpIfOp>(Operation);
-					const FString& Condition = Op.Condition ? RigVM_TrueFormat : RigVM_FalseFormat;
-					Lines.Add(Format(RigVM_JumpIfOpFormat, *GetOperandName(Op.Arg, false), *Condition, OperationIndex + Op.InstructionIndex));
-					break;
-				}
-				case ERigVMOpCode::JumpBackwardIf:
-				{
-					const FRigVMJumpIfOp& Op = ByteCode.GetOpAt<FRigVMJumpIfOp>(Operation);
-					const FString& Condition = Op.Condition ? RigVM_TrueFormat : RigVM_FalseFormat;
-					Lines.Add(Format(RigVM_JumpIfOpFormat, *GetOperandName(Op.Arg, false), *Condition, OperationIndex - Op.InstructionIndex));
-					break;
-				}
-				case ERigVMOpCode::Exit:
+
+				for(const FRigVMFunctionArgument& Argument : Function->GetArguments())
 				{
-					if(OperationIndex != Group.Last)
+					if (Function->IsAdditionalArgument(Argument))
 					{
-						Lines.Add(Format(RigVM_JumpOpFormat, Operations.Num()));
-					}
-					break;
-				}
-				case ERigVMOpCode::BeginBlock:
-				{
-					const FRigVMBinaryOp& Op = ByteCode.GetOpAt<FRigVMBinaryOp>(Operation);
-					Lines.Add(Format(RigVM_BeginBlockOpFormat, *GetOperandName(Op.ArgA, false), *GetOperandName(Op.ArgB, false)));
-					break;
-				}
-				case ERigVMOpCode::EndBlock:
-				{
-					Lines.Add(RigVM_EndBlockOpFormat);
-					break;
-				}
-				case ERigVMOpCode::ArrayReset:
-				{
-					const FRigVMUnaryOp& Op = ByteCode.GetOpAt<FRigVMUnaryOp>(Operation);
-					Lines.Add(Format(RigVM_ArrayResetOpFormat, *GetOperandName(Op.Arg, false)));
-					break;
-				}
-				case ERigVMOpCode::ArrayGetNum:
-				{
-					const FRigVMBinaryOp& Op = ByteCode.GetOpAt<FRigVMBinaryOp>(Operation);
-					Lines.Add(Format(RigVM_ArrayGetNumOpFormat, *GetOperandName(Op.ArgB, false), *GetOperandName(Op.ArgA, false)));
-					break;
-				} 
-				case ERigVMOpCode::ArraySetNum:
-				{
-					const FRigVMBinaryOp& Op = ByteCode.GetOpAt<FRigVMBinaryOp>(Operation);
-					const FString ConditionString = Format(RigVM_IsValidArraySizeFormat, *GetOperandName(Op.ArgB, false));
-					Lines.Add(Format(RigVM_ArraySetNumOpFormat, *ConditionString, *GetOperandName(Op.ArgA, false), *GetOperandName(Op.ArgB, false)));
-					break;
-				}
-				case ERigVMOpCode::ArrayGetAtIndex:
-				{
-					const FRigVMTernaryOp& Op = ByteCode.GetOpAt<FRigVMTernaryOp>(Operation);
-					const FString ConditionString = Format(RigVM_IsValidArrayIndexFormat, *GetOperandCPPBaseType(Op.ArgA), *GetOperandName(Op.ArgA, false));
-					Lines.Add(Format(RigVM_TemporaryArrayIndexFormat, *GetOperandName(Op.ArgB, false)));
-					Lines.Add(Format(RigVM_ArrayGetAtIndexOpFormat, *ConditionString, *GetOperandName(Op.ArgC, false), *GetOperandName(Op.ArgA, false)));
-					break;
-				}  
-				case ERigVMOpCode::ArraySetAtIndex:
-				{
-					const FRigVMTernaryOp& Op = ByteCode.GetOpAt<FRigVMTernaryOp>(Operation);
-					const FString ConditionString = Format(RigVM_IsValidArrayIndexFormat,  *GetOperandCPPBaseType(Op.ArgA), *GetOperandName(Op.ArgA, false));
-					Lines.Add(Format(RigVM_TemporaryArrayIndexFormat, *GetOperandName(Op.ArgB, false)));
-					Lines.Add(Format(RigVM_ArraySetAtIndexOpFormat, *ConditionString, *GetOperandName(Op.ArgA, false), *GetOperandName(Op.ArgC, false), *GetOperandName(Op.ArgC, false)));
-					break;
-				}
-				case ERigVMOpCode::ArrayAdd:
-				{
-					const FRigVMTernaryOp& Op = ByteCode.GetOpAt<FRigVMTernaryOp>(Operation);
-					const FString NumString = Format(RigVM_ArrayNumPlusOneFormat, *GetOperandName(Op.ArgA, false));
-					const FString ConditionString = Format(RigVM_IsValidArraySizeFormat, *NumString);
-					Lines.Add(Format(RigVM_ArrayAddOpFormat, *ConditionString, *GetOperandName(Op.ArgC, false), *GetOperandName(Op.ArgA, false), *GetOperandName(Op.ArgB, false)));
-					break;
-				}
-				case ERigVMOpCode::ArrayInsert:
-				{
-					const FRigVMTernaryOp& Op = ByteCode.GetOpAt<FRigVMTernaryOp>(Operation);
-					const FString NumString = Format(RigVM_ArrayNumPlusOneFormat, *GetOperandName(Op.ArgA, false));
-					const FString ConditionString = Format(RigVM_IsValidArraySizeFormat, *NumString);
-					Lines.Add(Format(RigVM_ArrayInsertOpFormat, *ConditionString, *GetOperandName(Op.ArgA, false), *GetOperandName(Op.ArgC, false), *GetOperandName(Op.ArgB, false)));
-					break;
-				}
-				case ERigVMOpCode::ArrayRemove:
-				{
-					const FRigVMBinaryOp& Op = ByteCode.GetOpAt<FRigVMBinaryOp>(Operation);
-					Lines.Add(Format(RigVM_ArrayRemoveOpFormat, *GetOperandName(Op.ArgA, false), *GetOperandName(Op.ArgB, false)));
-					break;
+						const FOpaqueArgument* OpaqueArgument = InGroup.OpaqueArguments.FindByPredicate([Argument](const FOpaqueArgument& InArgument) -> bool
+						{
+							return InArgument.Get<0>() == Argument.Name;
+						});
+
+						check(OpaqueArgument);
+						Arguments.Add(OpaqueArgument->Get<1>());
+					}
 				}
-				case ERigVMOpCode::ArrayFind:
+
+				const FString JoinedArguments = FString::Join(Arguments, RigVM_CommaSeparator);
+				if(Function->Struct)
 				{
-					const FRigVMQuaternaryOp& Op = ByteCode.GetOpAt<FRigVMQuaternaryOp>(Operation);
-					const FString ArrayElementCPPType = RigVMTypeUtils::BaseTypeFromArrayType(GetOperandCPPBaseType(Op.ArgA));
-					const FString ElementCPPType = GetOperandCPPType(Op.ArgB);
-					const FString TemplateSuffix = ArrayElementCPPType == ElementCPPType ? Format(RigVM_TemplateOneArgFormat, *ElementCPPType) : FString();
-					Lines.Add(Format(RigVM_ArrayFindOpFormat, *GetOperandName(Op.ArgD, false), *TemplateSuffix, *GetOperandName(Op.ArgA, false), *GetOperandName(Op.ArgB, false), *GetOperandName(Op.ArgC, false)));
-					break;
+					Lines.Add(Prefix + Format(RigVM_CallExternOpFormat, *Function->Struct->GetStructCPPName(), *Function->GetMethodName().ToString(), *JoinedArguments));
 				}
-				case ERigVMOpCode::ArrayAppend:
+				else if(Function->Factory)
 				{
-					const FRigVMBinaryOp& Op = ByteCode.GetOpAt<FRigVMBinaryOp>(Operation);
-					const FString NumString = Format(RigVM_ArrayAppendOpNumFormat, *GetOperandName(Op.ArgA, false), *GetOperandName(Op.ArgB, false));
-					const FString ConditionString = Format(RigVM_IsValidArraySizeFormat, *NumString);
-					Lines.Add(Format(RigVM_ArrayAppendOpFormat, *ConditionString, *GetOperandName(Op.ArgA, false), *GetOperandName(Op.ArgB, false)));
-					break;
+					FString DispatchName;
+					for(auto Pair : Dispatches)
+					{
+						if(Pair.Value == Function)
+						{
+							DispatchName = Pair.Key;
+							break;
+						}
+					}
+					check(!DispatchName.IsEmpty());
+					Lines.Add(Prefix + Format(RigVM_InvokeDispatchFormat, *DispatchName, *JoinedArguments));
 				}
-				case ERigVMOpCode::ArrayIterator:
+				else
 				{
-					const FRigVMSenaryOp& Op = ByteCode.GetOpAt<FRigVMSenaryOp>(Operation);
-					const FString ArrayElementCPPType = RigVMTypeUtils::BaseTypeFromArrayType(GetOperandCPPType(Op.ArgA));
-					const FString ElementCPPType = GetOperandCPPType(Op.ArgB);
-					const FString TemplateSuffix = ArrayElementCPPType == ElementCPPType ? Format(RigVM_TemplateOneArgFormat, *ElementCPPType) : FString();
-					Lines.Add(Format(RigVM_ArrayIteratorOpFormat, *GetOperandName(Op.ArgF, false), *TemplateSuffix, *GetOperandName(Op.ArgA, false), *GetOperandName(Op.ArgB, false), *GetOperandName(Op.ArgC, false), *GetOperandName(Op.ArgD, false), *GetOperandName(Op.ArgE, false)));
-					break;
+					checkNoEntry();
 				}
-				case ERigVMOpCode::ArrayUnion:
+				break;
+			}
+			case ERigVMOpCode::Zero:
+			{
+				const FRigVMUnaryOp& Op = ByteCode.GetOpAt<FRigVMUnaryOp>(Instruction);
+				const FProperty* Property = VM->GetWorkMemory()->GetProperty(Op.Arg.GetRegisterIndex());
+				if(Property->IsA<FIntProperty>())
 				{
-					const FRigVMBinaryOp& Op = ByteCode.GetOpAt<FRigVMBinaryOp>(Operation);
-					const FString ElementCPPTypeA = RigVMTypeUtils::BaseTypeFromArrayType(GetOperandCPPType(Op.ArgA));
-					const FString ElementCPPTypeB = RigVMTypeUtils::BaseTypeFromArrayType(GetOperandCPPType(Op.ArgA));
-					const FString TemplateSuffix = ElementCPPTypeA == ElementCPPTypeB ? Format(RigVM_TemplateOneArgFormat, *ElementCPPTypeA) : FString();
-					Lines.Add(Format(RigVM_ArrayUnionOpFormat, *TemplateSuffix, *GetOperandName(Op.ArgA, false, false), *GetOperandName(Op.ArgB, false, false)));
-					break;
+					Lines.Add(Prefix + Format(RigVM_ZeroOpIntFormat, *GetOperandName(Op.Arg, false)));
 				}
-				case ERigVMOpCode::ArrayDifference:
+				else if(Property->IsA<FNameProperty>())
 				{
-					const FRigVMBinaryOp& Op = ByteCode.GetOpAt<FRigVMBinaryOp>(Operation);
-					const FString ElementCPPTypeA = RigVMTypeUtils::BaseTypeFromArrayType(GetOperandCPPType(Op.ArgA));
-					const FString ElementCPPTypeB = RigVMTypeUtils::BaseTypeFromArrayType(GetOperandCPPType(Op.ArgA));
-					const FString TemplateSuffix = ElementCPPTypeA == ElementCPPTypeB ? Format(RigVM_TemplateOneArgFormat, *ElementCPPTypeA) : FString();
-					Lines.Add(Format(RigVM_ArrayDifferenceOpFormat, *TemplateSuffix, *GetOperandName(Op.ArgA, false), *GetOperandName(Op.ArgB, false)));
-					break;
+					Lines.Add(Prefix + Format(RigVM_ZeroOpNameFormat, *GetOperandName(Op.Arg, false)));
 				}
-				case ERigVMOpCode::ArrayIntersection:
+				else
 				{
-					const FRigVMBinaryOp& Op = ByteCode.GetOpAt<FRigVMBinaryOp>(Operation);
-					const FString ElementCPPTypeA = RigVMTypeUtils::BaseTypeFromArrayType(GetOperandCPPType(Op.ArgA));
-					const FString ElementCPPTypeB = RigVMTypeUtils::BaseTypeFromArrayType(GetOperandCPPType(Op.ArgA));
-					const FString TemplateSuffix = ElementCPPTypeA == ElementCPPTypeB ? Format(RigVM_TemplateOneArgFormat, *ElementCPPTypeA) : FString();
-					Lines.Add(Format(RigVM_ArrayIntersectionOpFormat, *TemplateSuffix, *GetOperandName(Op.ArgA, false), *GetOperandName(Op.ArgB, false)));
-					break;
+					checkNoEntry();
 				}
-				case ERigVMOpCode::ArrayReverse:
+				break;
+			}
+			case ERigVMOpCode::BoolFalse:
+			{
+				const FRigVMUnaryOp& Op = ByteCode.GetOpAt<FRigVMUnaryOp>(Instruction);
+				Lines.Add(Prefix + Format(RigVM_BoolFalseOpFormat, *GetOperandName(Op.Arg, false)));
+				break;
+			}
+			case ERigVMOpCode::BoolTrue:
+			{
+				const FRigVMUnaryOp& Op = ByteCode.GetOpAt<FRigVMUnaryOp>(Instruction);
+				Lines.Add(Prefix + Format(RigVM_BoolTrueFormat, *GetOperandName(Op.Arg, false)));
+				break;
+			}
+			case ERigVMOpCode::Copy:
+			case ERigVMOpCode::ArrayClone:
+			{
+				const FRigVMBinaryOp& Op = ByteCode.GetOpAt<FRigVMBinaryOp>(Instruction);
+				const FString TargetOperand = GetOperandName(Op.ArgB, false, false);
+				const FString SourceOperand = GetOperandName(Op.ArgA, false, true);
+				const FString TargetCPPType = GetOperandCPPType(Op.ArgB);
+				const FString SourceCPPType = GetOperandCPPType(Op.ArgA);
+
+				if(RigVMTypeUtils::IsArrayType(TargetCPPType) &&
+					RigVMTypeUtils::IsArrayType(SourceCPPType) &&
+					TargetCPPType != SourceCPPType)
 				{
-					const FRigVMUnaryOp& Op = ByteCode.GetOpAt<FRigVMUnaryOp>(Operation);
-					const FString ElementCPPType = RigVMTypeUtils::BaseTypeFromArrayType(GetOperandCPPType(Op.Arg));
-					const FString TemplateSuffix = Format(RigVM_TemplateOneArgFormat, *ElementCPPType);
-					Lines.Add(Format(RigVM_ArrayReverseOpFormat, *TemplateSuffix, *GetOperandName(Op.Arg, false)));
-					break;
+					const FString TargetBaseCPPType = RigVMTypeUtils::BaseTypeFromArrayType(TargetCPPType);
+					const FString SourceBaseCPPType = RigVMTypeUtils::BaseTypeFromArrayType(SourceCPPType);
+					Lines.Add(Prefix + Format(RigVM_CopyUnrelatedArraysFormat, *TargetBaseCPPType, *SourceBaseCPPType, *TargetOperand, *SourceOperand));
 				}
-				case ERigVMOpCode::InvokeEntry:
+				else
 				{
-					const FRigVMInvokeEntryOp& Op = ByteCode.GetOpAt<FRigVMInvokeEntryOp>(Operation);
-					const FString EntryName = Op.EntryName.ToString();
-					Lines.Add(Format(RigVM_InvokeEntryByNameFormat2, *EntryName, *GetEntryParameters()));
-					break;
+					const FString CastPrefix = TargetCPPType != SourceCPPType ? Format(RigVM_BracesFormat, *TargetCPPType) : FString();  
+					
+					if(TargetOperand.EndsWith(TEXT("(")))
+					{
+						Lines.Add(Prefix + Format(RigVM_CopyOpMethodFormat, *TargetOperand, *CastPrefix, *SourceOperand));
+					}
+					else
+					{
+						Lines.Add(Prefix + Format(RigVM_CopyOpAssignFormat, *TargetOperand, *CastPrefix, *SourceOperand));
+					}
 				}
-				case ERigVMOpCode::Invalid:
-				case ERigVMOpCode::ChangeType:
-				default:
+				break;
+			}
+			case ERigVMOpCode::Increment:
+			{
+				const FRigVMUnaryOp& Op = ByteCode.GetOpAt<FRigVMUnaryOp>(Instruction);
+				Lines.Add(Prefix + Format(RigVM_IncrementOpFormat, *GetOperandName(Op.Arg, false)));
+				break;
+			}
+			case ERigVMOpCode::Decrement:
+			{
+				const FRigVMUnaryOp& Op = ByteCode.GetOpAt<FRigVMUnaryOp>(Instruction);
+				Lines.Add(Prefix + Format(RigVM_DecrementOpFormat, *GetOperandName(Op.Arg, false)));
+				break;
+			}
+			case ERigVMOpCode::Equals:
+			{
+				const FRigVMComparisonOp& Op = ByteCode.GetOpAt<FRigVMComparisonOp>(Instruction);
+				Lines.Add(Prefix + Format(RigVM_EqualsOpFormat, *GetOperandName(Op.Result, false), *GetOperandName(Op.B, false), *GetOperandName(Op.B, false)));
+				break;
+			}
+			case ERigVMOpCode::NotEquals:
+			{
+				const FRigVMComparisonOp& Op = ByteCode.GetOpAt<FRigVMComparisonOp>(Instruction);
+				Lines.Add(Prefix + Format(RigVM_NotEqualsOpFormat, *GetOperandName(Op.Result, false), *GetOperandName(Op.B, false), *GetOperandName(Op.B, false)));
+				break;
+			}
+			case ERigVMOpCode::JumpAbsolute:
+			{
+				const FRigVMJumpOp& Op = ByteCode.GetOpAt<FRigVMJumpOp>(Instruction);
+				Lines.Add(Prefix + Format(RigVM_JumpOpFormat, Op.InstructionIndex));
+				break;
+			}
+			case ERigVMOpCode::JumpForward:
+			{
+				const FRigVMJumpOp& Op = ByteCode.GetOpAt<FRigVMJumpOp>(Instruction);
+				Lines.Add(Prefix + Format(RigVM_JumpOpFormat, InstructionIndex + Op.InstructionIndex));
+				break;
+			}
+			case ERigVMOpCode::JumpBackward:
+			{
+				const FRigVMJumpOp& Op = ByteCode.GetOpAt<FRigVMJumpOp>(Instruction);
+				Lines.Add(Prefix + Format(RigVM_JumpOpFormat, InstructionIndex - Op.InstructionIndex));
+				break;
+			}
+			case ERigVMOpCode::JumpAbsoluteIf:
+			{
+				const FRigVMJumpIfOp& Op = ByteCode.GetOpAt<FRigVMJumpIfOp>(Instruction);
+				const FString& Condition = Op.Condition ? RigVM_TrueFormat : RigVM_FalseFormat;
+				Lines.Add(Prefix + Format(RigVM_JumpIfOpFormat, *GetOperandName(Op.Arg, false), *Condition, Op.InstructionIndex));
+				break;
+			}
+			case ERigVMOpCode::JumpForwardIf:
+			{
+				const FRigVMJumpIfOp& Op = ByteCode.GetOpAt<FRigVMJumpIfOp>(Instruction);
+				const FString& Condition = Op.Condition ? RigVM_TrueFormat : RigVM_FalseFormat;
+				Lines.Add(Prefix + Format(RigVM_JumpIfOpFormat, *GetOperandName(Op.Arg, false), *Condition, InstructionIndex + Op.InstructionIndex));
+				break;
+			}
+			case ERigVMOpCode::JumpBackwardIf:
+			{
+				const FRigVMJumpIfOp& Op = ByteCode.GetOpAt<FRigVMJumpIfOp>(Instruction);
+				const FString& Condition = Op.Condition ? RigVM_TrueFormat : RigVM_FalseFormat;
+				Lines.Add(Prefix + Format(RigVM_JumpIfOpFormat, *GetOperandName(Op.Arg, false), *Condition, InstructionIndex - Op.InstructionIndex));
+				break;
+			}
+			case ERigVMOpCode::Exit:
+			{
+				if(InstructionIndex != InGroup.Last)
 				{
-					// we expect to cover all op types
-					checkNoEntry();
+					Lines.Add(Prefix + Format(RigVM_JumpOpFormat, Instructions.Num()));
 				}
+				break;
 			}
-		}
-	}
-	else
-	{
-		// we have child groups - we need to invoke those
-		for(int32 ChildGroupIndex : Group.ChildGroups)
-		{
-			const FOperationGroup& ChildGroup = OperationGroups[ChildGroupIndex];
-
-			FString Parameters;
-			TArray<FString> ParameterArray = {RigVM_ContextPublicFormat};
-			if(!ChildGroup.OpaqueArguments.IsEmpty())
+			case ERigVMOpCode::BeginBlock:
 			{
-				for(const FOpaqueArgument& OpaqueArgument : ChildGroup.OpaqueArguments)
+				const FRigVMBinaryOp& Op = ByteCode.GetOpAt<FRigVMBinaryOp>(Instruction);
+				Lines.Add(Prefix + Format(RigVM_BeginBlockOpFormat, *GetOperandName(Op.ArgA, false), *GetOperandName(Op.ArgB, false)));
+				break;
+			}
+			case ERigVMOpCode::EndBlock:
+			{
+				Lines.Add(Prefix + RigVM_EndBlockOpFormat);
+				break;
+			}
+			case ERigVMOpCode::ArrayReset:
+			{
+				const FRigVMUnaryOp& Op = ByteCode.GetOpAt<FRigVMUnaryOp>(Instruction);
+				Lines.Add(Prefix + Format(RigVM_ArrayResetOpFormat, *GetOperandName(Op.Arg, false)));
+				break;
+			}
+			case ERigVMOpCode::ArrayGetNum:
+			{
+				const FRigVMBinaryOp& Op = ByteCode.GetOpAt<FRigVMBinaryOp>(Instruction);
+				Lines.Add(Prefix + Format(RigVM_ArrayGetNumOpFormat, *GetOperandName(Op.ArgB, false), *GetOperandName(Op.ArgA, false)));
+				break;
+			} 
+			case ERigVMOpCode::ArraySetNum:
+			{
+				const FRigVMBinaryOp& Op = ByteCode.GetOpAt<FRigVMBinaryOp>(Instruction);
+				const FString ConditionString = Format(RigVM_IsValidArraySizeFormat, *GetOperandName(Op.ArgB, false));
+				Lines.Add(Prefix + Format(RigVM_ArraySetNumOpFormat, *ConditionString, *GetOperandName(Op.ArgA, false), *GetOperandName(Op.ArgB, false)));
+				break;
+			}
+			case ERigVMOpCode::ArrayGetAtIndex:
+			{
+				const FRigVMTernaryOp& Op = ByteCode.GetOpAt<FRigVMTernaryOp>(Instruction);
+				const FString ConditionString = Format(RigVM_IsValidArrayIndexFormat, *GetOperandCPPBaseType(Op.ArgA), *GetOperandName(Op.ArgA, false));
+				Lines.Add(Prefix + Format(RigVM_TemporaryArrayIndexFormat, *GetOperandName(Op.ArgB, false)));
+				Lines.Add(Prefix + Format(RigVM_ArrayGetAtIndexOpFormat, *ConditionString, *GetOperandName(Op.ArgC, false), *GetOperandName(Op.ArgA, false)));
+				break;
+			}  
+			case ERigVMOpCode::ArraySetAtIndex:
+			{
+				const FRigVMTernaryOp& Op = ByteCode.GetOpAt<FRigVMTernaryOp>(Instruction);
+				const FString ConditionString = Format(RigVM_IsValidArrayIndexFormat,  *GetOperandCPPBaseType(Op.ArgA), *GetOperandName(Op.ArgA, false));
+				Lines.Add(Prefix + Format(RigVM_TemporaryArrayIndexFormat, *GetOperandName(Op.ArgB, false)));
+				Lines.Add(Prefix + Format(RigVM_ArraySetAtIndexOpFormat, *ConditionString, *GetOperandName(Op.ArgA, false), *GetOperandName(Op.ArgC, false), *GetOperandName(Op.ArgC, false)));
+				break;
+			}
+			case ERigVMOpCode::ArrayAdd:
+			{
+				const FRigVMTernaryOp& Op = ByteCode.GetOpAt<FRigVMTernaryOp>(Instruction);
+				const FString NumString = Format(RigVM_ArrayNumPlusOneFormat, *GetOperandName(Op.ArgA, false));
+				const FString ConditionString = Format(RigVM_IsValidArraySizeFormat, *NumString);
+				Lines.Add(Prefix + Format(RigVM_ArrayAddOpFormat, *ConditionString, *GetOperandName(Op.ArgC, false), *GetOperandName(Op.ArgA, false), *GetOperandName(Op.ArgB, false)));
+				break;
+			}
+			case ERigVMOpCode::ArrayInsert:
+			{
+				const FRigVMTernaryOp& Op = ByteCode.GetOpAt<FRigVMTernaryOp>(Instruction);
+				const FString NumString = Format(RigVM_ArrayNumPlusOneFormat, *GetOperandName(Op.ArgA, false));
+				const FString ConditionString = Format(RigVM_IsValidArraySizeFormat, *NumString);
+				Lines.Add(Prefix + Format(RigVM_ArrayInsertOpFormat, *ConditionString, *GetOperandName(Op.ArgA, false), *GetOperandName(Op.ArgC, false), *GetOperandName(Op.ArgB, false)));
+				break;
+			}
+			case ERigVMOpCode::ArrayRemove:
+			{
+				const FRigVMBinaryOp& Op = ByteCode.GetOpAt<FRigVMBinaryOp>(Instruction);
+				Lines.Add(Prefix + Format(RigVM_ArrayRemoveOpFormat, *GetOperandName(Op.ArgA, false), *GetOperandName(Op.ArgB, false)));
+				break;
+			}
+			case ERigVMOpCode::ArrayFind:
+			{
+				const FRigVMQuaternaryOp& Op = ByteCode.GetOpAt<FRigVMQuaternaryOp>(Instruction);
+				const FString ArrayElementCPPType = RigVMTypeUtils::BaseTypeFromArrayType(GetOperandCPPBaseType(Op.ArgA));
+				const FString ElementCPPType = GetOperandCPPType(Op.ArgB);
+				const FString TemplateSuffix = ArrayElementCPPType == ElementCPPType ? Format(RigVM_TemplateOneArgFormat, *ElementCPPType) : FString();
+				Lines.Add(Prefix + Format(RigVM_ArrayFindOpFormat, *GetOperandName(Op.ArgD, false), *TemplateSuffix, *GetOperandName(Op.ArgA, false), *GetOperandName(Op.ArgB, false), *GetOperandName(Op.ArgC, false)));
+				break;
+			}
+			case ERigVMOpCode::ArrayAppend:
+			{
+				const FRigVMBinaryOp& Op = ByteCode.GetOpAt<FRigVMBinaryOp>(Instruction);
+				const FString NumString = Format(RigVM_ArrayAppendOpNumFormat, *GetOperandName(Op.ArgA, false), *GetOperandName(Op.ArgB, false));
+				const FString ConditionString = Format(RigVM_IsValidArraySizeFormat, *NumString);
+				Lines.Add(Prefix + Format(RigVM_ArrayAppendOpFormat, *ConditionString, *GetOperandName(Op.ArgA, false), *GetOperandName(Op.ArgB, false)));
+				break;
+			}
+			case ERigVMOpCode::ArrayIterator:
+			{
+				const FRigVMSenaryOp& Op = ByteCode.GetOpAt<FRigVMSenaryOp>(Instruction);
+				const FString ArrayElementCPPType = RigVMTypeUtils::BaseTypeFromArrayType(GetOperandCPPType(Op.ArgA));
+				const FString ElementCPPType = GetOperandCPPType(Op.ArgB);
+				const FString TemplateSuffix = ArrayElementCPPType == ElementCPPType ? Format(RigVM_TemplateOneArgFormat, *ElementCPPType) : FString();
+				Lines.Add(Prefix + Format(RigVM_ArrayIteratorOpFormat, *GetOperandName(Op.ArgF, false), *TemplateSuffix, *GetOperandName(Op.ArgA, false), *GetOperandName(Op.ArgB, false), *GetOperandName(Op.ArgC, false), *GetOperandName(Op.ArgD, false), *GetOperandName(Op.ArgE, false)));
+				break;
+			}
+			case ERigVMOpCode::ArrayUnion:
+			{
+				const FRigVMBinaryOp& Op = ByteCode.GetOpAt<FRigVMBinaryOp>(Instruction);
+				const FString ElementCPPTypeA = RigVMTypeUtils::BaseTypeFromArrayType(GetOperandCPPType(Op.ArgA));
+				const FString ElementCPPTypeB = RigVMTypeUtils::BaseTypeFromArrayType(GetOperandCPPType(Op.ArgA));
+				const FString TemplateSuffix = ElementCPPTypeA == ElementCPPTypeB ? Format(RigVM_TemplateOneArgFormat, *ElementCPPTypeA) : FString();
+				Lines.Add(Prefix + Format(RigVM_ArrayUnionOpFormat, *TemplateSuffix, *GetOperandName(Op.ArgA, false, false), *GetOperandName(Op.ArgB, false, false)));
+				break;
+			}
+			case ERigVMOpCode::ArrayDifference:
+			{
+				const FRigVMBinaryOp& Op = ByteCode.GetOpAt<FRigVMBinaryOp>(Instruction);
+				const FString ElementCPPTypeA = RigVMTypeUtils::BaseTypeFromArrayType(GetOperandCPPType(Op.ArgA));
+				const FString ElementCPPTypeB = RigVMTypeUtils::BaseTypeFromArrayType(GetOperandCPPType(Op.ArgA));
+				const FString TemplateSuffix = ElementCPPTypeA == ElementCPPTypeB ? Format(RigVM_TemplateOneArgFormat, *ElementCPPTypeA) : FString();
+				Lines.Add(Prefix + Format(RigVM_ArrayDifferenceOpFormat, *TemplateSuffix, *GetOperandName(Op.ArgA, false), *GetOperandName(Op.ArgB, false)));
+				break;
+			}
+			case ERigVMOpCode::ArrayIntersection:
+			{
+				const FRigVMBinaryOp& Op = ByteCode.GetOpAt<FRigVMBinaryOp>(Instruction);
+				const FString ElementCPPTypeA = RigVMTypeUtils::BaseTypeFromArrayType(GetOperandCPPType(Op.ArgA));
+				const FString ElementCPPTypeB = RigVMTypeUtils::BaseTypeFromArrayType(GetOperandCPPType(Op.ArgA));
+				const FString TemplateSuffix = ElementCPPTypeA == ElementCPPTypeB ? Format(RigVM_TemplateOneArgFormat, *ElementCPPTypeA) : FString();
+				Lines.Add(Prefix + Format(RigVM_ArrayIntersectionOpFormat, *TemplateSuffix, *GetOperandName(Op.ArgA, false), *GetOperandName(Op.ArgB, false)));
+				break;
+			}
+			case ERigVMOpCode::ArrayReverse:
+			{
+				const FRigVMUnaryOp& Op = ByteCode.GetOpAt<FRigVMUnaryOp>(Instruction);
+				const FString ElementCPPType = RigVMTypeUtils::BaseTypeFromArrayType(GetOperandCPPType(Op.Arg));
+				const FString TemplateSuffix = Format(RigVM_TemplateOneArgFormat, *ElementCPPType);
+				Lines.Add(Prefix + Format(RigVM_ArrayReverseOpFormat, *TemplateSuffix, *GetOperandName(Op.Arg, false)));
+				break;
+			}
+			case ERigVMOpCode::InvokeEntry:
+			{
+				const FRigVMInvokeEntryOp& Op = ByteCode.GetOpAt<FRigVMInvokeEntryOp>(Instruction);
+				const FString EntryName = Op.EntryName.ToString();
+				Lines.Add(Prefix + Format(RigVM_InvokeEntryByNameFormat2, *EntryName, *GetEntryParameters()));
+				break;
+			}
+			case ERigVMOpCode::JumpToBranch:
+			{
+				const FRigVMJumpToBranchOp& Op = ByteCode.GetOpAt<FRigVMJumpToBranchOp>(Instruction);
+				const TArray<FRigVMBranchInfo>& Branches = ByteCode.BranchInfos;
+				for(int32 BranchIndex = Op.FirstBranchInfoIndex; BranchIndex < Branches.Num(); BranchIndex++)
 				{
-					ParameterArray.Add(OpaqueArgument.Get<1>());
+					const FRigVMBranchInfo& Branch = Branches[BranchIndex];
+					if(Branch.InstructionIndex != InstructionIndex)
+					{
+						break;
+					}
+					Lines.Add(Prefix + Format(RigVM_JumpToBranchFormat, *GetOperandName(Op.Arg, false), Branches[BranchIndex].Label.ToString(), (int32)Branch.FirstInstruction));
 				}
+				break;
+			}
+			case ERigVMOpCode::Invalid:
+			case ERigVMOpCode::ChangeType:
+			default:
+			{
+				// we expect to cover all op types
+				checkNoEntry();
 			}
-			Parameters = FString::Join(ParameterArray, RigVM_CommaSeparator);
-			Lines.Add(Format(RigVM_InvokeExecuteGroupFormat, *ChildGroup.Entry, ChildGroupIndex, *Parameters));
 		}
 	}
-	
-	Lines.Emplace();
-	if(Group.Depth <= 0)
-	{
-		Lines.Add(RigVM_ExecuteReachedExitFormat);
-	}
-	Lines.Add(RigVM_ReturnSucceededFormat);
 
 	return DumpLines(Lines, bLog);
 }
 
+TArray<int32> FRigVMCodeGenerator::GetInstructionIndicesFromRange(int32 First, int32 Last)
+{
+	if(First > Last || First == INDEX_NONE || Last == INDEX_NONE)
+	{
+		return {};
+	}
+	if(First == Last)
+	{
+		return {First};
+	}
+	TArray<int32> Indices;
+	Indices.Reserve(Last - First + 1);
+	for(int32 Index = First; Index <= Last; Index++)
+	{
+		Indices.Add(Index);
+	}
+	return Indices;
+}
+
+
 FString FRigVMCodeGenerator::DumpHeader(bool bLog)
 {
 	const FRigVMByteCode& ByteCode = VM->GetByteCode();
@@ -1005,10 +1253,12 @@ FString FRigVMCodeGenerator::DumpHeader(bool bLog)
 		Lines.Add(DumpWrappedArrayTypes());
 		Lines.Emplace();
 	}
+
 	Lines.Add(Format(RigVM_UClassDefinitionFormat, *ModuleName.ToUpper(), *ClassName));
-	Lines.Add(Format(RigVM_GetVMHashFormat, VM->GetVMHash()));
+	Lines.Add(Format(RigVM_GetVMHashFormat, FString::Printf(TEXT("%lu"), VM->GetVMHash())));
 	Lines.Add(Format(RigVM_GetEntryNamesFormat, *FString::Join(FormattedEntries, RigVM_CommaSeparator)));
 	Lines.Emplace();
+	Lines.Add(FString(RigVM_DeclareInitializeFormat));
 	Lines.Add(FString(RigVM_DeclareExecuteFormat));
 	Lines.Emplace();
 	Lines.Add(RigVM_ProtectedFormat);
@@ -1017,9 +1267,9 @@ FString FRigVMCodeGenerator::DumpHeader(bool bLog)
 		Lines.Add(RigVM_DeclareUpdateExternalVariablesFormat);
 	}
 
-	for(int32 GroupIndex = 0; GroupIndex < OperationGroups.Num(); GroupIndex++)
+	for(int32 GroupIndex = 0; GroupIndex < InstructionGroups.Num(); GroupIndex++)
 	{
-		const FOperationGroup& Group = OperationGroups[GroupIndex];
+		const FInstructionGroup& Group = InstructionGroups[GroupIndex];
 
 		FString Parameters;
 		TArray<FString> ParameterArray = {Format(RigVM_ContextPublicParameterFormat, ExecuteContextType)};
@@ -1041,7 +1291,7 @@ FString FRigVMCodeGenerator::DumpHeader(bool bLog)
 			Lines.Add(Format(RigVM_DeclareExecuteGroupFormat, *Group.Entry, GroupIndex, *Parameters));
 		}
 
-		if(GroupIndex == OperationGroups.Num() - 1)
+		if(GroupIndex == InstructionGroups.Num() - 1)
 		{
 			if(!Parameters.IsEmpty())
 			{
@@ -1050,16 +1300,26 @@ FString FRigVMCodeGenerator::DumpHeader(bool bLog)
 			Lines.Add(Format(RigVM_DeclareInvokeEntryByNameFormat, *Parameters));
 		}
 	}
-
+	
 	Lines.Emplace();
 	Lines.Add(DumpEntries(true));
+	const FString BlockNames = DumpBlockNames(true);
+	if(!BlockNames.IsEmpty())
+	{
+		Lines.Add(BlockNames);
+	}
 
 	Lines.Emplace();
 	Lines.Add(DumpProperties(true, INDEX_NONE));
+	Lines.Add(DumpDispatches(true));
 	if(!VM->GetExternalVariables().IsEmpty())
 	{
 		Lines.Add(DumpExternalVariables(true));
 	}
+	if(!RequiredUProperties.IsEmpty())
+	{
+		Lines.Add(DumpRequiredUProperties());
+	}
 	Lines.Add(TEXT("};"));
 	Lines.Emplace();
 	
@@ -1078,9 +1338,27 @@ FString FRigVMCodeGenerator::DumpSource(bool bLog)
 	Lines.Add(Format(RigVM_IncludeQuoteFormat, *ClassName));
 	Lines.Emplace();
 	Lines.Add(DumpEntries(false));
-
+	const FString BlockNames = DumpBlockNames(false);
+	if(!BlockNames.IsEmpty())
+	{
+		Lines.Add(BlockNames);
+	}
+	for(auto Pair : RequiredUProperties)
+	{
+		Lines.Add(Format(RigVM_UPropertyMember2Format, *ClassName, Pair.Value.Get<1>()));
+	}
+	Lines.Emplace();
+	Lines.Add(Format(RigVM_DefineInitializeFormat, *ClassName));
+	const FString InitializeContent = DumpInitialize();
+	if(!InitializeContent.IsEmpty())
+	{
+		Lines.Add(InitializeContent);
+	}
+	Lines.Add(TEXT("}"));
 	Lines.Emplace();
 	Lines.Add(Format(RigVM_DefineExecuteFormat, *ClassName));
+	Lines.Add(RigVM_StartProfilingFormat);
+	Lines.Add(Format(RigVM_SetupInstructionTrackingFormat, VM->GetByteCode().GetNumInstructions()));
 	Lines.Add(DumpInstructions(INDEX_NONE));
 	Lines.Add(TEXT("}"));
 
@@ -1092,9 +1370,9 @@ FString FRigVMCodeGenerator::DumpSource(bool bLog)
 		Lines.Add(TEXT("}"));
 	}
 
-	for(int32 GroupIndex = 0; GroupIndex < OperationGroups.Num(); GroupIndex++)
+	for(int32 GroupIndex = 0; GroupIndex < InstructionGroups.Num(); GroupIndex++)
 	{
-		const FOperationGroup& Group = OperationGroups[GroupIndex];
+		const FInstructionGroup& Group = InstructionGroups[GroupIndex];
 
 		FString Parameters;
 		TArray<FString> ParameterArray = {Format(RigVM_ContextPublicParameterFormat, ExecuteContextType)};
@@ -1126,7 +1404,7 @@ FString FRigVMCodeGenerator::DumpSource(bool bLog)
 		Lines.Add(DumpInstructions(GroupIndex));
 		Lines.Add(TEXT("}"));
 
-		if(GroupIndex == OperationGroups.Num() - 1)
+		if(GroupIndex == InstructionGroups.Num() - 1)
 		{
 			if(!Parameters.IsEmpty())
 			{
@@ -1161,16 +1439,19 @@ void FRigVMCodeGenerator::Reset()
 	ModuleName.Reset();
 	Libraries.Reset();
 	Includes.Reset();
-	OperationGroups.Reset();
+	Dispatches.Reset();
+	RequiredUProperties.Reset();
+	InstructionGroups.Reset();
 	WrappedArrayCPPTypes.Reset();
 	MappedCPPTypes.Reset();
 	Properties.Reset();
 	PropertyNameToIndex.Reset();
+	OverriddenOperatorNames.Reset();
 }
 
 void FRigVMCodeGenerator::ParseVM(const FString& InClassName, const FString& InModuleName,
 	URigVMGraph* InModelToNativize, URigVM* InVMToNativize,
-	TMap<FString,FRigVMOperand> InPinToOperandMap, int32 InMaxOperationsPerFunction)
+	TMap<FString,FRigVMOperand> InPinToOperandMap, int32 InMaxInstructionsPerFunction)
 {
 	check(InVMToNativize);
 
@@ -1179,7 +1460,7 @@ void FRigVMCodeGenerator::ParseVM(const FString& InClassName, const FString& InM
 	Model = TStrongObjectPtr<URigVMGraph>(InModelToNativize);
 	VM = TStrongObjectPtr<URigVM>(InVMToNativize);
 	PinToOperandMap = InPinToOperandMap;
-	MaxOperationsPerFunction = InMaxOperationsPerFunction;
+	MaxInstructionsPerFunction = InMaxInstructionsPerFunction;
 	ClassName = InClassName;
 	ModuleName = InModuleName;
 	ExecuteContextType = VM->GetContextPublicDataStruct()->GetStructCPPName();
@@ -1200,7 +1481,8 @@ void FRigVMCodeGenerator::ParseVM(const FString& InClassName, const FString& InM
 	ParseMemory(InVMToNativize->GetLiteralMemory(true));
 	ParseMemory(InVMToNativize->GetWorkMemory(true));
 
-	ParseOperationGroups();
+	ParseRequiredUProperties();
+	ParseInstructionGroups();
 }
 
 void FRigVMCodeGenerator::ParseInclude(UStruct* InDependency, const FName& InMethodName)
@@ -1228,6 +1510,70 @@ void FRigVMCodeGenerator::ParseInclude(UStruct* InDependency, const FName& InMet
 	{
 		Libraries.AddUnique(RigVM_RigVMCoreLibraryFormat);
 		Includes.AddUnique(RigVM_RigVMCoreIncludeFormat);
+		Includes.AddUnique(RigVM_RigVMModuleIncludeFormat);
+	}
+}
+
+void FRigVMCodeGenerator::ParseRequiredUProperties()
+{
+	const TArray<FName>& Functions = VM->GetFunctionNames();
+	const FRigVMByteCode& ByteCode = VM->GetByteCode();
+	const FRigVMInstructionArray Instructions = ByteCode.GetInstructions();
+
+	auto AddRequiredUProperty = [this](const FString& InTypeName)
+	{
+		const TRigVMTypeIndex TypeIndex = FRigVMRegistry::Get().GetTypeIndexFromCPPType(InTypeName);
+		check(TypeIndex != INDEX_NONE);
+		if(!RequiredUProperties.Contains(TypeIndex))
+		{
+			const FRigVMTemplateArgumentType& Type = FRigVMRegistry::Get().GetType(TypeIndex);
+			FString TypeLabel = Type.GetBaseCPPType();
+			if(Type.IsArray())
+			{
+				TypeLabel += TEXT("_Array");
+			}
+			TypeLabel = TEXT("Property_") + TypeLabel; 
+			RequiredUProperties.Add(TypeIndex, {InTypeName, TypeLabel});
+		}
+	};
+
+	for(const FRigVMInstruction& Instruction : Instructions)
+	{
+		if(Instruction.OpCode >= ERigVMOpCode::Execute_0_Operands &&
+			Instruction.OpCode <= ERigVMOpCode::Execute_64_Operands)
+		{
+			const FRigVMExecuteOp& Op = ByteCode.GetOpAt<FRigVMExecuteOp>(Instruction);
+			const FRigVMFunction* Function = FRigVMRegistry::Get().FindFunction(*Functions[Op.FunctionIndex].ToString());
+			if(const FRigVMDispatchFactory* Factory = Function->Factory)
+			{
+				const FString FunctionName = Function->GetName();
+				const int32 PermutationIndex = Factory->GetTemplate()->FindPermutation(Function); 
+				check(PermutationIndex != INDEX_NONE);
+				const FString DispatchKey = Format(RigVM_DispatchKeyFormat, Factory->GetFactoryName().ToString(), PermutationIndex);
+				Dispatches.Add(DispatchKey, Function);
+
+				for(const FRigVMFunctionArgument& Argument : Function->Arguments)
+				{
+					AddRequiredUProperty(Argument.Type);
+				}
+			}
+		}
+	}
+
+	// also parse all lazy branches
+	for(const FRigVMBranchInfo& BranchInfo : ByteCode.BranchInfos)
+	{
+		// skip output block branches
+		if(BranchInfo.IsOutputBranch())
+		{
+			continue;
+		}
+
+		const FRigVMOperandArray Operands = ByteCode.GetOperandsForOp(Instructions[BranchInfo.InstructionIndex]);
+		check(Operands.IsValidIndex(BranchInfo.ArgumentIndex));
+		const FRigVMOperand Operand = Operands[BranchInfo.ArgumentIndex];
+		const FString CPPType = GetOperandCPPType(Operand);
+		AddRequiredUProperty(CPPType);
 	}
 }
 
@@ -1306,83 +1652,102 @@ void FRigVMCodeGenerator::ParseProperty(ERigVMMemoryType InMemoryType, const FPr
 	}
 }
 
-void FRigVMCodeGenerator::ParseOperationGroups()
+void FRigVMCodeGenerator::ParseInstructionGroups()
 {
 	const TArray<FName>& Functions = VM->GetFunctionNames();
 	const FRigVMByteCode& ByteCode = VM->GetByteCode();
-	const FRigVMInstructionArray Operations = ByteCode.GetInstructions();
+	const FRigVMInstructionArray Instructions = ByteCode.GetInstructions();
 
 	for(int32 EntryIndex = 0; EntryIndex < ByteCode.NumEntries(); EntryIndex++)
 	{
 		const FRigVMByteCodeEntry& Entry = ByteCode.GetEntry(EntryIndex);
 
-		FOperationGroup Group;
+		FInstructionGroup Group;
 		Group.Entry = Entry.GetSanitizedName();
 		Group.Depth = 0;
 		Group.First = Entry.InstructionIndex;
-		Group.Last = Operations.Num() - 1;
+		Group.Last = Instructions.Num() - 1;
 
-		for(int32 OperationIndex = Group.First+1; OperationIndex < Operations.Num(); OperationIndex++)
+		for(int32 InstructionIndex = Group.First+1; InstructionIndex < Instructions.Num(); InstructionIndex++)
 		{
-			if(Operations[OperationIndex].OpCode == ERigVMOpCode::Exit)
+			if(Instructions[InstructionIndex].OpCode == ERigVMOpCode::Exit)
 			{
-				Group.Last = OperationIndex;
+				Group.Last = InstructionIndex;
 				break;
 			}
 		}
 
-		OperationGroups.Add(Group);
+		InstructionGroups.Add(Group);
 	}
 
-	for(int32 GroupIndex = 0; GroupIndex < OperationGroups.Num(); GroupIndex++)
+	for(int32 GroupIndex = 0; GroupIndex < InstructionGroups.Num(); GroupIndex++)
 	{
 		// copy the group here since it will be invalid soon
-		FOperationGroup Group = OperationGroups[GroupIndex];
-		if(Group.Last - Group.First + 1 > MaxOperationsPerFunction)
+		FInstructionGroup Group = InstructionGroups[GroupIndex];
+		if(Group.Last - Group.First + 1 > MaxInstructionsPerFunction)
 		{
 			// find all of the constraints / jumps
 			TArray<TTuple<int32,int32>> Constraints;
-			for(int32 OperationIndex = Group.First; OperationIndex <= Group.Last; OperationIndex++)
+			for(int32 InstructionIndex = Group.First; InstructionIndex <= Group.Last; InstructionIndex++)
 			{
-				int32 OtherIndex = INDEX_NONE;
+				auto AddConstraint = [&Constraints, InstructionIndex](int32 OtherInstructionIndex)
+				{
+					Constraints.Add(TTuple<int32,int32>(
+						FMath::Min(InstructionIndex, OtherInstructionIndex),
+						FMath::Max(InstructionIndex, OtherInstructionIndex)));
+				};
 				
-				const FRigVMInstruction& Operation = Operations[OperationIndex];
-				switch(Operation.OpCode)
+				const FRigVMInstruction& Instruction = Instructions[InstructionIndex];
+				switch(Instruction.OpCode)
 				{
 					case ERigVMOpCode::JumpAbsolute:
 					{
-						const FRigVMJumpOp& Op = ByteCode.GetOpAt<FRigVMJumpOp>(Operation);
-						OtherIndex = Op.InstructionIndex;
+						const FRigVMJumpOp& Op = ByteCode.GetOpAt<FRigVMJumpOp>(Instruction);
+						AddConstraint(Op.InstructionIndex);
 						break;
 					}
 					case ERigVMOpCode::JumpForward:
 					{
-						const FRigVMJumpOp& Op = ByteCode.GetOpAt<FRigVMJumpOp>(Operation);
-						OtherIndex = OperationIndex + Op.InstructionIndex;
+						const FRigVMJumpOp& Op = ByteCode.GetOpAt<FRigVMJumpOp>(Instruction);
+						AddConstraint(InstructionIndex + Op.InstructionIndex);
 						break;
 					}
 					case ERigVMOpCode::JumpBackward:
 					{
-						const FRigVMJumpOp& Op = ByteCode.GetOpAt<FRigVMJumpOp>(Operation);
-						OtherIndex = OperationIndex - Op.InstructionIndex;
+						const FRigVMJumpOp& Op = ByteCode.GetOpAt<FRigVMJumpOp>(Instruction);
+						AddConstraint(InstructionIndex - Op.InstructionIndex);
 						break;
 					}
 					case ERigVMOpCode::JumpAbsoluteIf:
 					{
-						const FRigVMJumpIfOp& Op = ByteCode.GetOpAt<FRigVMJumpIfOp>(Operation);
-						OtherIndex = OperationIndex;
+						const FRigVMJumpIfOp& Op = ByteCode.GetOpAt<FRigVMJumpIfOp>(Instruction);
+						AddConstraint(InstructionIndex);
 						break;
 					}
 					case ERigVMOpCode::JumpForwardIf:
 					{
-						const FRigVMJumpIfOp& Op = ByteCode.GetOpAt<FRigVMJumpIfOp>(Operation);
-						OtherIndex = OperationIndex + Op.InstructionIndex;
+						const FRigVMJumpIfOp& Op = ByteCode.GetOpAt<FRigVMJumpIfOp>(Instruction);
+						AddConstraint(InstructionIndex + Op.InstructionIndex);
 						break;
 					}
 					case ERigVMOpCode::JumpBackwardIf:
 					{
-						const FRigVMJumpIfOp& Op = ByteCode.GetOpAt<FRigVMJumpIfOp>(Operation);
-						OtherIndex = OperationIndex - Op.InstructionIndex;
+						const FRigVMJumpIfOp& Op = ByteCode.GetOpAt<FRigVMJumpIfOp>(Instruction);
+						AddConstraint(InstructionIndex - Op.InstructionIndex);
+						break;
+					}
+					case ERigVMOpCode::JumpToBranch:
+					{
+						const FRigVMJumpToBranchOp& Op = ByteCode.GetOpAt<FRigVMJumpToBranchOp>(Instruction);
+						for(int32 BranchIndex = Op.FirstBranchInfoIndex; BranchIndex < ByteCode.BranchInfos.Num(); BranchIndex++)
+						{
+							const FRigVMBranchInfo& Branch = ByteCode.BranchInfos[BranchIndex];
+							if(Branch.InstructionIndex != InstructionIndex)
+							{
+								break;
+							}
+							AddConstraint(Branch.FirstInstruction);
+						}
 						break;
 					}
 					default:
@@ -1390,17 +1755,22 @@ void FRigVMCodeGenerator::ParseOperationGroups()
 						break;
 					}
 				}
+			}
+
+			const int32 First = Group.First;
+			const int32 Last = Group.Last;
 
-				if(OtherIndex != INDEX_NONE)
+			// also constrain the group split by the branches used
+			for(const FRigVMBranchInfo& Branch : ByteCode.BranchInfos)
+			{
+				if(Branch.FirstInstruction >= First && Branch.LastInstruction <= Last)
 				{
 					Constraints.Add(TTuple<int32,int32>(
-						FMath::Min(OperationIndex, OtherIndex),
-						FMath::Max(OperationIndex, OtherIndex)));
+						FMath::Min(Branch.FirstInstruction, Branch.LastInstruction),
+						FMath::Max(Branch.FirstInstruction, Branch.LastInstruction)));
 				}
 			}
-
-			const int32 First = Group.First;
-			const int32 Last = Group.Last;
+			
 			int32 Middle = (Group.Last + Group.First) / 2;
 			int32 MiddleHead = Middle;
 			int32 MiddleTail = Middle;
@@ -1429,36 +1799,36 @@ void FRigVMCodeGenerator::ParseOperationGroups()
 				continue;
 			}
 
-			FOperationGroup HeadGroup = Group;
+			FInstructionGroup HeadGroup = Group;
 			HeadGroup.ParentGroup = GroupIndex;
 			HeadGroup.ChildGroups.Reset();
 			HeadGroup.First = First;
 			HeadGroup.Last = Middle;
 			HeadGroup.Depth++;
 
-			FOperationGroup TailGroup = Group;
+			FInstructionGroup TailGroup = Group;
 			TailGroup.ParentGroup = GroupIndex;
 			TailGroup.ChildGroups.Reset();
 			TailGroup.First = Middle+1;
 			TailGroup.Last = Last;
 			TailGroup.Depth++;
 			
-			const int32 HeadGroupIndex = OperationGroups.Add(HeadGroup);
-			const int32 TailGroupIndex = OperationGroups.Add(TailGroup);
-			OperationGroups[GroupIndex].ChildGroups.Add(HeadGroupIndex);
-			OperationGroups[GroupIndex].ChildGroups.Add(TailGroupIndex);
+			const int32 HeadGroupIndex = InstructionGroups.Add(HeadGroup);
+			const int32 TailGroupIndex = InstructionGroups.Add(TailGroup);
+			InstructionGroups[GroupIndex].ChildGroups.Add(HeadGroupIndex);
+			InstructionGroups[GroupIndex].ChildGroups.Add(TailGroupIndex);
 		}
 	}
 
-	for(int32 GroupIndex = 0; GroupIndex < OperationGroups.Num(); GroupIndex++)
+	for(int32 GroupIndex = 0; GroupIndex < InstructionGroups.Num(); GroupIndex++)
 	{
-		FOperationGroup& Group = OperationGroups[GroupIndex];
+		FInstructionGroup& Group = InstructionGroups[GroupIndex];
 		
 		// determine all the necessary labels and optional arguments
-		for(int32 OperationIndex = Group.First; OperationIndex <= Group.Last; OperationIndex++)
+		for(int32 InstructionIndex = Group.First; InstructionIndex <= Group.Last; InstructionIndex++)
 		{
-			const FRigVMInstruction& Operation = Operations[OperationIndex];
-			switch(Operation.OpCode)
+			const FRigVMInstruction& Instruction = Instructions[InstructionIndex];
+			switch(Instruction.OpCode)
 			{
 				case ERigVMOpCode::Execute_0_Operands:
 				case ERigVMOpCode::Execute_1_Operands:
@@ -1526,12 +1896,19 @@ void FRigVMCodeGenerator::ParseOperationGroups()
 				case ERigVMOpCode::Execute_63_Operands:
 				case ERigVMOpCode::Execute_64_Operands:
 				{
-					const FRigVMExecuteOp& Op = ByteCode.GetOpAt<FRigVMExecuteOp>(Operation);
+					const FRigVMExecuteOp& Op = ByteCode.GetOpAt<FRigVMExecuteOp>(Instruction);
 					const FRigVMFunction* Function = FRigVMRegistry::Get().FindFunction(*Functions[Op.FunctionIndex].ToString());
 					check(Function);
 
 					// make sure to include the required header
-					ParseInclude(Function->Struct, Function->GetMethodName());
+					if(Function->Struct)
+					{
+						ParseInclude(Function->Struct, Function->GetMethodName());
+					}
+					if(Function->Factory)
+					{
+						ParseInclude(Function->Factory->GetScriptStruct(), Function->GetMethodName());
+					}
 
 					// check all of the arguments that may not be part of the struct definition
 					for(const FRigVMFunctionArgument& Argument : Function->GetArguments())
@@ -1553,38 +1930,53 @@ void FRigVMCodeGenerator::ParseOperationGroups()
 				}
 				case ERigVMOpCode::JumpAbsolute:
 				{
-					const FRigVMJumpOp& Op = ByteCode.GetOpAt<FRigVMJumpOp>(Operation);
-					Group.RequiredLabels.AddUnique(Op.InstructionIndex);break;
+					const FRigVMJumpOp& Op = ByteCode.GetOpAt<FRigVMJumpOp>(Instruction);
+					Group.RequiredLabels.AddUnique(Op.InstructionIndex);
 					break;
 				}
 				case ERigVMOpCode::JumpForward:
 				{
-					const FRigVMJumpOp& Op = ByteCode.GetOpAt<FRigVMJumpOp>(Operation);
-					Group.RequiredLabels.AddUnique(OperationIndex + Op.InstructionIndex);break;
+					const FRigVMJumpOp& Op = ByteCode.GetOpAt<FRigVMJumpOp>(Instruction);
+					Group.RequiredLabels.AddUnique(InstructionIndex + Op.InstructionIndex);
 					break;
 				}
 				case ERigVMOpCode::JumpBackward:
 				{
-					const FRigVMJumpOp& Op = ByteCode.GetOpAt<FRigVMJumpOp>(Operation);
-					Group.RequiredLabels.AddUnique(OperationIndex - Op.InstructionIndex);break;
+					const FRigVMJumpOp& Op = ByteCode.GetOpAt<FRigVMJumpOp>(Instruction);
+					Group.RequiredLabels.AddUnique(InstructionIndex - Op.InstructionIndex);
 					break;
 				}
 				case ERigVMOpCode::JumpAbsoluteIf:
 				{
-					const FRigVMJumpIfOp& Op = ByteCode.GetOpAt<FRigVMJumpIfOp>(Operation);
-					Group.RequiredLabels.AddUnique(Op.InstructionIndex);break;
+					const FRigVMJumpIfOp& Op = ByteCode.GetOpAt<FRigVMJumpIfOp>(Instruction);
+					Group.RequiredLabels.AddUnique(Op.InstructionIndex);
 					break;
 				}
 				case ERigVMOpCode::JumpForwardIf:
 				{
-					const FRigVMJumpIfOp& Op = ByteCode.GetOpAt<FRigVMJumpIfOp>(Operation);
-					Group.RequiredLabels.AddUnique(OperationIndex + Op.InstructionIndex);break;
+					const FRigVMJumpIfOp& Op = ByteCode.GetOpAt<FRigVMJumpIfOp>(Instruction);
+					Group.RequiredLabels.AddUnique(InstructionIndex + Op.InstructionIndex);
 					break;
 				}
 				case ERigVMOpCode::JumpBackwardIf:
 				{
-					const FRigVMJumpIfOp& Op = ByteCode.GetOpAt<FRigVMJumpIfOp>(Operation);
-					Group.RequiredLabels.AddUnique(OperationIndex - Op.InstructionIndex);break;
+					const FRigVMJumpIfOp& Op = ByteCode.GetOpAt<FRigVMJumpIfOp>(Instruction);
+					Group.RequiredLabels.AddUnique(InstructionIndex - Op.InstructionIndex);
+					break;
+				}
+				case ERigVMOpCode::JumpToBranch:
+				{
+					const FRigVMJumpToBranchOp& Op = ByteCode.GetOpAt<FRigVMJumpToBranchOp>(Instruction);
+					const TArray<FRigVMBranchInfo>& Branches = ByteCode.BranchInfos;
+					for(int32 BranchIndex = Op.FirstBranchInfoIndex; BranchIndex < Branches.Num(); BranchIndex++)
+					{
+						// only add labels for output blocks - so blocks that belong
+						// to a jump to branch instruction
+						if(Branches[BranchIndex].IsOutputBranch())
+						{
+							Group.RequiredLabels.AddUnique((int32)Branches[BranchIndex].FirstInstruction);
+						}
+					}
 					break;
 				}
 				default:
@@ -1594,7 +1986,7 @@ void FRigVMCodeGenerator::ParseOperationGroups()
 			}
 
 			// update the usage table to know which property is used where
-			FRigVMOperandArray Operands = ByteCode.GetOperandsForOp(Operation);
+			FRigVMOperandArray Operands = ByteCode.GetOperandsForOp(Instruction);
 			for(const FRigVMOperand& Operand : Operands)
 			{
 				if(Operand.GetMemoryType() == ERigVMMemoryType::Literal ||
@@ -1612,7 +2004,7 @@ void FRigVMCodeGenerator::ParseOperationGroups()
 	// groups. parent groups which only pass them to one sub group can be removed from the
 	// property's group table. the top level group which still contains the property is the
 	// one to declare it.
-	const TArray<FOperationGroup>& Groups = OperationGroups;
+	const TArray<FInstructionGroup>& Groups = InstructionGroups;
 	for(int32 PropertyIndex = 0; PropertyIndex < Properties.Num(); PropertyIndex++)
 	{
 		// work / sliced properties are always defined at the top
@@ -1655,6 +2047,11 @@ FString FRigVMCodeGenerator::GetOperandName(const FRigVMOperand& InOperand, bool
 		OperandName = SanitizeName(OperandName, Property.CPPType);
 	}
 
+	if(const FString* OverrideName = OverriddenOperatorNames.Find(OperandName))
+	{
+		return *OverrideName;
+	}
+
 	// if we are an array on work memory this indicates that we'll be sliced.
 	if (bPerSlice && InOperand.GetMemoryType() == ERigVMMemoryType::Work && RigVMTypeUtils::IsArrayType(Property.CPPType))
 	{
@@ -1990,14 +2387,14 @@ FString FRigVMCodeGenerator::GetMappedArrayTypeName(const FString InBaseElementT
 	return Format(RigVM_WrappedTypeNameFormat, *InBaseElementType, *ClassName);
 }
 
-const FRigVMCodeGenerator::FOperationGroup& FRigVMCodeGenerator::GetGroup(int32 InGroupIndex) const
+const FRigVMCodeGenerator::FInstructionGroup& FRigVMCodeGenerator::GetGroup(int32 InGroupIndex) const
 {
-	if(OperationGroups.IsValidIndex(InGroupIndex))
+	if(InstructionGroups.IsValidIndex(InGroupIndex))
 	{
-		return OperationGroups[InGroupIndex];
+		return InstructionGroups[InGroupIndex];
 	}
 	
-	static FOperationGroup CompleteGroup;
+	static FInstructionGroup CompleteGroup;
 	if(CompleteGroup.First == INDEX_NONE)
 	{
 		CompleteGroup.First = 0;
@@ -2006,17 +2403,17 @@ const FRigVMCodeGenerator::FOperationGroup& FRigVMCodeGenerator::GetGroup(int32
 	return CompleteGroup;
 }
 
-bool FRigVMCodeGenerator::IsOperationPartOfGroup(int32 InOperationIndex, int32 InGroupIndex, bool bIncludeChildGroups) const
+bool FRigVMCodeGenerator::IsInstructionPartOfGroup(int32 InInstructionIndex, int32 InGroupIndex, bool bIncludeChildGroups) const
 {
-	if(!OperationGroups.IsValidIndex(InGroupIndex) && OperationGroups.IsEmpty())
+	if(!InstructionGroups.IsValidIndex(InGroupIndex) && InstructionGroups.IsEmpty())
 	{
 		return true;
 	}
 
-	const FOperationGroup& Group = OperationGroups[InGroupIndex];
+	const FInstructionGroup& Group = InstructionGroups[InGroupIndex];
 	for(const int32 ChildGroupIndex : Group.ChildGroups)
 	{
-		if(IsOperationPartOfGroup(InOperationIndex, ChildGroupIndex, true))
+		if(IsInstructionPartOfGroup(InInstructionIndex, ChildGroupIndex, true))
 		{
 			// if we found the operation in a child group
 			// we'll return true if we are also looking within the child groups,
@@ -2026,7 +2423,7 @@ bool FRigVMCodeGenerator::IsOperationPartOfGroup(int32 InOperationIndex, int32 I
 		}
 	}
 
-	return FMath::IsWithinInclusive(InOperationIndex, Group.First, Group.Last);;
+	return FMath::IsWithinInclusive(InInstructionIndex, Group.First, Group.Last);;
 }
 
 bool FRigVMCodeGenerator::IsPropertyPartOfGroup(int32 InPropertyIndex, int32 InGroupIndex) const
@@ -2036,7 +2433,7 @@ bool FRigVMCodeGenerator::IsPropertyPartOfGroup(int32 InPropertyIndex, int32 InG
 		return true;
 	}
 
-	if(!OperationGroups.IsValidIndex(InGroupIndex))
+	if(!InstructionGroups.IsValidIndex(InGroupIndex))
 	{
 		return false;
 	}
@@ -2044,7 +2441,7 @@ bool FRigVMCodeGenerator::IsPropertyPartOfGroup(int32 InPropertyIndex, int32 InG
 	if(Properties.IsValidIndex(InPropertyIndex))
 	{
 		const FPropertyInfo& Property = Properties[InPropertyIndex];
-		const FOperationGroup& Group = OperationGroups[InGroupIndex];
+		const FInstructionGroup& Group = InstructionGroups[InGroupIndex];
 
 		// properties are only defines in the leaf groups
 		if(!Group.ChildGroups.IsEmpty())
@@ -2066,6 +2463,37 @@ bool FRigVMCodeGenerator::IsPropertyPartOfGroup(int32 InPropertyIndex, int32 InG
 	return false;
 }
 
+FString FRigVMCodeGenerator::GetEntryParameters() const
+{
+	const FRigVMByteCode& ByteCode = VM->GetByteCode();
+	const int32 EntryIndex = ByteCode.NumEntries() - 1;
+	const FRigVMByteCodeEntry& Entry = ByteCode.GetEntry(EntryIndex);
+	const FString EntryName = Entry.GetSanitizedName();
+	FString Parameters;
+
+	// find the entry's group and provide the needed arguments
+	for(const FInstructionGroup& EntryGroup : InstructionGroups)
+	{
+		if(EntryGroup.Depth <= 0 && EntryGroup.Entry == EntryName)
+		{
+			TArray<FString> ArgumentNames = {RigVM_ContextPublicFormat};
+			for(const FOpaqueArgument& OpaqueArgument : EntryGroup.OpaqueArguments)
+			{
+				ArgumentNames.Add(OpaqueArgument.Get<1>());
+			}
+			Parameters = FString::Join(ArgumentNames, RigVM_CommaSeparator);
+			break;
+		}
+	}
+
+	if(!Parameters.IsEmpty())
+	{
+		Parameters = RigVM_CommaSeparator + Parameters;
+	}
+		
+	return Parameters;
+}
+
 const TMap<FName, FRigVMCodeGenerator::TStructConstGenerator>& FRigVMCodeGenerator::GetStructConstGenerators()
 {
 	static TMap<FName, FRigVMCodeGenerator::TStructConstGenerator> StructConstGenerators;
diff --git a/Engine/Source/Developer/RigVMDeveloper/Private/RigVMModel/Nodes/RigVMTemplateNode.cpp b/Engine/Source/Developer/RigVMDeveloper/Private/RigVMModel/Nodes/RigVMTemplateNode.cpp
--- a/Engine/Source/Developer/RigVMDeveloper/Private/RigVMModel/Nodes/RigVMTemplateNode.cpp
+++ b/Engine/Source/Developer/RigVMDeveloper/Private/RigVMModel/Nodes/RigVMTemplateNode.cpp
@@ -565,7 +565,7 @@ const FRigVMFunction* URigVMTemplateNode::GetResolvedFunction() const
 		{
 			if(ResolvedPermutation != INDEX_NONE)
 			{
-				CachedFunction = GetTemplate()->GetPermutation(ResolvedPermutation);
+				CachedFunction = ((FRigVMTemplate*)GetTemplate())->GetOrCreatePermutation(ResolvedPermutation);
 			}
 		}
 	}
diff --git a/Engine/Source/Developer/RigVMDeveloper/Public/RigVMCompiler/RigVMCodeGenerator.h b/Engine/Source/Developer/RigVMDeveloper/Public/RigVMCompiler/RigVMCodeGenerator.h
--- a/Engine/Source/Developer/RigVMDeveloper/Public/RigVMCompiler/RigVMCodeGenerator.h
+++ b/Engine/Source/Developer/RigVMDeveloper/Public/RigVMCompiler/RigVMCodeGenerator.h
@@ -16,17 +16,21 @@ public:
 
 	FRigVMCodeGenerator(const FString& InClassName, const FString& InModuleName,
 		URigVMGraph* InModelToNativize, URigVM* InVMToNativize, TMap<FString,FRigVMOperand> InPinToOperandMap,
-		int32 InMaxOperationsPerFunction = 100)
+		int32 InMaxInstructionsPerFunction = 100)
 	{
-		ParseVM(InClassName, InModuleName, InModelToNativize, InVMToNativize, InPinToOperandMap, InMaxOperationsPerFunction);
+		ParseVM(InClassName, InModuleName, InModelToNativize, InVMToNativize, InPinToOperandMap, InMaxInstructionsPerFunction);
 	}
 
 	FString DumpIncludes(bool bLog = false);
 	FString DumpWrappedArrayTypes(bool bLog = false);
 	FString DumpExternalVariables(bool bForHeader, bool bLog = false);
 	FString DumpEntries(bool bForHeader, bool bLog = false);
-	FString DumpProperties(bool bForHeader, int32 InOperationGroup, bool bLog = false);
-	FString DumpInstructions(int32 InOperationGroup, bool bLog = false);
+	FString DumpBlockNames(bool bForHeader, bool bLog = false);
+	FString DumpProperties(bool bForHeader, int32 InInstructionGroup, bool bLog = false);
+	FString DumpDispatches(bool bLog = false);
+	FString DumpRequiredUProperties(bool bLog = false);
+	FString DumpInitialize(bool bLog = false);
+	FString DumpInstructions(int32 InInstructionGroup, bool bLog = false);
 	FString DumpHeader(bool bLog = false);
 	FString DumpSource(bool bLog = false);
 	FString DumpLines(const TArray<FString>& InLines, bool bLog = false);
@@ -56,7 +60,7 @@ private:
 		TArray<int32> Groups;
 	};
 
-	struct FOperationGroup
+	struct FInstructionGroup
 	{
 		FString Entry;
 		int32 First;
@@ -67,7 +71,7 @@ private:
 		TArray<int32> RequiredLabels;
 		TArray<FOpaqueArgument> OpaqueArguments;
 
-		FOperationGroup()
+		FInstructionGroup()
 			: Entry()
 			, First(INDEX_NONE)
 			, Last(INDEX_NONE)
@@ -75,7 +79,7 @@ private:
 			, ParentGroup(INDEX_NONE)
 		{}
 
-		FOperationGroup(const FOperationGroup& InOther)
+		FInstructionGroup(const FInstructionGroup& InOther)
 			: Entry(InOther.Entry)
 			, First(InOther.First)
 			, Last(InOther.Last)
@@ -90,11 +94,14 @@ private:
 	void Reset();
 	void ParseVM(const FString& InClassName, const FString& InModuleName,
 		URigVMGraph* InModelToNativize, URigVM* InVMToNativize,
-		TMap<FString,FRigVMOperand> InPinToOperandMap, int32 InMaxOperationsPerFunction);
+		TMap<FString,FRigVMOperand> InPinToOperandMap, int32 InMaxInstructionsPerFunction);
 	void ParseInclude(UStruct* InDependency, const FName& InMethodName = NAME_None);
+	void ParseRequiredUProperties();
 	void ParseMemory(URigVMMemoryStorage* InMemory);
 	void ParseProperty(ERigVMMemoryType InMemoryType, const FProperty* InProperty, URigVMMemoryStorage* InMemory);
-	void ParseOperationGroups();
+	void ParseInstructionGroups();
+	FString DumpInstructions(const FString& InPrefix, int32 InFirstInstruction, int32 InLastInstruction, const FInstructionGroup& InGroup, bool bLog = false);
+	FString DumpInstructions(const FString& InPrefix, const TArray<int32> InInstructionIndices, const FInstructionGroup& InGroup, bool bLog = false);
 	FString GetOperandName(const FRigVMOperand& InOperand, bool bPerSlice = true, bool bAsInput = true) const;
 	FString GetOperandCPPType(const FRigVMOperand& InOperand) const;
 	FString GetOperandCPPBaseType(const FRigVMOperand& InOperand) const;
@@ -109,9 +116,11 @@ private:
 	FString GetMappedType(const FString& InCPPType) const;
 	const FString& GetMappedTypeSuffix(const FString& InCPPType) const;
 	FString GetMappedArrayTypeName(const FString InBaseElementType) const;
-	const FOperationGroup& GetGroup(int32 InGroupIndex) const;
-	bool IsOperationPartOfGroup(int32 InOperationIndex, int32 InGroupIndex, bool bIncludeChildGroups) const;
+	const FInstructionGroup& GetGroup(int32 InGroupIndex) const;
+	bool IsInstructionPartOfGroup(int32 InInstructionIndex, int32 InGroupIndex, bool bIncludeChildGroups) const;
 	bool IsPropertyPartOfGroup(int32 InPropertyIndex, int32 InGroupIndex) const;
+	FString GetEntryParameters() const;
+	static TArray<int32> GetInstructionIndicesFromRange(int32 First, int32 Last);
 
 	FORCEINLINE static FString FormatArgs(const TCHAR* InFormatString, const FStringFormatOrderedArguments& InArgs)
 	{
@@ -225,19 +234,22 @@ private:
 	TStrongObjectPtr<URigVM> VM;
 	TMap<FString,FRigVMOperand> PinToOperandMap;
 	TMap<FRigVMOperand, FString> OperandToPinMap;
-	int32 MaxOperationsPerFunction;
+	int32 MaxInstructionsPerFunction;
 	FString ClassName;
 	FString ModuleName;
 	FString ExecuteContextType;
 	
 	FStringArray Libraries;
 	FStringArray Includes;
-	TArray<FOperationGroup> OperationGroups;
+	TMap<FString,const FRigVMFunction*> Dispatches;
+	TMap<TRigVMTypeIndex,TTuple<FString,FString>> RequiredUProperties;
+	TArray<FInstructionGroup> InstructionGroups;
 	FStringArray WrappedArrayCPPTypes;
 	TMap<FString, FMappedType> MappedCPPTypes;
 	TArray<FPropertyInfo> Properties;
 	TMap<FName, int32> PropertyNameToIndex;
 	TMap<FName, FName> MappedPropertyNames;
+	TMap<FString, FString> OverriddenOperatorNames;
 
 	static const TMap<FName, TStructConstGenerator>& GetStructConstGenerators();
 };
diff --git a/Engine/Source/Runtime/RigVM/Private/RigVMCore/RigVM.cpp b/Engine/Source/Runtime/RigVM/Private/RigVMCore/RigVM.cpp
--- a/Engine/Source/Runtime/RigVM/Private/RigVMCore/RigVM.cpp
+++ b/Engine/Source/Runtime/RigVM/Private/RigVMCore/RigVM.cpp
@@ -914,13 +914,16 @@ const TArray<FName>& URigVM::GetEntryNames() const
 	return EntryNames;
 }
 
-bool URigVM::CanExecuteEntry(const FName& InEntryName) const
+bool URigVM::CanExecuteEntry(const FName& InEntryName, bool bLogErrorForMissingEntry) const
 {
 	const int32 EntryIndex = FindEntry(InEntryName);
 	if(EntryIndex == INDEX_NONE)
 	{
-		static constexpr TCHAR MissingEntry[] = TEXT("Entry('%s') cannot be found.");
-		Context.GetPublicData<>().Logf(EMessageSeverity::Error, MissingEntry, *InEntryName.ToString());
+		if(bLogErrorForMissingEntry)
+		{
+			static constexpr TCHAR MissingEntry[] = TEXT("Entry('%s') cannot be found.");
+			Context.GetPublicData<>().Logf(EMessageSeverity::Error, MissingEntry, *InEntryName.ToString());
+		}
 		return false;
 	}
 	
@@ -1659,16 +1662,21 @@ bool URigVM::Initialize(TArrayView<URigVMMemoryStorage*> Memory, TArrayView<void
 				// find out the largest slice count
 				int32 MaxSliceCount = 1;
 
-				// todo Deal with slice counts
+				const bool bContainsLazyValue = OpHandles.ContainsByPredicate([](const FRigVMMemoryHandle& Handle)
+				{
+					return Handle.IsLazy();
+				});
 
-				Context.BeginSlice(MaxSliceCount);
-				for (int32 SliceIndex = 0; SliceIndex < MaxSliceCount; SliceIndex++)
+				if(!bContainsLazyValue)
 				{
-					(*Functions[Op.FunctionIndex]->FunctionPtr)(Context, OpHandles);
-					Context.IncrementSlice();
+					Context.BeginSlice(MaxSliceCount);
+					for (int32 SliceIndex = 0; SliceIndex < MaxSliceCount; SliceIndex++)
+					{
+						(*Functions[Op.FunctionIndex]->FunctionPtr)(Context, OpHandles);
+						Context.IncrementSlice();
+					}
+					Context.EndSlice();
 				}
-				Context.EndSlice();
-
 				break;
 			}
 			case ERigVMOpCode::Zero:
@@ -1785,18 +1793,7 @@ ERigVMExecuteResult URigVM::Execute(TArrayView<URigVMMemoryStorage*> Memory, TAr
 	{
 		if (FirstEntryEventInQueue == NAME_None || FirstEntryEventInQueue == InEntryName)
 		{
-			InstructionVisitedDuringLastRun.Reset();
-			InstructionVisitOrder.Reset();
-			InstructionVisitedDuringLastRun.SetNumZeroed(Instructions.Num());
-			InstructionCyclesDuringLastRun.Reset();
-			if(ContextPublicData.RuntimeSettings.bEnableProfiling)
-			{
-				InstructionCyclesDuringLastRun.SetNumUninitialized(Instructions.Num());
-				for(int32 DurationIndex=0;DurationIndex<InstructionCyclesDuringLastRun.Num();DurationIndex++)
-				{
-					InstructionCyclesDuringLastRun[DurationIndex] = UINT64_MAX;
-				}
-			}
+			SetupInstructionTracking(Instructions.Num());
 		}
 	}
 #endif
@@ -1871,12 +1868,7 @@ ERigVMExecuteResult URigVM::Execute(TArrayView<URigVMMemoryStorage*> Memory, TAr
 	}
 
 #if WITH_EDITOR
-	StartCycles = 0;
-	OverallCycles = 0;
-	if(ContextPublicData.RuntimeSettings.bEnableProfiling)
-	{
-		StartCycles = FPlatformTime::Cycles64();
-	}
+	StartProfiling();
 	
 #if UE_RIGVM_DEBUG_EXECUTION
 	FString DebugMemoryString;
@@ -2248,9 +2240,7 @@ ERigVMExecuteResult URigVM::ExecuteInstructions(int32 InFirstInstruction, int32
 			{
 				if(bCurrentlyRunningRootEntry)
 				{
-#if WITH_EDITOR
-					Context.LastExecutionMicroSeconds = OverallCycles * FPlatformTime::GetSecondsPerCycle() * 1000.0 * 1000.0;
-#endif
+					StopProfiling();
 					ExecutionReachedExit().Broadcast(CurrentEntryName);
 #if WITH_EDITOR					
 					if (HaltedAtBreakpoint.IsValid())
@@ -3617,3 +3607,40 @@ void URigVM::CopyArray(FScriptArrayHelper& TargetHelper, FRigVMMemoryHandle& Tar
 	}
 }
 
+void URigVM::SetupInstructionTracking(int32 InInstructionCount)
+{
+#if WITH_EDITOR
+	InstructionVisitedDuringLastRun.Reset();
+	InstructionVisitOrder.Reset();
+	InstructionVisitedDuringLastRun.SetNumZeroed(InInstructionCount);
+	InstructionCyclesDuringLastRun.Reset();
+
+	if(Context.GetPublicData<>().RuntimeSettings.bEnableProfiling)
+	{
+		InstructionCyclesDuringLastRun.SetNumUninitialized(InInstructionCount);
+		for(int32 DurationIndex=0;DurationIndex<InstructionCyclesDuringLastRun.Num();DurationIndex++)
+		{
+			InstructionCyclesDuringLastRun[DurationIndex] = UINT64_MAX;
+		}
+	}
+#endif
+}
+
+void URigVM::StartProfiling()
+{
+#if WITH_EDITOR
+	OverallCycles = StartCycles = 0;
+	if(Context.GetPublicData<>().RuntimeSettings.bEnableProfiling)
+	{
+		StartCycles = FPlatformTime::Cycles64();
+	}
+#endif
+}
+
+void URigVM::StopProfiling()
+{
+#if WITH_EDITOR
+	const uint64 Cycles = OverallCycles > 0 ? OverallCycles : (FPlatformTime::Cycles64() - StartCycles); 
+	Context.LastExecutionMicroSeconds = Cycles * FPlatformTime::GetSecondsPerCycle() * 1000.0 * 1000.0;
+#endif
+}
diff --git a/Engine/Source/Runtime/RigVM/Private/RigVMCore/RigVMDispatchFactory.cpp b/Engine/Source/Runtime/RigVM/Private/RigVMCore/RigVMDispatchFactory.cpp
--- a/Engine/Source/Runtime/RigVM/Private/RigVMCore/RigVMDispatchFactory.cpp
+++ b/Engine/Source/Runtime/RigVM/Private/RigVMCore/RigVMDispatchFactory.cpp
@@ -124,6 +124,11 @@ FString FRigVMDispatchFactory::GetKeywords() const
 	return FString();
 }
 
+bool FRigVMDispatchFactory::IsLazyInputArgument(const FName& InArgumentName) const
+{
+	return HasArgumentMetaData(InArgumentName, FRigVMStruct::ComputeLazilyMetaName);
+}
+
 #endif
 
 const TArray<FName>& FRigVMDispatchFactory::GetControlFlowBlocks() const
diff --git a/Engine/Source/Runtime/RigVM/Private/RigVMCore/RigVMMemoryStorage.cpp b/Engine/Source/Runtime/RigVM/Private/RigVMCore/RigVMMemoryStorage.cpp
--- a/Engine/Source/Runtime/RigVM/Private/RigVMCore/RigVMMemoryStorage.cpp
+++ b/Engine/Source/Runtime/RigVM/Private/RigVMCore/RigVMMemoryStorage.cpp
@@ -81,7 +81,13 @@ void FRigVMBranchInfo::Serialize(FArchive& Ar)
 ERigVMExecuteResult FRigVMLazyBranch::Execute()
 {
 	check(VM);
-	return VM->ExecuteLazyBranch(BranchInfo);
+	if(BranchInfo.IsValid())
+	{
+		return VM->ExecuteLazyBranch(BranchInfo);
+	}
+
+	check(FunctionPtr);
+	return FunctionPtr();
 }
 
 ERigVMExecuteResult FRigVMLazyBranch::ExecuteIfRequired(int32 InSliceIndex)
diff --git a/Engine/Source/Runtime/RigVM/Private/RigVMCore/RigVMNativized.cpp b/Engine/Source/Runtime/RigVM/Private/RigVMCore/RigVMNativized.cpp
--- a/Engine/Source/Runtime/RigVM/Private/RigVMCore/RigVMNativized.cpp
+++ b/Engine/Source/Runtime/RigVM/Private/RigVMCore/RigVMNativized.cpp
@@ -19,6 +19,14 @@ void URigVMNativized::Serialize(FArchive& Ar)
 	UObject::Serialize(Ar);
 }
 
+void URigVMNativized::Reset(bool IsIgnoringArchetypeRef)
+{
+	// don't call super on purpose
+	LazyMemoryHandles.Reset();
+	LazyMemoryBranches.Reset();
+	ByteCodeStorage.Reset();
+}
+
 bool URigVMNativized::Initialize(TArrayView<URigVMMemoryStorage*> Memory, TArrayView<void*> AdditionalArguments,
 	bool bInitializeMemory)
 {
@@ -35,6 +43,48 @@ ERigVMExecuteResult URigVMNativized::Execute(TArrayView<URigVMMemoryStorage*> Me
 
 const FRigVMInstructionArray& URigVMNativized::GetInstructions()
 {
+	if(GetByteCode().GetNumInstructions() > 0)
+	{
+		return Super::GetInstructions();
+	}
 	static const FRigVMInstructionArray EmptyInstructions;
 	return EmptyInstructions;
 }
+
+#if WITH_EDITOR
+
+void URigVMNativized::SetByteCode(const FRigVMByteCode& InByteCode)
+{
+	ByteCodeStorage = InByteCode;
+	ByteCodePtr = &ByteCodeStorage;
+}
+
+#endif
+
+const FRigVMMemoryHandle& URigVMNativized::GetLazyMemoryHandle(int32 InIndex, uint8* InMemory, const FProperty* InProperty, TFunction<ERigVMExecuteResult()> InLambda)
+{
+	check(InIndex != INDEX_NONE);
+	check(LazyMemoryHandles.IsValidIndex(InIndex));
+
+	FRigVMMemoryHandle& Handle = LazyMemoryHandles[InIndex];
+	if(!Handle.IsLazy())
+	{
+		check(LazyMemoryBranches.IsValidIndex(InIndex));
+		FRigVMLazyBranch& LazyBranch = LazyMemoryBranches[InIndex];
+		LazyBranch.VM = this;
+		LazyBranch.FunctionPtr = InLambda;
+
+		Handle = FRigVMMemoryHandle(InMemory, InProperty, nullptr, &LazyBranch);
+	}
+
+	return Handle;
+}
+
+void URigVMNativized::AllocateLazyMemoryHandles(int32 InCount)
+{
+	if(LazyMemoryHandles.Num() != InCount)
+	{
+		LazyMemoryHandles.SetNum(InCount);
+		LazyMemoryBranches.SetNum(InCount);
+	}
+}
diff --git a/Engine/Source/Runtime/RigVM/Public/RigVMCore/RigVM.h b/Engine/Source/Runtime/RigVM/Public/RigVMCore/RigVM.h
--- a/Engine/Source/Runtime/RigVM/Public/RigVMCore/RigVM.h
+++ b/Engine/Source/Runtime/RigVM/Public/RigVMCore/RigVM.h
@@ -298,7 +298,7 @@ public:
 	virtual const TArray<FName>& GetEntryNames() const;
 
 	// returns false if an entry can not be executed
-	bool CanExecuteEntry(const FName& InEntryName) const;
+	bool CanExecuteEntry(const FName& InEntryName, bool bLogErrorForMissingEntry = true) const;
 
 #if WITH_EDITOR
 	
@@ -711,12 +711,13 @@ protected:
 	UPROPERTY(transient)
 	FRigVMExtendedExecuteContext Context;
 
-	FORCEINLINE void SetInstructionIndex(uint16 InInstructionIndex) { Context.GetPublicData<>().InstructionIndex = InInstructionIndex; }
+	virtual void SetInstructionIndex(uint16 InInstructionIndex) { Context.GetPublicData<>().InstructionIndex = InInstructionIndex; }
 
-private:
 	UPROPERTY(transient)
 	uint32 NumExecutions;
 
+private:
+
 #if WITH_EDITOR
 	FRigVMDebugInfo* DebugInfo;
 	FRigVMBreakpoint HaltedAtBreakpoint;
@@ -772,11 +773,14 @@ private:
 	
 #if WITH_EDITOR
 
+protected:
+	
 	// stores the number of times each instruction was visited
 	TArray<int32> InstructionVisitedDuringLastRun;
 	TArray<uint64> InstructionCyclesDuringLastRun;
 	TArray<int32> InstructionVisitOrder;
-	
+
+private:
 	// Control Rig can run multiple events per evaluation, such as the Backward&Forward Solve Mode,
 	// store the first event such that we know when to reset data for a new round of rig evaluation
 	FName FirstEntryEventInQueue;
@@ -847,8 +851,10 @@ private:
 	TArrayView<void*> CurrentAdditionalArguments;
 
 #if WITH_EDITOR
+protected:
 	uint64 StartCycles = 0;
 	uint64 OverallCycles = 0;
+private:
 #endif
 
 	UPROPERTY(transient)
@@ -861,6 +867,12 @@ private:
 	FExecutionHaltedEvent OnExecutionHalted;
 #endif
 
+protected:
+
+	void SetupInstructionTracking(int32 InInstructionCount);
+	void StartProfiling();
+	void StopProfiling();
+
 	friend class URigVMCompiler;
 	friend struct FRigVMCompilerWorkData;
 	friend struct FRigVMCodeGenerator;
diff --git a/Engine/Source/Runtime/RigVM/Public/RigVMCore/RigVMDispatchFactory.h b/Engine/Source/Runtime/RigVM/Public/RigVMCore/RigVMDispatchFactory.h
--- a/Engine/Source/Runtime/RigVM/Public/RigVMCore/RigVMDispatchFactory.h
+++ b/Engine/Source/Runtime/RigVM/Public/RigVMCore/RigVMDispatchFactory.h
@@ -86,6 +86,9 @@ public:
 	// Returns the keywords used for looking up this factory
 	virtual FString GetKeywords() const;
 
+	// Returns true if the argument is lazy
+	bool IsLazyInputArgument(const FName& InArgumentName) const;
+
 #endif
 
 	// returns true if the dispatch is a control flow dispatch
diff --git a/Engine/Source/Runtime/RigVM/Public/RigVMCore/RigVMMemoryStorage.h b/Engine/Source/Runtime/RigVM/Public/RigVMCore/RigVMMemoryStorage.h
--- a/Engine/Source/Runtime/RigVM/Public/RigVMCore/RigVMMemoryStorage.h
+++ b/Engine/Source/Runtime/RigVM/Public/RigVMCore/RigVMMemoryStorage.h
@@ -67,6 +67,16 @@ struct RIGVM_API FRigVMBranchInfo
 		return InstructionIndex != INDEX_NONE;
 	}
 
+	bool IsOutputBranch() const
+	{
+		return ArgumentIndex == INDEX_NONE;
+	}
+
+	bool IsLazyEvaluationBranch() const
+	{
+		return !IsOutputBranch();
+	}
+	
 	void Serialize(FArchive& Ar);
 	FORCEINLINE friend FArchive& operator<<(FArchive& Ar, FRigVMBranchInfo& P)
 	{
@@ -77,7 +87,7 @@ struct RIGVM_API FRigVMBranchInfo
 	friend uint32 GetTypeHash(const FRigVMBranchInfo& InBranchInfo)
 	{
 		uint32 Hash = GetTypeHash(GetTypeHash(InBranchInfo.Index));
-		Hash = HashCombine(Hash, GetTypeHash(InBranchInfo.Label));
+		Hash = HashCombine(Hash, GetTypeHash(InBranchInfo.Label.ToString()));
 		Hash = HashCombine(Hash, GetTypeHash(InBranchInfo.InstructionIndex));
 		Hash = HashCombine(Hash, GetTypeHash(InBranchInfo.ArgumentIndex));
 		Hash = HashCombine(Hash, GetTypeHash(InBranchInfo.FirstInstruction));
@@ -116,6 +126,7 @@ public:
 		: VM(nullptr)
 		, LastVMNumExecutions()
 		, BranchInfo()
+		, FunctionPtr(nullptr)
 	{}
 
 	ERigVMExecuteResult Execute();
@@ -126,8 +137,10 @@ private:
 	URigVM* VM;
 	TArray<int32> LastVMNumExecutions;
 	FRigVMBranchInfo BranchInfo;
+	TFunction<ERigVMExecuteResult()> FunctionPtr;
 
 	friend class URigVM;
+	friend class URigVMNativized;
 };
 
 /**
@@ -145,6 +158,7 @@ public:
 	}
 
 	const uint8* GetData() const;
+	const FRigVMMemoryHandle& GetMemoryHandle()const { return *MemoryHandle; }
 
 protected:
 	
@@ -220,11 +234,11 @@ public:
 	{}
 
 	// Constructor from complete data
-	FORCEINLINE_DEBUGGABLE FRigVMMemoryHandle(uint8* InPtr, const FProperty* InProperty,  const FRigVMPropertyPath* InPropertyPath)
+	FORCEINLINE_DEBUGGABLE FRigVMMemoryHandle(uint8* InPtr, const FProperty* InProperty,  const FRigVMPropertyPath* InPropertyPath, FRigVMLazyBranch* InLazyBranch = nullptr)
 		: Ptr(InPtr)
 		, Property(InProperty)
 		, PropertyPath(InPropertyPath)
-		, LazyBranch(nullptr)
+		, LazyBranch(InLazyBranch)
 	{
 	}
 
diff --git a/Engine/Source/Runtime/RigVM/Public/RigVMCore/RigVMNativized.h b/Engine/Source/Runtime/RigVM/Public/RigVMCore/RigVMNativized.h
--- a/Engine/Source/Runtime/RigVM/Public/RigVMCore/RigVMNativized.h
+++ b/Engine/Source/Runtime/RigVM/Public/RigVMCore/RigVMNativized.h
@@ -46,7 +46,7 @@ public:
 	// URigVM interface
 	virtual void Serialize(FArchive& Ar) override;
 	virtual void PostLoad() override { return; }
-	virtual void Reset(bool IsIgnoringArchetypeRef = false) override { return; }
+	virtual void Reset(bool IsIgnoringArchetypeRef = false) override;
 	virtual bool IsNativized() const override { return true; }
 	virtual void Empty() override { return; }
 	virtual void CopyFrom(URigVM* InVM, bool bDeferCopy = false, bool bReferenceLiteralMemory = false, bool bReferenceByteCode = false, bool bCopyExternalVariables = false, bool bCopyDynamicRegisters = false) override { return; }
@@ -60,6 +60,18 @@ public:
 	virtual bool ContainsEntry(const FName& InEntryName) const override { return GetEntryNames().Contains(InEntryName); }
 	virtual int32 FindEntry(const FName& InEntryName) const override { return GetEntryNames().Find(InEntryName); }
 	virtual const TArray<FName>& GetEntryNames() const override { static const TArray<FName> EmptyEntries; return EmptyEntries; }
+	virtual void SetInstructionIndex(uint16 InInstructionIndex) override
+	{
+		Super::SetInstructionIndex(InInstructionIndex);
+#if WITH_EDITOR
+		InstructionVisitedDuringLastRun[InInstructionIndex]++;
+		InstructionVisitOrder.Add(InInstructionIndex);
+#endif
+	}
+
+#if WITH_EDITOR
+	void SetByteCode(const FRigVMByteCode& InByteCode);
+#endif
 	
 protected:
 
@@ -69,6 +81,7 @@ protected:
 		UpdateExternalVariables();
 	
 		Context.Reset();
+		Context.VM = this;
 		Context.OpaqueArguments = AdditionalArguments;
 		Context.SliceOffsets.AddZeroed(InNumberInstructions);
 		Context.GetPublicData<ExecuteContextType>().EventName = InEntryName;
@@ -155,6 +168,7 @@ protected:
 		{
 			ExecutionReachedExit().Broadcast(Context.GetPublicData<>().GetEventName());
 		}
+		NumExecutions++;
 	}
 
 	template<typename T>
@@ -420,5 +434,33 @@ protected:
 		FTransform& Transform;
 	};
 
+	template <typename T,
+			  typename TEnableIf<!TIsArray<T>::Value, int>::Type = 0>
+	bool ImportDefaultValue(T& OutValue, const FProperty* Property, const FString& InBuffer)
+	{
+		return false;
+	}
+
+	template <typename T,
+			  typename TEnableIf<TIsArray<T>::Value, int>::Type = 0>
+	bool ImportDefaultValue(T& OutValue, const FProperty* Property, const FString& InBuffer)
+	{
+		return false;
+	}
+
 	int32 TemporaryArrayIndex;
+
+	const FRigVMMemoryHandle& GetLazyMemoryHandle(int32 InIndex, uint8* InMemory, const FProperty* InProperty, TFunction<ERigVMExecuteResult()> InLambda);
+	
+	template<typename T>
+	TRigVMLazyValue<T> GetLazyValue(int32 InIndex, T& InMemory, const FProperty* InProperty, TFunction<ERigVMExecuteResult()> InLambda)
+	{
+		const FRigVMMemoryHandle& Handle = GetLazyMemoryHandle(InIndex, (uint8*)&InMemory, InProperty, InLambda);
+		return Handle.GetDataLazily<T>(false, INDEX_NONE);
+	}
+
+	void AllocateLazyMemoryHandles(int32 InCount);
+
+	TArray<FRigVMLazyBranch> LazyMemoryBranches;
+	TArray<FRigVMMemoryHandle> LazyMemoryHandles;
 };
