RigVM: deprecate branch node

diff --git a/Engine/Plugins/Animation/ControlRig/Source/ControlRigEditor/Private/ControlRigEditorModule.cpp b/Engine/Plugins/Animation/ControlRig/Source/ControlRigEditor/Private/ControlRigEditorModule.cpp
--- a/Engine/Plugins/Animation/ControlRig/Source/ControlRigEditor/Private/ControlRigEditorModule.cpp
+++ b/Engine/Plugins/Animation/ControlRig/Source/ControlRigEditor/Private/ControlRigEditorModule.cpp
@@ -49,7 +49,6 @@
 #include "Graph/NodeSpawners/ControlRigUnitNodeSpawner.h"
 #include "Graph/NodeSpawners/ControlRigVariableNodeSpawner.h"
 #include "Graph/NodeSpawners/ControlRigRerouteNodeSpawner.h"
-#include "Graph/NodeSpawners/ControlRigBranchNodeSpawner.h"
 #include "Graph/NodeSpawners/ControlRigIfNodeSpawner.h"
 #include "Graph/NodeSpawners/ControlRigSelectNodeSpawner.h"
 #include "Graph/NodeSpawners/ControlRigTemplateNodeSpawner.h"
@@ -1113,12 +1112,6 @@ void FControlRigEditorModule::GetTypeActions(UControlRigBlueprint* CRB, FBluepri
 		LOCTEXT("RerouteSpawnerTooltip", "Adds a new reroute node to the graph"));
 	ActionRegistrar.AddBlueprintAction(ActionKey, RerouteNodeSpawner);
 
-	UBlueprintNodeSpawner* BranchNodeSpawner = UControlRigBranchNodeSpawner::CreateGeneric(
-		LOCTEXT("BranchSpawnerDesc", "Branch"),
-		LOCTEXT("BranchSpawnerCategory", "Execution"),
-		LOCTEXT("BranchSpawnerTooltip", "Adds a new 'branch' node to the graph"));
-	ActionRegistrar.AddBlueprintAction(ActionKey, BranchNodeSpawner);
-
 	UBlueprintNodeSpawner* IfNodeSpawner = UControlRigIfNodeSpawner::CreateGeneric(
 		LOCTEXT("IfSpawnerDesc", "If"),
 		LOCTEXT("IfSpawnerCategory", "Execution"),
diff --git a/Engine/Plugins/Animation/ControlRig/Source/ControlRigEditor/Private/Editor/ControlRigEditor.cpp b/Engine/Plugins/Animation/ControlRig/Source/ControlRigEditor/Private/Editor/ControlRigEditor.cpp
--- a/Engine/Plugins/Animation/ControlRig/Source/ControlRigEditor/Private/Editor/ControlRigEditor.cpp
+++ b/Engine/Plugins/Animation/ControlRig/Source/ControlRigEditor/Private/Editor/ControlRigEditor.cpp
@@ -100,6 +100,7 @@
 #include "ToolMenus.h"
 #include "Styling/AppStyle.h"
 #include "AssetRegistry/AssetRegistryModule.h"
+#include "RigVMFunctions/RigVMFunction_ControlFlow.h"
 
 #define LOCTEXT_NAMESPACE "ControlRigEditor"
 
@@ -2670,7 +2671,7 @@ FReply FControlRigEditor::OnSpawnGraphNodeByShortcut(FInputChord InChord, const
 			{
 				if(InChord.Key == EKeys::B)
 				{
-					Controller->AddBranchNode(InPosition, "", true, true);
+					Controller->AddUnitNode(FRigVMFunction_ControlFlowBranch::StaticStruct(), FRigVMStruct::ExecuteName, InPosition, FString(), true, true);
 				}
 				else if(InChord.Key == EKeys::S)
 				{
diff --git a/Engine/Plugins/Animation/ControlRig/Source/ControlRigEditor/Private/Graph/NodeSpawners/ControlRigBranchNodeSpawner.cpp b/Engine/Plugins/Animation/ControlRig/Source/ControlRigEditor/Private/Graph/NodeSpawners/ControlRigBranchNodeSpawner.cpp
deleted file mode 100644
diff --git a/Engine/Plugins/Animation/ControlRig/Source/ControlRigEditor/Private/Graph/NodeSpawners/ControlRigBranchNodeSpawner.h b/Engine/Plugins/Animation/ControlRig/Source/ControlRigEditor/Private/Graph/NodeSpawners/ControlRigBranchNodeSpawner.h
deleted file mode 100644
diff --git a/Engine/Source/Developer/RigVMDeveloper/Private/RigVMCompiler/RigVMAST.cpp b/Engine/Source/Developer/RigVMDeveloper/Private/RigVMCompiler/RigVMAST.cpp
--- a/Engine/Source/Developer/RigVMDeveloper/Private/RigVMCompiler/RigVMAST.cpp
+++ b/Engine/Source/Developer/RigVMDeveloper/Private/RigVMCompiler/RigVMAST.cpp
@@ -7,7 +7,6 @@
 #include "RigVMModel/Nodes/RigVMVariableNode.h"
 #include "RigVMModel/Nodes/RigVMCommentNode.h"
 #include "RigVMModel/Nodes/RigVMRerouteNode.h"
-#include "RigVMModel/Nodes/RigVMBranchNode.h"
 #include "RigVMModel/Nodes/RigVMIfNode.h"
 #include "RigVMModel/Nodes/RigVMSelectNode.h"
 #include "RigVMModel/Nodes/RigVMEnumNode.h"
@@ -84,10 +83,6 @@ FName FRigVMExprAST::GetTypeName() const
 		{
 			return TEXT("[.Exit..]");
 		}
-		case EType::Branch:
-		{
-			return TEXT("[Branch.]");
-		}
 		case EType::If:
 		{
 			return TEXT("[..If...]");
@@ -620,42 +615,6 @@ bool FRigVMVarExprAST::SupportsSoftLinks() const
 	return false;
 }
 
-
-bool FRigVMBranchExprAST::IsConstant() const
-{
-	if (IsAlwaysTrue())
-	{
-		return GetTrueExpr()->IsConstant();
-	}
-	else if(IsAlwaysFalse())
-	{
-		return GetFalseExpr()->IsConstant();
-	}
-	return FRigVMNodeExprAST::IsConstant();
-}
-
-bool FRigVMBranchExprAST::IsAlwaysTrue() const
-{
-	const FRigVMVarExprAST* ConditionExpr = GetConditionExpr();
-	if (ConditionExpr->IsA(EType::Literal))
-	{
-		const FString& PinDefaultValue = ConditionExpr->GetDefaultValue();
-		return PinDefaultValue == TEXT("True");
-	}
-	return false;
-}
-
-bool FRigVMBranchExprAST::IsAlwaysFalse() const
-{
-	const FRigVMVarExprAST* ConditionExpr = GetConditionExpr();
-	if (ConditionExpr->IsA(EType::Literal))
-	{
-		const FString& PinDefaultValue = ConditionExpr->GetDefaultValue();
-		return PinDefaultValue == TEXT("False") || PinDefaultValue.IsEmpty();
-	}
-	return false;
-}
-
 bool FRigVMIfExprAST::IsConstant() const
 {
 	if (IsAlwaysTrue())
@@ -865,10 +824,6 @@ FRigVMParserAST::FRigVMParserAST(TArray<URigVMGraph*> InGraphs, URigVMController
 				bContinueToFoldConstantBranches = true;
 			}
 		}
-		if (FoldUnreachableBranches(InGraphs))
-		{
-			bContinueToFoldConstantBranches = true;
-		}
 	}
 
 	BubbleUpExpressions();
@@ -1077,10 +1032,6 @@ FRigVMExprAST* FRigVMParserAST::CreateExpressionForNode(const FRigVMASTProxy& In
 		{
 			NodeExpr = MakeExpr<FRigVMNoOpExprAST>(InNodeProxy);
 		}
-		else if (InNodeProxy.IsA<URigVMBranchNode>())
-		{
-			NodeExpr = MakeExpr<FRigVMBranchExprAST>(InNodeProxy);
-		}
 		else if (InNodeProxy.IsA<URigVMIfNode>())
 		{
 			NodeExpr = MakeExpr<FRigVMIfExprAST>(InNodeProxy);
@@ -2302,105 +2253,6 @@ bool FRigVMParserAST::FoldConstantValuesToLiterals(TArray<URigVMGraph*> InGraphs
 	return ExpressionsToRemove.Num() > 0;
 }
 
-bool FRigVMParserAST::FoldUnreachableBranches(TArray<URigVMGraph*> InGraphs)
-{
-	DECLARE_SCOPE_HIERARCHICAL_COUNTER_FUNC()
-
-	TArray<FRigVMExprAST*> ExpressionsToRemove;
-
-	for (const FRigVMASTProxy& NodeProxy : NodeProxies)
-	{
-		if (NodeProxy.IsA<URigVMVariableNode>())
-		{
-			continue;
-		}
-
-		//URigVMNode* Node = NodeProxy.GetSubjectChecked<URigVMNode>();
-
-		FRigVMExprAST** NodeExprPtr = SubjectToExpression.Find(NodeProxy);
-		if (NodeExprPtr == nullptr)
-		{
-			continue;
-		}
-
-		FRigVMExprAST* NodeExpr = *NodeExprPtr;
-		if (NodeExpr->NumParents() == 0)
-		{
-			continue;
-		}
-
-		if (NodeExpr->IsA(FRigVMExprAST::EType::Branch))
-		{
-			const FRigVMBranchExprAST* BranchExpr = NodeExpr->To<FRigVMBranchExprAST>();
-			FRigVMExprAST* ExprReplacement = nullptr;
-
-			if (BranchExpr->IsAlwaysTrue())
-			{
-				ExprReplacement = (FRigVMExprAST*)BranchExpr->GetTrueExpr();
-			}
-			else if (BranchExpr->IsAlwaysFalse())
-			{
-				ExprReplacement = (FRigVMExprAST*)BranchExpr->GetFalseExpr();
-			}
-
-			if (ExprReplacement)
-			{
-				if (ExprReplacement->NumChildren() == 1)
-				{
-					ExprReplacement = (FRigVMExprAST*)ExprReplacement->ChildAt(0);
-					if (ExprReplacement->IsA(FRigVMExprAST::EType::Block))
-					{
-						ExprReplacement->RemoveParent((FRigVMExprAST*)ExprReplacement->GetParent());
-						NodeExpr->ReplaceBy(ExprReplacement);
-						ExpressionsToRemove.Add(NodeExpr);
-					}
-				}
-			}
-		}
-		else
-		{
-			FRigVMExprAST* CachedValueExpr = (FRigVMExprAST*)NodeExpr->GetParent();
-			if (!CachedValueExpr->IsA(FRigVMExprAST::EType::CachedValue))
-			{
-				continue;
-			}
-
-			FRigVMExprAST* ExprReplacement = nullptr;
-			if (NodeExpr->IsA(FRigVMExprAST::EType::If))
-			{
-				const FRigVMIfExprAST* IfExpr = NodeExpr->To<FRigVMIfExprAST>();
-				if (IfExpr->IsAlwaysTrue())
-				{
-					ExprReplacement = (FRigVMExprAST*)IfExpr->GetTrueExpr();
-				}
-				else if (IfExpr->IsAlwaysFalse())
-				{
-					ExprReplacement = (FRigVMExprAST*)IfExpr->GetFalseExpr();
-				}
-			}
-			else if (NodeExpr->IsA(FRigVMExprAST::EType::Select))
-			{
-				const FRigVMSelectExprAST* SelectExpr = NodeExpr->To<FRigVMSelectExprAST>();
-				int32 ConstantCaseIndex = SelectExpr->GetConstantValueIndex();
-				if (ConstantCaseIndex != INDEX_NONE)
-				{
-					ExprReplacement = (FRigVMExprAST*)SelectExpr->GetValueExpr(ConstantCaseIndex);
-				}
-			}
-
-			if (ExprReplacement)
-			{
-				ExprReplacement->RemoveParent((FRigVMExprAST*)ExprReplacement->GetParent());
-				CachedValueExpr->ReplaceBy(ExprReplacement);
-				ExpressionsToRemove.Add(CachedValueExpr);
-			}
-		}
-	}
-
-	RemoveExpressions(ExpressionsToRemove);
-	return ExpressionsToRemove.Num() > 0;
-}
-
 void FRigVMParserAST::FoldLiterals()
 {
 	DECLARE_SCOPE_HIERARCHICAL_COUNTER_FUNC()
diff --git a/Engine/Source/Developer/RigVMDeveloper/Private/RigVMCompiler/RigVMCompiler.cpp b/Engine/Source/Developer/RigVMDeveloper/Private/RigVMCompiler/RigVMCompiler.cpp
--- a/Engine/Source/Developer/RigVMDeveloper/Private/RigVMCompiler/RigVMCompiler.cpp
+++ b/Engine/Source/Developer/RigVMDeveloper/Private/RigVMCompiler/RigVMCompiler.cpp
@@ -3,6 +3,7 @@
 #include "RigVMCompiler/RigVMCompiler.h"
 #include "RigVMModel/RigVMController.h"
 #include "RigVMModel/Nodes/RigVMDispatchNode.h"
+#include "RigVMModel/Nodes/RigVMBranchNode.h"
 #include "RigVMCore/RigVMExecuteContext.h"
 #include "RigVMCore/RigVMNativized.h"
 #include "RigVMDeveloperModule.h"
@@ -428,7 +429,14 @@ bool URigVMCompiler::Compile(TArray<URigVMGraph*> InGraphs, URigVMController* In
 					return false;
 				}
 			}
-			
+
+			if(ModelNode->IsA<UDEPRECATED_RigVMBranchNode>())
+			{
+				static const FString LinkedMessage = TEXT("Node @@ is a deprecated branch node. Cannot compile.");
+				Settings.ASTSettings.Report(EMessageSeverity::Error, ModelNode, LinkedMessage);
+				bEncounteredGraphError = true;
+			}
+
 			if(!InController->RemoveUnusedOrphanedPins(ModelNode, true))
 			{
 				static const FString LinkedMessage = TEXT("Node @@ uses pins that no longer exist. Please rewire the links and re-compile.");
@@ -1068,11 +1076,6 @@ void URigVMCompiler::TraverseExpression(const FRigVMExprAST* InExpr, FRigVMCompi
 			TraverseExit(InExpr->To<FRigVMExitExprAST>(), WorkData);
 			break;
 		}
-		case FRigVMExprAST::EType::Branch:
-		{
-			TraverseBranch(InExpr->To<FRigVMBranchExprAST>(), WorkData);
-			break;
-		}
 		case FRigVMExprAST::EType::If:
 		{
 			TraverseIf(InExpr->To<FRigVMIfExprAST>(), WorkData);
@@ -2099,64 +2102,6 @@ void URigVMCompiler::TraverseExit(const FRigVMExitExprAST* InExpr, FRigVMCompile
 	}
 }
 
-void URigVMCompiler::TraverseBranch(const FRigVMBranchExprAST* InExpr, FRigVMCompilerWorkData& WorkData)
-{
-	ensure(InExpr->NumChildren() == 4);
-
-	if (WorkData.bSetupMemory)
-	{
-		TraverseChildren(InExpr, WorkData);
-		return;
-	}
-
-	URigVMBranchNode* BranchNode = Cast<URigVMBranchNode>(InExpr->GetNode());
-	const FRigVMCallstack Callstack = InExpr->GetProxy().GetCallstack();
-
-	const FRigVMVarExprAST* ExecuteContextExpr = InExpr->ChildAt<FRigVMVarExprAST>(0);
-	const FRigVMVarExprAST* ConditionExpr = InExpr->ChildAt<FRigVMVarExprAST>(1);
-	const FRigVMVarExprAST* TrueExpr = InExpr->ChildAt<FRigVMVarExprAST>(2);
-	const FRigVMVarExprAST* FalseExpr = InExpr->ChildAt<FRigVMVarExprAST>(3);
-
-	// traverse the condition first
-	TraverseExpression(ConditionExpr, WorkData);
-
-	if (ConditionExpr->IsA(FRigVMExprAST::CachedValue))
-	{
-		ConditionExpr = ConditionExpr->To<FRigVMCachedValueExprAST>()->GetVarExpr();
-	}
-
-	FRigVMOperand& ConditionOperand = WorkData.ExprToOperand.FindChecked(GetSourceVarExpr(ConditionExpr));
-
-	// setup the first jump
-	uint64 JumpToFalseByte = WorkData.VM->GetByteCode().AddJumpIfOp(ERigVMOpCode::JumpForwardIf, 1, ConditionOperand, false);
-	int32 JumpToFalseInstruction = WorkData.VM->GetByteCode().GetNumInstructions() - 1;
-	if (Settings.SetupNodeInstructionIndex)
-	{
-		WorkData.VM->GetByteCode().SetSubject(WorkData.VM->GetByteCode().GetNumInstructions() - 1, Callstack.GetCallPath(), Callstack.GetStack());
-	}
-
-	// traverse the true case
-	TraverseExpression(TrueExpr, WorkData);
-
-	uint64 JumpToEndByte = WorkData.VM->GetByteCode().AddJumpOp(ERigVMOpCode::JumpForward, 1);
-	int32 JumpToEndInstruction = WorkData.VM->GetByteCode().GetNumInstructions() - 1;
-	if (Settings.SetupNodeInstructionIndex)
-	{
-		WorkData.VM->GetByteCode().SetSubject(WorkData.VM->GetByteCode().GetNumInstructions() - 1, Callstack.GetCallPath(), Callstack.GetStack());
-	}
-
-	// correct the jump to false instruction index
-	int32 NumInstructionsInTrueCase = WorkData.VM->GetByteCode().GetNumInstructions() - JumpToFalseInstruction;
-	WorkData.VM->GetByteCode().GetOpAt<FRigVMJumpIfOp>(JumpToFalseByte).InstructionIndex = NumInstructionsInTrueCase;
-
-	// traverse the false case
-	TraverseExpression(FalseExpr, WorkData);
-
-	// correct the jump to end instruction index
-	int32 NumInstructionsInFalseCase = WorkData.VM->GetByteCode().GetNumInstructions() - JumpToEndInstruction;
-	WorkData.VM->GetByteCode().GetOpAt<FRigVMJumpOp>(JumpToEndByte).InstructionIndex = NumInstructionsInFalseCase;
-}
-
 void URigVMCompiler::TraverseIf(const FRigVMIfExprAST* InExpr, FRigVMCompilerWorkData& WorkData)
 {
 	ensure(InExpr->NumChildren() == 4);
diff --git a/Engine/Source/Developer/RigVMDeveloper/Private/RigVMModel/Nodes/RigVMBranchNode.cpp b/Engine/Source/Developer/RigVMDeveloper/Private/RigVMModel/Nodes/RigVMBranchNode.cpp
--- a/Engine/Source/Developer/RigVMDeveloper/Private/RigVMModel/Nodes/RigVMBranchNode.cpp
+++ b/Engine/Source/Developer/RigVMDeveloper/Private/RigVMModel/Nodes/RigVMBranchNode.cpp
@@ -4,8 +4,3 @@
 
 #include UE_INLINE_GENERATED_CPP_BY_NAME(RigVMBranchNode)
 
-const FString URigVMBranchNode::BranchName = TEXT("Branch");
-const FString URigVMBranchNode::ConditionName = TEXT("Condition");
-const FString URigVMBranchNode::TrueName = TEXT("True");
-const FString URigVMBranchNode::FalseName = TEXT("False");
-
diff --git a/Engine/Source/Developer/RigVMDeveloper/Private/RigVMModel/RigVMClient.cpp b/Engine/Source/Developer/RigVMDeveloper/Private/RigVMModel/RigVMClient.cpp
--- a/Engine/Source/Developer/RigVMDeveloper/Private/RigVMModel/RigVMClient.cpp
+++ b/Engine/Source/Developer/RigVMDeveloper/Private/RigVMModel/RigVMClient.cpp
@@ -664,6 +664,7 @@ void FRigVMClient::PatchModelsOnLoad()
 	{
 		URigVMController* Controller = GetOrCreateController(Model);
 		Controller->PatchDispatchNodesOnLoad();
+		Controller->PatchBranchNodesOnLoad();
 	}
 }
 
diff --git a/Engine/Source/Developer/RigVMDeveloper/Private/RigVMModel/RigVMController.cpp b/Engine/Source/Developer/RigVMDeveloper/Private/RigVMModel/RigVMController.cpp
--- a/Engine/Source/Developer/RigVMDeveloper/Private/RigVMModel/RigVMController.cpp
+++ b/Engine/Source/Developer/RigVMDeveloper/Private/RigVMModel/RigVMController.cpp
@@ -11,6 +11,7 @@
 #include "RigVMCore/RigVMExecuteContext.h"
 #include "RigVMCore/RigVMUnknownType.h"
 #include "RigVMCore/RigVMByteCode.h"
+#include "RigVMFunctions/RigVMFunction_ControlFlow.h"
 #include "RigVMCompiler/RigVMCompiler.h"
 #include "RigVMDeveloperModule.h"
 #include "UObject/PropertyPortFlags.h"
@@ -24,6 +25,7 @@
 #include "RigVMTypeUtils.h"
 #include "Engine/UserDefinedStruct.h"
 #include "RigVMModel/RigVMClient.h"
+#include "RigVMModel/Nodes/RigVMBranchNode.h"
 
 #include UE_INLINE_GENERATED_CPP_BY_NAME(RigVMController)
 
@@ -598,14 +600,6 @@ TArray<FString> URigVMController::GetAddNodePythonCommands(URigVMNode* Node) con
 					*RigVMPythonUtils::LinearColorToPythonString(CommentNode->GetNodeColor()),
 					*NodeName));	
 	}
-	else if (const URigVMBranchNode* BranchNode = Cast<URigVMBranchNode>(Node))
-	{
-		// add_branch_node(position=[0.0, 0.0], node_name='', undo=True)
-		Commands.Add(FString::Printf(TEXT("blueprint.get_controller_by_name('%s').add_branch_node(%s, '%s')"),
-						*GraphName,
-						*RigVMPythonUtils::Vector2DToPythonString(BranchNode->GetPosition()),
-						*NodeName));	
-	}
 	else if (const URigVMIfNode* IfNode = Cast<URigVMIfNode>(Node))
 	{
 		// add_if_node(cpp_type, cpp_type_object_path, position=[0.0, 0.0], node_name='', undo=True)
@@ -12730,69 +12724,9 @@ URigVMRerouteNode* URigVMController::AddFreeRerouteNode(bool bShowAsFullNode, co
 	return Node;
 }
 
-URigVMBranchNode* URigVMController::AddBranchNode(const FVector2D& InPosition, const FString& InNodeName, bool bSetupUndoRedo, bool bPrintPythonCommand)
+URigVMNode* URigVMController::AddBranchNode(const FVector2D& InPosition, const FString& InNodeName, bool bSetupUndoRedo, bool bPrintPythonCommand)
 {
-	if(!IsValidGraph())
-	{
-		return nullptr;
-	}
-
-	if (!bIsTransacting && !IsGraphEditable())
-	{
-		return nullptr;
-	}
-
-	URigVMGraph* Graph = GetGraph();
-	check(Graph);
-
-	FString Name = GetValidNodeName(InNodeName.IsEmpty() ? FString(TEXT("BranchNode")) : InNodeName);
-	URigVMBranchNode* Node = NewObject<URigVMBranchNode>(Graph, *Name);
-	Node->Position = InPosition;
-
-	URigVMPin* ExecutePin = MakeExecutePin(Node, FRigVMStruct::ExecuteContextName);
-	ExecutePin->Direction = ERigVMPinDirection::Input;
-	AddNodePin(Node, ExecutePin);
-
-	URigVMPin* ConditionPin = NewObject<URigVMPin>(Node, *URigVMBranchNode::ConditionName);
-	ConditionPin->CPPType = RigVMTypeUtils::BoolType;
-	ConditionPin->Direction = ERigVMPinDirection::Input;
-	AddNodePin(Node, ConditionPin);
-
-	URigVMPin* TruePin = NewObject<URigVMPin>(Node, *URigVMBranchNode::TrueName);
-	TruePin->CPPType = ExecutePin->CPPType;
-	TruePin->CPPTypeObject = ExecutePin->CPPTypeObject;
-	TruePin->CPPTypeObjectPath = ExecutePin->CPPTypeObjectPath;
-	TruePin->Direction = ERigVMPinDirection::Output;
-	AddNodePin(Node, TruePin);
-
-	URigVMPin* FalsePin = NewObject<URigVMPin>(Node, *URigVMBranchNode::FalseName);
-	FalsePin->CPPType = ExecutePin->CPPType;
-	FalsePin->CPPTypeObject = ExecutePin->CPPTypeObject;
-	FalsePin->CPPTypeObjectPath = ExecutePin->CPPTypeObjectPath;
-	FalsePin->Direction = ERigVMPinDirection::Output;
-	AddNodePin(Node, FalsePin);
-
-	Graph->Nodes.Add(Node);
-
-	Notify(ERigVMGraphNotifType::NodeAdded, Node);
-
-	FRigVMControllerCompileBracketScope CompileScope(this);
-	if (bSetupUndoRedo)
-	{
-		ActionStack->AddAction(FRigVMAddBranchNodeAction(Node));
-	}
-
-	if (bPrintPythonCommand)
-	{
-		TArray<FString> Commands = GetAddNodePythonCommands(Node);
-		for (const FString& Command : Commands)
-		{
-			RigVMPythonUtils::Print(GetGraphOuterName(), 
-				FString::Printf(TEXT("%s"), *Command));
-		}
-	}
-
-	return Node;
+	return AddUnitNode(FRigVMFunction_ControlFlowBranch::StaticStruct(), FRigVMStruct::ExecuteName, InPosition, InNodeName, bSetupUndoRedo, bPrintPythonCommand);
 }
 
 URigVMIfNode* URigVMController::AddIfNode(const FString& InCPPType, const FName& InCPPTypeObjectPath, const FVector2D& InPosition, const FString& InNodeName, bool  bSetupUndoRedo, bool bPrintPythonCommand)
@@ -19705,6 +19639,37 @@ void URigVMController::PatchDispatchNodesOnLoad()
 	}
 }
 
+void URigVMController::PatchBranchNodesOnLoad()
+{
+	if (const URigVMGraph* Graph = GetGraph())
+	{
+		TArray<URigVMNode*> BranchNodes = Graph->GetNodes().FilterByPredicate([](URigVMNode* Node)
+		{
+			return Node->IsA<UDEPRECATED_RigVMBranchNode>();
+		});
+
+		for(URigVMNode* BranchNode : BranchNodes)
+		{
+			TArray<TPair<FString, FString>> LinkedPaths = GetLinkedPinPaths(BranchNode);
+			const FVector2D NodePosition = BranchNode->GetPosition();
+			const FString NodeName = BranchNode->GetName();
+			const URigVMPin* OldConditionPin = BranchNode->FindPin(GET_MEMBER_NAME_CHECKED(FRigVMFunction_ControlFlowBranch, Condition).ToString());
+			const FString ConditionDefault = GetPinDefaultValue(OldConditionPin->GetPinPath());
+
+			RemoveNode(BranchNode, false, true, false, false);
+
+			const URigVMNode* NewNode = AddUnitNode(FRigVMFunction_ControlFlowBranch::StaticStruct(), FRigVMStruct::ExecuteName, NodePosition, NodeName, false, false);
+
+			if(!ConditionDefault.IsEmpty())
+			{
+				const URigVMPin* ConditionPin = NewNode->FindPin(GET_MEMBER_NAME_CHECKED(FRigVMFunction_ControlFlowBranch, Condition).ToString());
+				SetPinDefaultValue(ConditionPin->GetPinPath(), ConditionDefault, false, false, false, false);
+			}
+			RestoreLinkedPaths(LinkedPaths, {}, {}, false);
+		}
+	}
+}
+
 #endif
 
 URigVMControllerSettings::URigVMControllerSettings(const FObjectInitializer& Initializer)
diff --git a/Engine/Source/Developer/RigVMDeveloper/Private/RigVMModel/RigVMControllerActions.cpp b/Engine/Source/Developer/RigVMDeveloper/Private/RigVMModel/RigVMControllerActions.cpp
--- a/Engine/Source/Developer/RigVMDeveloper/Private/RigVMModel/RigVMControllerActions.cpp
+++ b/Engine/Source/Developer/RigVMDeveloper/Private/RigVMModel/RigVMControllerActions.cpp
@@ -487,38 +487,6 @@ bool FRigVMAddRerouteNodeAction::Redo(URigVMController* InController)
 	return false;
 }
 
-FRigVMAddBranchNodeAction::FRigVMAddBranchNodeAction()
-	: Position(FVector2D::ZeroVector)
-	, NodePath()
-{
-}
-
-FRigVMAddBranchNodeAction::FRigVMAddBranchNodeAction(URigVMBranchNode* InNode)
-	: Position(InNode->GetPosition())
-	, NodePath(InNode->GetNodePath())
-{
-}
-
-bool FRigVMAddBranchNodeAction::Undo(URigVMController* InController)
-{
-	if (!FRigVMBaseAction::Undo(InController))
-	{
-		return false;
-	}
-	return InController->RemoveNodeByName(*NodePath, false);
-}
-
-bool FRigVMAddBranchNodeAction::Redo(URigVMController* InController)
-{
-#if WITH_EDITOR
-	if (URigVMBranchNode* Node = InController->AddBranchNode(Position, NodePath, false))
-	{
-		return FRigVMBaseAction::Redo(InController);
-	}
-#endif
-	return false;
-}
-
 FRigVMAddIfNodeAction::FRigVMAddIfNodeAction()
 	: CPPType()
 	, CPPTypeObjectPath(NAME_None)
@@ -738,10 +706,6 @@ FRigVMRemoveNodeAction::FRigVMRemoveNodeAction(URigVMNode* InNode, URigVMControl
 	{
 		InverseAction.AddAction(FRigVMAddRerouteNodeAction(RerouteNode), InController);
 	}
-	else if (URigVMBranchNode* BranchNode = Cast<URigVMBranchNode>(InNode))
-	{
-		InverseAction.AddAction(FRigVMAddBranchNodeAction(BranchNode), InController);
-	}
 	else if (URigVMIfNode* IfNode = Cast<URigVMIfNode>(InNode))
 	{
 		InverseAction.AddAction(FRigVMAddIfNodeAction(IfNode), InController);
diff --git a/Engine/Source/Developer/RigVMDeveloper/Public/RigVMCompiler/RigVMAST.h b/Engine/Source/Developer/RigVMDeveloper/Public/RigVMCompiler/RigVMAST.h
--- a/Engine/Source/Developer/RigVMDeveloper/Public/RigVMCompiler/RigVMAST.h
+++ b/Engine/Source/Developer/RigVMDeveloper/Public/RigVMCompiler/RigVMAST.h
@@ -26,7 +26,6 @@ class FRigVMAssignExprAST;
 class FRigVMCopyExprAST;
 class FRigVMCachedValueExprAST;
 class FRigVMExitExprAST;
-class FRigVMBranchExprAST;
 class FRigVMIfExprAST;
 class FRigVMSelectExprAST;
 class FRigVMArrayExprAST;
@@ -462,17 +461,6 @@ inline const FRigVMExitExprAST* FRigVMExprAST::To() const
 	return (const FRigVMExitExprAST*)this;
 }
 
-// specialized cast for type checking
-// for a Branch / FRigVMBranchExprAST expression
-// will raise if types are not compatible
-// @return this expression cast to FRigVMBranchExprAST 
-template<>
-inline const FRigVMBranchExprAST* FRigVMExprAST::To() const
-{
-	ensure(IsA(EType::Branch));
-	return (const FRigVMBranchExprAST*)this;
-}
-
 // specialized cast for type checking
 // for a If / FRigVMIfExprAST expression
 // will raise if types are not compatible
@@ -1118,49 +1106,6 @@ private:
 	friend class FRigVMParserAST;
 };
 
-/*
- * An abstract syntax tree branch expression represents a branch point
- * for two blocks.
- * In C++ the branch is is the definition: if(bCondition){a();}else{b();}
- */
-class RIGVMDEVELOPER_API FRigVMBranchExprAST : public FRigVMNodeExprAST
-{
-public:
-
-	// virtual destructor
-	virtual ~FRigVMBranchExprAST() {}
-
-	// disable copy constructor
-	FRigVMBranchExprAST(const FRigVMEntryExprAST&) = delete;
-
-	// overload of the type checking mechanism
-	virtual bool IsA(EType InType) const override
-	{
-		return InType == EType::Branch;
-	};
-
-	virtual bool IsConstant() const override;
-
-	const FRigVMVarExprAST* GetConditionExpr() const { return ChildAt(1)->To<FRigVMVarExprAST>(); }
-	const FRigVMVarExprAST* GetTrueExpr() const { return ChildAt(2)->To<FRigVMVarExprAST>(); }
-	const FRigVMVarExprAST* GetFalseExpr() const { return ChildAt(3)->To<FRigVMVarExprAST>(); }
-
-	bool IsAlwaysTrue() const;
-	bool IsAlwaysFalse() const;
-
-protected:
-
-	// default constructor (protected so that only parser can access it)
-	FRigVMBranchExprAST(const FRigVMASTProxy& InProxy)
-		: FRigVMNodeExprAST(EType::Branch, InProxy)
-	{}
-
-
-private:
-
-	friend class FRigVMParserAST;
-};
-
 /*
  * An abstract syntax tree if expression represents a branch point for values.
  * In C++ the branch is is the definition: (Condition ? True : False)
@@ -1535,9 +1480,6 @@ private:
 	// helper function to fold constant values into literals
 	bool FoldConstantValuesToLiterals(TArray<URigVMGraph*> InGraphs, URigVMController* InController, const TArray<FRigVMExternalVariable>& InExternalVariables, const TArray<FRigVMUserDataArray>& InRigVMUserData);
 
-	// helper function to fold unreachable branches 
-	bool FoldUnreachableBranches(TArray<URigVMGraph*> InGraphs);
-
 	// helper function to inline all contributing nodes of the graph
 	void Inline(TArray<URigVMGraph*> InGraphs);
 	void Inline(TArray<URigVMGraph*> InGraphs, const TArray<FRigVMASTProxy>& InNodeProxies);
diff --git a/Engine/Source/Developer/RigVMDeveloper/Public/RigVMCompiler/RigVMCompiler.h b/Engine/Source/Developer/RigVMDeveloper/Public/RigVMCompiler/RigVMCompiler.h
--- a/Engine/Source/Developer/RigVMDeveloper/Public/RigVMCompiler/RigVMCompiler.h
+++ b/Engine/Source/Developer/RigVMDeveloper/Public/RigVMCompiler/RigVMCompiler.h
@@ -211,7 +211,6 @@ private:
 	void TraverseCopy(const FRigVMCopyExprAST* InExpr, FRigVMCompilerWorkData& WorkData);
 	void TraverseCachedValue(const FRigVMCachedValueExprAST* InExpr, FRigVMCompilerWorkData& WorkData);
 	void TraverseExit(const FRigVMExitExprAST* InExpr, FRigVMCompilerWorkData& WorkData);
-	void TraverseBranch(const FRigVMBranchExprAST* InExpr, FRigVMCompilerWorkData& WorkData);
 	void TraverseIf(const FRigVMIfExprAST* InExpr, FRigVMCompilerWorkData& WorkData);
 	void TraverseSelect(const FRigVMSelectExprAST* InExpr, FRigVMCompilerWorkData& WorkData);
 	void TraverseArray(const FRigVMArrayExprAST* InExpr, FRigVMCompilerWorkData& WorkData);
diff --git a/Engine/Source/Developer/RigVMDeveloper/Public/RigVMModel/Nodes/RigVMBranchNode.h b/Engine/Source/Developer/RigVMDeveloper/Public/RigVMModel/Nodes/RigVMBranchNode.h
--- a/Engine/Source/Developer/RigVMDeveloper/Public/RigVMModel/Nodes/RigVMBranchNode.h
+++ b/Engine/Source/Developer/RigVMDeveloper/Public/RigVMModel/Nodes/RigVMBranchNode.h
@@ -2,12 +2,7 @@
 
 #pragma once
 
-#include "Containers/UnrealString.h"
-#include "Math/Color.h"
 #include "RigVMModel/RigVMNode.h"
-#include "UObject/ObjectMacros.h"
-#include "UObject/UObjectGlobals.h"
-
 #include "RigVMBranchNode.generated.h"
 
 class UObject;
@@ -15,26 +10,9 @@ class UObject;
 /**
  * A branch node is used to branch between two blocks of execution
  */
-UCLASS(BlueprintType)
-class RIGVMDEVELOPER_API URigVMBranchNode : public URigVMNode
+UCLASS(BlueprintType, Deprecated)
+class RIGVMDEVELOPER_API UDEPRECATED_RigVMBranchNode : public URigVMNode
 {
 	GENERATED_BODY()
-
-public:
-
-	// Override from URigVMNode
-	virtual FString GetNodeTitle() const override { return BranchName; }
-	virtual FLinearColor GetNodeColor() const override { return FLinearColor::Black; }
-
-private:
-
-	static const FString BranchName;
-	static const FString ConditionName;
-	static const FString TrueName;
-	static const FString FalseName;
-
-	friend class URigVMController;
-	friend class URigVMCompiler;
-	friend class UControlRigBranchNodeSpawner;
 };
 
diff --git a/Engine/Source/Developer/RigVMDeveloper/Public/RigVMModel/RigVMController.h b/Engine/Source/Developer/RigVMDeveloper/Public/RigVMModel/RigVMController.h
--- a/Engine/Source/Developer/RigVMDeveloper/Public/RigVMModel/RigVMController.h
+++ b/Engine/Source/Developer/RigVMDeveloper/Public/RigVMModel/RigVMController.h
@@ -9,7 +9,6 @@
 #include "RigVMModel/Nodes/RigVMParameterNode.h"
 #include "RigVMModel/Nodes/RigVMCommentNode.h"
 #include "RigVMModel/Nodes/RigVMRerouteNode.h"
-#include "RigVMModel/Nodes/RigVMBranchNode.h"
 #include "RigVMModel/Nodes/RigVMIfNode.h"
 #include "RigVMModel/Nodes/RigVMSelectNode.h"
 #include "RigVMModel/Nodes/RigVMTemplateNode.h"
@@ -339,8 +338,9 @@ public:
 	// Adds a branch node to the graph.
 	// Branch nodes can be used to split the execution of into multiple branches,
 	// allowing to drive behavior by logic.
-	UFUNCTION(BlueprintCallable, Category = RigVMController)
-	URigVMBranchNode* AddBranchNode(const FVector2D& InPosition = FVector2D::ZeroVector, const FString& InNodeName = TEXT(""), bool bSetupUndoRedo = true, bool bPrintPythonCommand = false);
+	UE_DEPRECATED(5.2, "Please use AddUnitNode with FRigVMFunction_ControlFlowBranch::StaticStruct()")
+	UFUNCTION(BlueprintCallable, Category = RigVMController, meta=(DeprecatedFunction))
+	URigVMNode* AddBranchNode(const FVector2D& InPosition = FVector2D::ZeroVector, const FString& InNodeName = TEXT(""), bool bSetupUndoRedo = true, bool bPrintPythonCommand = false);
 
 	// Adds an if node to the graph.
 	// If nodes can be used to pick between two values based on a condition.
@@ -1175,6 +1175,7 @@ protected:
 
 	// backwards compatibility code
 	void PatchDispatchNodesOnLoad();
+	void PatchBranchNodesOnLoad();
 	
 private: 
 	UPROPERTY(transient)
diff --git a/Engine/Source/Developer/RigVMDeveloper/Public/RigVMModel/RigVMControllerActions.h b/Engine/Source/Developer/RigVMDeveloper/Public/RigVMModel/RigVMControllerActions.h
--- a/Engine/Source/Developer/RigVMDeveloper/Public/RigVMModel/RigVMControllerActions.h
+++ b/Engine/Source/Developer/RigVMDeveloper/Public/RigVMModel/RigVMControllerActions.h
@@ -523,30 +523,6 @@ public:
 	FString NodePath;
 };
 
-/**
- * An action adding a branch node to the graph.
- */
-USTRUCT()
-struct FRigVMAddBranchNodeAction : public FRigVMBaseAction
-{
-	GENERATED_BODY()
-
-public:
-
-	FRigVMAddBranchNodeAction();
-	FRigVMAddBranchNodeAction(URigVMBranchNode* InNode);
-	virtual ~FRigVMAddBranchNodeAction() {};
-	virtual UScriptStruct* GetScriptStruct() const override { return FRigVMAddBranchNodeAction::StaticStruct(); }
-	virtual bool Undo(URigVMController* InController) override;
-	virtual bool Redo(URigVMController* InController) override;
-
-	UPROPERTY()
-	FVector2D Position;
-
-	UPROPERTY()
-	FString NodePath;
-};
-
 /**
  * An action adding an if node to the graph.
  */
diff --git a/Engine/Source/Runtime/RigVM/Private/RigVMFunctions/RigVMFunction_ControlFlow.cpp b/Engine/Source/Runtime/RigVM/Private/RigVMFunctions/RigVMFunction_ControlFlow.cpp
new file mode 100644
--- /dev/null
+++ b/Engine/Source/Runtime/RigVM/Private/RigVMFunctions/RigVMFunction_ControlFlow.cpp
@@ -0,0 +1,26 @@
+// Copyright Epic Games, Inc. All Rights Reserved.
+
+#include "RigVMFunctions/RigVMFunction_ControlFlow.h"
+
+#include UE_INLINE_GENERATED_CPP_BY_NAME(RigVMFunction_ControlFlow)
+
+FRigVMFunction_ControlFlowBranch_Execute()
+{
+	if(BlockToRun.IsNone())
+	{
+		if(Condition)
+		{
+			static const FName TrueName = GET_MEMBER_NAME_CHECKED(FRigVMFunction_ControlFlowBranch, True); 
+			BlockToRun = TrueName;
+		}
+		else
+		{
+			static const FName FalseName = GET_MEMBER_NAME_CHECKED(FRigVMFunction_ControlFlowBranch, False); 
+			BlockToRun = FalseName;
+		}
+	}
+	else
+	{
+		BlockToRun = FRigVMStruct::ControlFlowCompletedName;
+	}
+}
diff --git a/Engine/Source/Runtime/RigVM/Public/RigVMFunctions/RigVMFunction_ControlFlow.h b/Engine/Source/Runtime/RigVM/Public/RigVMFunctions/RigVMFunction_ControlFlow.h
new file mode 100644
--- /dev/null
+++ b/Engine/Source/Runtime/RigVM/Public/RigVMFunctions/RigVMFunction_ControlFlow.h
@@ -0,0 +1,58 @@
+// Copyright Epic Games, Inc. All Rights Reserved.
+
+#pragma once
+
+#include "RigVMCore/RigVMStruct.h"
+#include "RigVMFunction_ControlFlow.generated.h"
+
+USTRUCT(meta=(Abstract, Category="Execution", NodeColor = "0, 0, 0, 1"))
+struct RIGVM_API FRigVMFunction_ControlFlowBase : public FRigVMStruct
+{
+	GENERATED_BODY()
+};
+
+/**
+ * Executes either the True or False branch based on the condition
+ */
+USTRUCT(meta = (DisplayName = "Branch", Keywords = "If"))
+struct RIGVM_API FRigVMFunction_ControlFlowBranch : public FRigVMFunction_ControlFlowBase
+{
+	GENERATED_BODY()
+
+	FRigVMFunction_ControlFlowBranch()
+	{
+		Condition = false;
+		BlockToRun = NAME_None;
+	}
+
+	RIGVM_METHOD()
+	virtual void Execute();
+
+	virtual const TArray<FName>& GetControlFlowBlocks_Impl() const override
+	{
+		static const TArray<FName> Blocks = {
+			GET_MEMBER_NAME_CHECKED(FRigVMFunction_ControlFlowBranch, True),
+			GET_MEMBER_NAME_CHECKED(FRigVMFunction_ControlFlowBranch, False),
+			ForLoopCompletedPinName
+		};
+		return Blocks;
+	}
+
+	UPROPERTY(meta=(Input, DisplayName="Execute"))
+	FRigVMExecuteContext ExecuteContext;
+
+	UPROPERTY(meta=(Input))
+	bool Condition;
+
+	UPROPERTY(meta=(Output))
+	FRigVMExecuteContext True;
+
+	UPROPERTY(meta=(Output))
+	FRigVMExecuteContext False;
+
+	UPROPERTY(meta=(Output))
+	FRigVMExecuteContext Completed;
+
+	UPROPERTY(meta=(Singleton))
+	FName BlockToRun;
+};
