ControlRig: Initialize BuildData outside ControlRigEditor

diff --git a/Engine/Plugins/Animation/ControlRig/Source/ControlRigDeveloper/Private/ControlRigBlueprint.cpp b/Engine/Plugins/Animation/ControlRig/Source/ControlRigDeveloper/Private/ControlRigBlueprint.cpp
--- a/Engine/Plugins/Animation/ControlRig/Source/ControlRigDeveloper/Private/ControlRigBlueprint.cpp
+++ b/Engine/Plugins/Animation/ControlRig/Source/ControlRigDeveloper/Private/ControlRigBlueprint.cpp
@@ -922,7 +922,7 @@ void UControlRigBlueprint::PostLoad()
 				{
 					if(URigVMFunctionReferenceNode* FunctionReferenceNode = Cast<URigVMFunctionReferenceNode>(Node))
 					{
-						if(URigVMBuildData* BuildData = URigVMController::GetBuildData())
+						if(URigVMBuildData* BuildData = URigVMBuildData::Get())
 						{
 							BuildData->RegisterFunctionReference(FunctionReferenceNode);
 						}
@@ -1053,7 +1053,7 @@ void UControlRigBlueprint::HandlePackageDone()
 		ShapeLibrariesToLoadOnPackageLoaded.Reset();
 	}
 
-	if(URigVMBuildData* BuildData = URigVMController::GetBuildData())
+	if(URigVMBuildData* BuildData = URigVMBuildData::Get())
 	{
 		if(URigVMFunctionLibrary* FunctionLibrary = RigVMClient.GetFunctionLibrary())
 		{
diff --git a/Engine/Plugins/Animation/ControlRig/Source/ControlRigDeveloper/Public/ControlRigBlueprint.h b/Engine/Plugins/Animation/ControlRig/Source/ControlRigDeveloper/Public/ControlRigBlueprint.h
--- a/Engine/Plugins/Animation/ControlRig/Source/ControlRigDeveloper/Public/ControlRigBlueprint.h
+++ b/Engine/Plugins/Animation/ControlRig/Source/ControlRigDeveloper/Public/ControlRigBlueprint.h
@@ -417,10 +417,6 @@ protected:
 	UPROPERTY(AssetRegistrySearchable)
 	TArray<FRigVMGraphFunctionHeader> PublicGraphFunctions;
 
-	/** Asset searchable information function references in this rig */
-	UPROPERTY(AssetRegistrySearchable)
-	TArray<FRigVMReferenceNodeData> FunctionReferenceNodeData;
-
 #if WITH_EDITORONLY_DATA
 
 	UPROPERTY(transient, DuplicateTransient)
diff --git a/Engine/Plugins/Animation/ControlRig/Source/ControlRigEditor/Private/Editor/ControlRigEditor.cpp b/Engine/Plugins/Animation/ControlRig/Source/ControlRigEditor/Private/Editor/ControlRigEditor.cpp
--- a/Engine/Plugins/Animation/ControlRig/Source/ControlRigEditor/Private/Editor/ControlRigEditor.cpp
+++ b/Engine/Plugins/Animation/ControlRig/Source/ControlRigEditor/Private/Editor/ControlRigEditor.cpp
@@ -111,8 +111,6 @@ const FName ControlRigEditorAppName(TEXT("ControlRigEditorApp"));
 
 const FName FControlRigEditorModes::ControlRigEditorMode("Rigging");
 
-bool FControlRigEditor::bAreFunctionReferencesInitialized = false;
-
 namespace ControlRigEditorTabs
 {
 	const FName DetailsTab(TEXT("DetailsTab"));
@@ -373,9 +371,6 @@ void FControlRigEditor::InitControlRigEditor(const EToolkitMode::Type Mode, cons
 	InControlRigBlueprint->InitializeModelIfRequired();
 
 	CommonInitialization(ControlRigBlueprints, false);
-
-	// update function references once
-	InitFunctionReferences();
 	
 	// user-defined-struct can change even after load
 	// refresh the models such that pins are updated to match
@@ -610,61 +605,6 @@ void FControlRigEditor::InitControlRigEditor(const EToolkitMode::Type Mode, cons
 	PropertyChangedHandle = FCoreUObjectDelegates::OnObjectPropertyChanged.AddSP(this, &FControlRigEditor::OnPropertyChanged);
 }
 
-void FControlRigEditor::InitFunctionReferences()
-{
-	if (bAreFunctionReferencesInitialized)
-	{
-		return;
-	}
-	
-	FFunctionGraphTask::CreateAndDispatchWhenReady([]()
-	{
-		if(URigVMBuildData* BuildData = URigVMController::GetBuildData())
-		{
-			const FArrayProperty* ReferenceNodeDataProperty =
-				CastField<FArrayProperty>(UControlRigBlueprint::StaticClass()->FindPropertyByName(TEXT("FunctionReferenceNodeData")));
-			if(ReferenceNodeDataProperty)
-			{
-				const FAssetRegistryModule& AssetRegistryModule = FModuleManager::LoadModuleChecked<FAssetRegistryModule>("AssetRegistry");
-
-				// find all control rigs in the project
-				TArray<FAssetData> ControlRigAssetDatas;
-				FARFilter ControlRigAssetFilter;
-				ControlRigAssetFilter.ClassPaths.Add(UControlRigBlueprint::StaticClass()->GetClassPathName());
-				AssetRegistryModule.Get().GetAssets(ControlRigAssetFilter, ControlRigAssetDatas);
-
-				// loop over all control rigs in the project
-				for(const FAssetData& ControlRigAssetData : ControlRigAssetDatas)
-				{
-					const FString ReferenceNodeDataString =
-						ControlRigAssetData.GetTagValueRef<FString>(ReferenceNodeDataProperty->GetFName());
-					if(ReferenceNodeDataString.IsEmpty())
-					{
-						continue;
-					}
-
-					TArray<FRigVMReferenceNodeData> ReferenceNodeDatas;
-					ReferenceNodeDataProperty->ImportText_Direct(*ReferenceNodeDataString, &ReferenceNodeDatas, nullptr, EPropertyPortFlags::PPF_None);
-
-					for(FRigVMReferenceNodeData& ReferenceNodeData : ReferenceNodeDatas)
-					{
-						if (ReferenceNodeData.ReferencedHeader.IsValid())
-						{
-							BuildData->RegisterFunctionReference(ReferenceNodeData.ReferencedHeader.LibraryPointer, ReferenceNodeData.GetReferenceNodeObjectPath());
-						}
-						else if (!ReferenceNodeData.ReferencedFunctionPath_DEPRECATED.IsEmpty())
-						{
-							BuildData->RegisterFunctionReference(ReferenceNodeData);							
-						}
-					}
-				}
-			}
-		}
-	}, TStatId(), NULL, ENamedThreads::GameThread);
-	
-	bAreFunctionReferencesInitialized = true;
-}
-
 void FControlRigEditor::BindCommands()
 {
 	DECLARE_SCOPE_HIERARCHICAL_COUNTER_FUNC()
diff --git a/Engine/Plugins/Animation/ControlRig/Source/ControlRigEditor/Private/Editor/ControlRigEditor.h b/Engine/Plugins/Animation/ControlRig/Source/ControlRigEditor/Private/Editor/ControlRigEditor.h
--- a/Engine/Plugins/Animation/ControlRig/Source/ControlRigEditor/Private/Editor/ControlRigEditor.h
+++ b/Engine/Plugins/Animation/ControlRig/Source/ControlRigEditor/Private/Editor/ControlRigEditor.h
@@ -84,10 +84,8 @@ private:
 
 	FControlRigEditorClosed ControlRigEditorClosedDelegate;
 
-	static bool bAreFunctionReferencesInitialized;
 
 public:
-	static void InitFunctionReferences();
 	
 	// IToolkit Interface
 	virtual FName GetToolkitFName() const override;
diff --git a/Engine/Source/Developer/RigVMDeveloper/Private/RigVMModel/RigVMBuildData.cpp b/Engine/Source/Developer/RigVMDeveloper/Private/RigVMModel/RigVMBuildData.cpp
--- a/Engine/Source/Developer/RigVMDeveloper/Private/RigVMModel/RigVMBuildData.cpp
+++ b/Engine/Source/Developer/RigVMDeveloper/Private/RigVMModel/RigVMBuildData.cpp
@@ -1,11 +1,20 @@
 // Copyright Epic Games, Inc. All Rights Reserved.
 
 #include "RigVMModel/RigVMBuildData.h"
-
+#include "AssetRegistry/AssetRegistryModule.h"
 #include "RigVMModel/RigVMFunctionLibrary.h"
+#include "UObject/UObjectIterator.h"
+#include "RigVMModel/RigVMClient.h"
+#include "UObject/StrongObjectPtr.h"
 
 #include UE_INLINE_GENERATED_CPP_BY_NAME(RigVMBuildData)
 
+// When the object system has been completely loaded, collect all the references between RigVM graphs 
+static FDelayedAutoRegisterHelper GRigVMBuildDataSingletonHelper(EDelayedRegisterRunPhase::EndOfEngineInit, []() -> void
+{
+	URigVMBuildData::Get()->InitializeIfNeeded();
+});
+
 FRigVMReferenceNodeData::FRigVMReferenceNodeData(URigVMFunctionReferenceNode* InReferenceNode)
 {
 	check(InReferenceNode);
@@ -40,12 +49,94 @@ URigVMFunctionReferenceNode* FRigVMReferenceNodeData::GetReferenceNode()
 	return nullptr;
 }
 
+bool URigVMBuildData::bInitialized = false; 
+
 URigVMBuildData::URigVMBuildData()
 : UObject()
 , bIsRunningUnitTest(false)
 {
 }
 
+URigVMBuildData* URigVMBuildData::Get()
+{
+	// static in a function scope ensures that the GC system is initiated before 
+	// the build data constructor is called
+	static TStrongObjectPtr<URigVMBuildData> sBuildData;
+	if(!sBuildData.IsValid() && IsInGameThread())
+	{
+		sBuildData = TStrongObjectPtr<URigVMBuildData>(
+			NewObject<URigVMBuildData>(
+				GetTransientPackage(), 
+				TEXT("RigVMBuildData"), 
+				RF_Transient));
+	}
+	return sBuildData.Get();
+
+}
+
+void URigVMBuildData::InitializeIfNeeded()
+{
+	if (bInitialized)
+	{
+		return;
+	}
+
+	// Find all classes which implement IRigVMClientHost
+	TArray<UClass*> ImplementedClasses;
+	for (TObjectIterator<UClass> ClassIterator; ClassIterator; ++ClassIterator)
+	{
+		if (ClassIterator->ImplementsInterface(URigVMClientHost::StaticClass()))
+		{
+			ImplementedClasses.Add(*ClassIterator);
+		}
+	}
+
+	// Loop the classes
+	for (UClass* Class : ImplementedClasses)
+	{
+		const FArrayProperty* ReferenceNodeDataProperty =
+		   CastField<FArrayProperty>(Class->FindPropertyByName(TEXT("FunctionReferenceNodeData")));
+		if(ReferenceNodeDataProperty)
+		{
+			const FAssetRegistryModule& AssetRegistryModule = FModuleManager::LoadModuleChecked<FAssetRegistryModule>("AssetRegistry");
+
+			// find all assets of this class in the project
+			TArray<FAssetData> ControlRigAssetDatas;
+			FARFilter ControlRigAssetFilter;
+			ControlRigAssetFilter.ClassPaths.Add(Class->GetClassPathName());
+			AssetRegistryModule.Get().GetAssets(ControlRigAssetFilter, ControlRigAssetDatas);
+
+			// loop over all found assets
+			for(const FAssetData& ControlRigAssetData : ControlRigAssetDatas)
+			{
+				const FString ReferenceNodeDataString =
+					ControlRigAssetData.GetTagValueRef<FString>(ReferenceNodeDataProperty->GetFName());
+				if(ReferenceNodeDataString.IsEmpty())
+				{
+					continue;
+				}
+
+				// See if it has reference node data, and register the references
+				TArray<FRigVMReferenceNodeData> ReferenceNodeDatas;
+				ReferenceNodeDataProperty->ImportText_Direct(*ReferenceNodeDataString, &ReferenceNodeDatas, nullptr, EPropertyPortFlags::PPF_None);	
+				for(FRigVMReferenceNodeData& ReferenceNodeData : ReferenceNodeDatas)
+				{
+					if (ReferenceNodeData.ReferencedHeader.IsValid())
+					{
+						RegisterFunctionReference(ReferenceNodeData.ReferencedHeader.LibraryPointer, ReferenceNodeData.GetReferenceNodeObjectPath());
+					}
+					else if (!ReferenceNodeData.ReferencedFunctionPath_DEPRECATED.IsEmpty())
+					{
+						RegisterFunctionReference(ReferenceNodeData);							
+					}
+				}
+			}
+		}
+	}
+	
+	bInitialized = true;
+}
+
 const FRigVMFunctionReferenceArray* URigVMBuildData::FindFunctionReferences(const FRigVMGraphFunctionIdentifier& InFunction) const
 {
 	return GraphFunctionReferences.Find(InFunction);
diff --git a/Engine/Source/Developer/RigVMDeveloper/Private/RigVMModel/RigVMClient.cpp b/Engine/Source/Developer/RigVMDeveloper/Private/RigVMModel/RigVMClient.cpp
--- a/Engine/Source/Developer/RigVMDeveloper/Private/RigVMModel/RigVMClient.cpp
+++ b/Engine/Source/Developer/RigVMDeveloper/Private/RigVMModel/RigVMClient.cpp
@@ -750,7 +750,7 @@ void FRigVMClient::HandleGraphModifiedEvent(ERigVMGraphNotifType InNotifType, UR
 			{
 				if (URigVMCollapseNode* CollapseNode = Cast<URigVMCollapseNode>(InSubject))
 				{
-					URigVMBuildData* BuildData = GetFunctionLibrary()->GetBuildData();
+					URigVMBuildData* BuildData = URigVMBuildData::Get();
 					IRigVMGraphFunctionHost* Host = Cast<IRigVMGraphFunctionHost>(CollapseNode->GetFunctionIdentifier().HostObject.ResolveObject());
 					FRigVMGraphFunctionData* Data = Host->GetRigVMGraphFunctionStore()->FindFunctionByName(CollapseNode->GetPreviousFName());
 					Data->Header = CollapseNode->GetFunctionHeader();
@@ -869,7 +869,7 @@ FRigVMGraphFunctionStore* FRigVMClient::FindFunctionStore(const URigVMLibraryNod
 
 bool FRigVMClient::UpdateFunctionReferences(const FRigVMGraphFunctionHeader& Header, bool bUpdateDependencies, bool bUpdateExternalVariables)
 {
-	URigVMBuildData* BuildData = GetFunctionLibrary()->GetBuildData();
+	URigVMBuildData* BuildData = URigVMBuildData::Get();
 	if (const FRigVMFunctionReferenceArray* FunctionReferenceArray = BuildData->FindFunctionReferences(Header.LibraryPointer))
 	{
 		for (int32 i=0; i<FunctionReferenceArray->Num(); ++i)
diff --git a/Engine/Source/Developer/RigVMDeveloper/Private/RigVMModel/RigVMController.cpp b/Engine/Source/Developer/RigVMDeveloper/Private/RigVMModel/RigVMController.cpp
--- a/Engine/Source/Developer/RigVMDeveloper/Private/RigVMModel/RigVMController.cpp
+++ b/Engine/Source/Developer/RigVMDeveloper/Private/RigVMModel/RigVMController.cpp
@@ -277,7 +277,7 @@ void URigVMController::SetIsRunningUnitTest(bool bIsRunning)
 {
 	bIsRunningUnitTest = bIsRunning;
 
-	if(URigVMBuildData* BuildData = GetBuildData())
+	if(URigVMBuildData* BuildData = URigVMBuildData::Get())
 	{
 		BuildData->SetIsRunningUnitTest(bIsRunning);
 	}	
@@ -3560,7 +3560,7 @@ TArray<FName> URigVMController::ImportNodesFromText(const FString& InText, bool
 
 			if (URigVMFunctionReferenceNode* FunctionRefNode = Cast<URigVMFunctionReferenceNode>(SubNode))
 			{
-				if(URigVMBuildData* BuildData = GetBuildData())
+				if(URigVMBuildData* BuildData = URigVMBuildData::Get())
 				{
 					BuildData->RegisterFunctionReference(FunctionRefNode->GetReferencedFunctionHeader().LibraryPointer, FunctionRefNode);
 				}
@@ -5768,7 +5768,7 @@ void URigVMController::SetReferencedFunction(URigVMFunctionReferenceNode* InFunc
 	
 	if(!(OldReferencedNode == InFunctionRefNode->ReferencedFunctionHeader))
 	{
-		if(URigVMBuildData* BuildData = GetBuildData())
+		if(URigVMBuildData* BuildData = URigVMBuildData::Get())
 		{
 			BuildData->UnregisterFunctionReference(OldReferencedNode.LibraryPointer, InFunctionRefNode);
 			BuildData->RegisterFunctionReference(InFunctionRefNode->ReferencedFunctionHeader.LibraryPointer, InFunctionRefNode);
@@ -5965,7 +5965,7 @@ bool URigVMController::RemoveNode(URigVMNode* InNode, bool bSetupUndoRedo, bool
 		// If we are removing a reference, remove the function references to this node in the function library
 		if(URigVMFunctionReferenceNode* FunctionReferenceNode = Cast<URigVMFunctionReferenceNode>(LibraryNode))
 		{
-			if(URigVMBuildData* BuildData = GetBuildData())
+			if(URigVMBuildData* BuildData = URigVMBuildData::Get())
 			{
 				BuildData->UnregisterFunctionReference(FunctionReferenceNode->GetReferencedFunctionHeader().LibraryPointer, FunctionReferenceNode);
 			}
@@ -5973,7 +5973,7 @@ bool URigVMController::RemoveNode(URigVMNode* InNode, bool bSetupUndoRedo, bool
 		// If we are removing a function, remove all the references first
 		else if (URigVMFunctionLibrary* FunctionLibrary = Cast<URigVMFunctionLibrary>(LibraryNode->GetGraph()))
 		{
-			if(URigVMBuildData* BuildData = GetBuildData())
+			if(URigVMBuildData* BuildData = URigVMBuildData::Get())
 			{
 				FRigVMGraphFunctionIdentifier Identifier = LibraryNode->GetFunctionIdentifier();
 				if (const FRigVMFunctionReferenceArray* ReferencesEntry = BuildData->FindFunctionReferences(Identifier))
@@ -6207,7 +6207,7 @@ bool URigVMController::RenameNode(URigVMNode* InNode, const FName& InNewName, bo
 		if (URigVMFunctionLibrary* FunctionLibrary = Cast<URigVMFunctionLibrary>(LibraryNode->GetGraph()))
 		{
 			// update the table in the build data
-			if(URigVMBuildData* BuildData = GetBuildData())
+			if(URigVMBuildData* BuildData = URigVMBuildData::Get())
 			{
 				for(const TPair< FRigVMGraphFunctionIdentifier, FRigVMFunctionReferenceArray >& Pair: BuildData->GraphFunctionReferences)
 				{
@@ -11348,7 +11348,7 @@ URigVMFunctionReferenceNode* URigVMController::AddFunctionReferenceNodeFromDescr
 
 	Notify(ERigVMGraphNotifType::NodeAdded, FunctionRefNode);
 
-	if (URigVMBuildData* BuildData = GetBuildData())
+	if (URigVMBuildData* BuildData = URigVMBuildData::Get())
 	{
 		BuildData->RegisterFunctionReference(FunctionRefNode->GetReferencedFunctionHeader().LibraryPointer, FunctionRefNode);
 	}
@@ -14588,20 +14588,6 @@ FProperty* URigVMController::FindPropertyForPin(const FString& InPinPath)
 	return nullptr;
 }
 
-URigVMBuildData* URigVMController::GetBuildData(bool bCreateIfNeeded)
-{
-	static TStrongObjectPtr<URigVMBuildData> sBuildData;
-	if(!sBuildData.IsValid() && bCreateIfNeeded && IsInGameThread())
-	{
-		sBuildData = TStrongObjectPtr<URigVMBuildData>(
-			NewObject<URigVMBuildData>(
-				GetTransientPackage(), 
-				TEXT("RigVMBuildData"), 
-				RF_Transient));
-	}
-	return sBuildData.Get();
-}
-
 int32 URigVMController::DetachLinksFromPinObjects(const TArray<URigVMLink*>* InLinks, bool bNotify)
 {
 	URigVMGraph* Graph = GetGraph();
diff --git a/Engine/Source/Developer/RigVMDeveloper/Private/RigVMModel/RigVMFunctionLibrary.cpp b/Engine/Source/Developer/RigVMDeveloper/Private/RigVMModel/RigVMFunctionLibrary.cpp
--- a/Engine/Source/Developer/RigVMDeveloper/Private/RigVMModel/RigVMFunctionLibrary.cpp
+++ b/Engine/Source/Developer/RigVMDeveloper/Private/RigVMModel/RigVMFunctionLibrary.cpp
@@ -71,11 +71,6 @@ URigVMLibraryNode* URigVMFunctionLibrary::FindFunctionForNode(URigVMNode* InNode
 	return Cast<URigVMLibraryNode>(Subject);
 }
 
-URigVMBuildData* URigVMFunctionLibrary::GetBuildData() const
-{
-	return URigVMController::GetBuildData();
-}
-
 TArray< TSoftObjectPtr<URigVMFunctionReferenceNode> > URigVMFunctionLibrary::GetReferencesForFunction(const FName& InFunctionName)
 {
 	TArray< TSoftObjectPtr<URigVMFunctionReferenceNode> > Result;
@@ -105,7 +100,7 @@ void URigVMFunctionLibrary::ForEachReference(const FName& InFunctionName,
 {
 	if (URigVMLibraryNode* Function = FindFunction(InFunctionName))
 	{
-		if(const URigVMBuildData* BuildData = GetBuildData())
+		if(const URigVMBuildData* BuildData = URigVMBuildData::Get())
 		{
 			BuildData->ForEachFunctionReference(Function->GetFunctionIdentifier(),PerReferenceFunction);
 		}
@@ -117,7 +112,7 @@ void URigVMFunctionLibrary::ForEachReferenceSoftPtr(const FName& InFunctionName,
 {
 	if (URigVMLibraryNode* Function = FindFunction(InFunctionName))
 	{
-		if(const URigVMBuildData* BuildData = GetBuildData())
+		if(const URigVMBuildData* BuildData = URigVMBuildData::Get())
 		{
 			BuildData->ForEachFunctionReferenceSoftPtr(Function->GetFunctionIdentifier(), PerReferenceFunction);
 		}
diff --git a/Engine/Source/Developer/RigVMDeveloper/Public/RigVMModel/RigVMBuildData.h b/Engine/Source/Developer/RigVMDeveloper/Public/RigVMModel/RigVMBuildData.h
--- a/Engine/Source/Developer/RigVMDeveloper/Public/RigVMModel/RigVMBuildData.h
+++ b/Engine/Source/Developer/RigVMDeveloper/Public/RigVMModel/RigVMBuildData.h
@@ -68,8 +68,11 @@ class RIGVMDEVELOPER_API URigVMBuildData : public UObject
 
 public:
 
-	// Default constructor
-	URigVMBuildData();
+	// Returns the singleton build data
+	static URigVMBuildData* Get();
+	
+	// Looks for all function references (in RigVMClientHost metadata) and initializes the URigVMBuildData
+	void InitializeIfNeeded();
 
 	// Returns the list of references for a given function definition
 	const FRigVMFunctionReferenceArray* FindFunctionReferences(const FRigVMGraphFunctionIdentifier& InFunction) const;
@@ -112,6 +115,11 @@ public:
 
 private:
 
+	// disable default constructor
+	URigVMBuildData();
+	
+	static bool bInitialized;
+
 	UPROPERTY(meta=(DeprecatedProperty))
 	TMap< TSoftObjectPtr<URigVMLibraryNode>, FRigVMFunctionReferenceArray > FunctionReferences_DEPRECATED;
 
diff --git a/Engine/Source/Developer/RigVMDeveloper/Public/RigVMModel/RigVMClient.h b/Engine/Source/Developer/RigVMDeveloper/Public/RigVMModel/RigVMClient.h
--- a/Engine/Source/Developer/RigVMDeveloper/Public/RigVMModel/RigVMClient.h
+++ b/Engine/Source/Developer/RigVMDeveloper/Public/RigVMModel/RigVMClient.h
@@ -48,6 +48,10 @@ public:
 
 	// Reacts to a request to configure a controller
 	virtual void HandleConfigureRigVMController(const FRigVMClient* InClient, URigVMController* InControllerToConfigure) = 0;
+
+	/** Asset searchable information function references in this rig */
+	UPROPERTY(AssetRegistrySearchable)
+	TArray<FRigVMReferenceNodeData> FunctionReferenceNodeData;
 };
 
 UINTERFACE()
diff --git a/Engine/Source/Developer/RigVMDeveloper/Public/RigVMModel/RigVMController.h b/Engine/Source/Developer/RigVMDeveloper/Public/RigVMModel/RigVMController.h
--- a/Engine/Source/Developer/RigVMDeveloper/Public/RigVMModel/RigVMController.h
+++ b/Engine/Source/Developer/RigVMDeveloper/Public/RigVMModel/RigVMController.h
@@ -862,9 +862,6 @@ public:
 	// A delegate to request to configure an options instance for a node workflow
 	FRigVMController_ConfigureWorkflowOptionsDelegate ConfigureWorkflowOptionsDelegate; 
 
-	// Returns the build data of the host
-	static URigVMBuildData* GetBuildData(bool bCreateIfNeeded = true);
-
 	int32 DetachLinksFromPinObjects(const TArray<URigVMLink*>* InLinks = nullptr, bool bNotify = false);
 	int32 ReattachLinksToPinObjects(bool bFollowCoreRedirectors = false, const TArray<URigVMLink*>* InLinks = nullptr, bool bNotify = false, bool bSetupOrphanedPins = false, bool bAllowNonArgumentLinks = false);
 	void AddPinRedirector(bool bInput, bool bOutput, const FString& OldPinPath, const FString& NewPinPath);
diff --git a/Engine/Source/Developer/RigVMDeveloper/Public/RigVMModel/RigVMFunctionLibrary.h b/Engine/Source/Developer/RigVMDeveloper/Public/RigVMModel/RigVMFunctionLibrary.h
--- a/Engine/Source/Developer/RigVMDeveloper/Public/RigVMModel/RigVMFunctionLibrary.h
+++ b/Engine/Source/Developer/RigVMDeveloper/Public/RigVMModel/RigVMFunctionLibrary.h
@@ -42,10 +42,6 @@ public:
 	UFUNCTION(BlueprintCallable, Category = RigVMGraph)
     URigVMLibraryNode* FindFunctionForNode(URigVMNode* InNode) const;
 
-	// returns the build data for for this library
-	UFUNCTION()
-	URigVMBuildData* GetBuildData() const;
-
 	// Returns all references for a given function name
 	UFUNCTION(BlueprintCallable, Category = RigVMGraph)
 	TArray< TSoftObjectPtr<URigVMFunctionReferenceNode> > GetReferencesForFunction(const FName& InFunctionName);
