RigVM: Guard against incompatible execute context in compiler

diff --git a/Engine/Source/Developer/RigVMDeveloper/Private/RigVMCompiler/RigVMCompiler.cpp b/Engine/Source/Developer/RigVMDeveloper/Private/RigVMCompiler/RigVMCompiler.cpp
--- a/Engine/Source/Developer/RigVMDeveloper/Private/RigVMCompiler/RigVMCompiler.cpp
+++ b/Engine/Source/Developer/RigVMDeveloper/Private/RigVMCompiler/RigVMCompiler.cpp
@@ -405,6 +405,13 @@ bool URigVMCompiler::Compile(TArray<URigVMGraph*> InGraphs, URigVMController* In
 	TArray<URigVMGraph*> VisitedGraphs;
 	VisitedGraphs.Append(InGraphs);
 
+	// also during traverse - find all known execute context and argument notations
+	// for functions / dispatches / templates.
+	// we only allow compatible execute context structs within a VM
+	TArray<UScriptStruct*> ExecuteContextStructs;
+	
+	const FRigVMRegistry& Registry = FRigVMRegistry::Get();
+
 	for(int32 GraphIndex=0; GraphIndex<VisitedGraphs.Num(); GraphIndex++)
 	{
 		URigVMGraph* VisitedGraph = VisitedGraphs[GraphIndex];
@@ -599,6 +606,57 @@ bool URigVMCompiler::Compile(TArray<URigVMGraph*> InGraphs, URigVMController* In
 				}
 			}
 
+			if(const URigVMDispatchNode* DispatchNode = Cast<URigVMDispatchNode>(ModelNode))
+			{
+				if(const FRigVMDispatchFactory* Factory = DispatchNode->GetFactory())
+				{
+					ExecuteContextStructs.AddUnique(Factory->GetExecuteContextStruct());
+				}
+			}
+			else if(const URigVMTemplateNode* TemplateNode = Cast<URigVMTemplateNode>(ModelNode))
+			{
+				if(const FRigVMFunction* ResolvedFunction = TemplateNode->GetResolvedFunction())
+				{
+					if(UScriptStruct* RigVMStruct = ResolvedFunction->Struct)
+					{
+						for (TFieldIterator<FProperty> It(RigVMStruct, EFieldIterationFlags::IncludeAll); It; ++It)
+						{
+							const FProperty* Property = *It;
+							if(const FArrayProperty* ArrayProperty = CastField<FArrayProperty>(Property))
+							{
+								Property = ArrayProperty->Inner;
+							}
+							
+							if(const FStructProperty* StructProperty = CastField<FStructProperty>(Property))
+							{
+								if(StructProperty->Struct->IsChildOf(FRigVMExecuteContext::StaticStruct()))
+								{
+									ExecuteContextStructs.AddUnique(StructProperty->Struct);
+								}
+							}
+						}
+					}
+				}
+				
+				if(const FRigVMTemplate* Template = TemplateNode->GetTemplate())
+				{
+					for(int32 Index = 0; Index < Template->NumExecuteArguments(); Index++)
+					{
+						if(const FRigVMExecuteArgument* Argument = Template->GetExecuteArgument(Index))
+						{
+							if(Registry.IsExecuteType(Argument->TypeIndex))
+							{
+								const FRigVMTemplateArgumentType& Type = Registry.GetType(Argument->TypeIndex);
+								if(UScriptStruct* ExecuteContextStruct = Cast<UScriptStruct>(Type.CPPTypeObject))
+								{
+									ExecuteContextStructs.AddUnique(ExecuteContextStruct);
+								}
+							}
+						}
+					}
+				}
+			}
+
 			for(URigVMPin* Pin : ModelNode->Pins)
 			{
 				if(!URigVMController::EnsurePinValidity(Pin, true))
@@ -609,6 +667,32 @@ bool URigVMCompiler::Compile(TArray<URigVMGraph*> InGraphs, URigVMController* In
 		}
 	}
 
+	// check if the execute context structs are compatible
+	UScriptStruct* TopLevelExecuteContextStruct = FRigVMExecuteContext::StaticStruct();
+	for(UScriptStruct* ExecuteContextStruct : ExecuteContextStructs)
+	{
+		if(TopLevelExecuteContextStruct->IsChildOf(ExecuteContextStruct))
+		{
+			continue;
+		}
+		
+		if(ExecuteContextStruct->IsChildOf(TopLevelExecuteContextStruct))
+		{
+			TopLevelExecuteContextStruct = ExecuteContextStruct;
+			continue;
+		}
+
+		// at this point the execute context structs are unrelated
+		static constexpr TCHAR CompatibleExecuteContextsFormat[] = TEXT("ExecuteContext types '%s' and '%s' are not compatible. Cannot compile this graph.");
+		const FString CompatibleExecuteContextsMessage = FString::Printf(
+			CompatibleExecuteContextsFormat,
+			*ExecuteContextStruct->GetStructCPPName(),
+			*TopLevelExecuteContextStruct->GetStructCPPName());
+		Settings.ASTSettings.Report(EMessageSeverity::Error, nullptr, CompatibleExecuteContextsMessage);
+		bEncounteredGraphError = true;
+		break;
+	}
+
 	if(bEncounteredGraphError)
 	{
 		return false;
