RigVM: Add executecontext as compiler setting

diff --git a/Engine/Plugins/Animation/ControlRig/Source/ControlRigDeveloper/Private/ControlRigBlueprint.cpp b/Engine/Plugins/Animation/ControlRig/Source/ControlRigDeveloper/Private/ControlRigBlueprint.cpp
--- a/Engine/Plugins/Animation/ControlRig/Source/ControlRigDeveloper/Private/ControlRigBlueprint.cpp
+++ b/Engine/Plugins/Animation/ControlRig/Source/ControlRigDeveloper/Private/ControlRigBlueprint.cpp
@@ -1223,7 +1223,8 @@ void UControlRigBlueprint::RecompileVM()
 		}
 
 		URigVMCompiler* Compiler = URigVMCompiler::StaticClass()->GetDefaultObject<URigVMCompiler>();
-		Compiler->Settings = (bCompileInDebugMode) ? FRigVMCompileSettings::Fast() : VMCompileSettings;
+		VMCompileSettings.SetExecuteContextStruct(FControlRigExecuteContext::StaticStruct());
+		Compiler->Settings = (bCompileInDebugMode) ? FRigVMCompileSettings::Fast(VMCompileSettings.GetExecuteContextStruct()) : VMCompileSettings;
 		Compiler->Compile(RigVMClient.GetAllModels(false, false), GetOrCreateController(), CDO->VM, CDO->GetExternalVariablesImpl(false), UserData, &PinToOperandMap);
 
 		if (bErrorsDuringCompilation)
diff --git a/Engine/Source/Developer/RigVMDeveloper/Private/RigVMCompiler/RigVMAST.cpp b/Engine/Source/Developer/RigVMDeveloper/Private/RigVMCompiler/RigVMAST.cpp
--- a/Engine/Source/Developer/RigVMDeveloper/Private/RigVMCompiler/RigVMAST.cpp
+++ b/Engine/Source/Developer/RigVMDeveloper/Private/RigVMCompiler/RigVMAST.cpp
@@ -2073,7 +2073,7 @@ bool FRigVMParserAST::FoldConstantValuesToLiterals(TArray<URigVMGraph*> InGraphs
 	TempCompiler->Settings.SetupNodeInstructionIndex = false;
 	TempCompiler->Settings.IsPreprocessorPhase = true;
 	TempCompiler->Settings.EnablePinWatches = false;
-	TempCompiler->Settings.ASTSettings = FRigVMParserASTSettings::Fast();
+	TempCompiler->Settings.ASTSettings = FRigVMParserASTSettings::Fast(Settings.ExecuteContextStruct);
 
 	TempCompiler->Compile(InGraphs, InController, TempVM, InExternalVariables, InRigVMUserData, &Operands, TempAST);
 
diff --git a/Engine/Source/Developer/RigVMDeveloper/Private/RigVMCompiler/RigVMCompiler.cpp b/Engine/Source/Developer/RigVMDeveloper/Private/RigVMCompiler/RigVMCompiler.cpp
--- a/Engine/Source/Developer/RigVMDeveloper/Private/RigVMCompiler/RigVMCompiler.cpp
+++ b/Engine/Source/Developer/RigVMDeveloper/Private/RigVMCompiler/RigVMCompiler.cpp
@@ -68,6 +68,16 @@ FRigVMCompileSettings::FRigVMCompileSettings()
 {
 }
 
+FRigVMCompileSettings::FRigVMCompileSettings(UScriptStruct* InExecuteContextScriptStruct)
+	: FRigVMCompileSettings()
+{
+	ASTSettings.ExecuteContextStruct = InExecuteContextScriptStruct;
+	if(ASTSettings.ExecuteContextStruct == nullptr)
+	{
+		ASTSettings.ExecuteContextStruct = FRigVMExecuteContext::StaticStruct();
+	}
+}
+
 FRigVMOperand FRigVMCompilerWorkData::AddProperty(
 	ERigVMMemoryType InMemoryType,
 	const FName& InName,
@@ -263,6 +273,29 @@ bool URigVMCompiler::Compile(TArray<URigVMGraph*> InGraphs, URigVMController* In
 		return false;
 	}
 
+	if (Settings.GetExecuteContextStruct() == nullptr)
+	{
+		ReportError(TEXT("Compiler settings don't provide the ExecuteContext to use. Cannot compile."));
+		return false;;
+	}
+
+	for(URigVMGraph* Graph : InGraphs)
+	{
+		if(Graph->GetExecuteContextStruct())
+		{
+			if(!Settings.GetExecuteContextStruct()->IsChildOf(Graph->GetExecuteContextStruct()))
+			{
+				ReportErrorf(
+					TEXT("Compiler settings' ExecuteContext (%s) is not compatible with '%s' graph's ExecuteContext (%s). Cannot compile."),
+					*Settings.GetExecuteContextStruct()->GetStructCPPName(),
+					*Graph->GetNodePath(),
+					*Graph->GetExecuteContextStruct()->GetStructCPPName()
+				);
+				return false;;
+			}
+		}
+	}
+
 	for(int32 Index = 1; Index < InGraphs.Num(); Index++)
 	{
 		if(InGraphs[0]->GetOuter() != InGraphs[Index]->GetOuter())
@@ -683,16 +716,27 @@ bool URigVMCompiler::Compile(TArray<URigVMGraph*> InGraphs, URigVMController* In
 		}
 
 		// at this point the execute context structs are unrelated
-		static constexpr TCHAR CompatibleExecuteContextsFormat[] = TEXT("ExecuteContext types '%s' and '%s' are not compatible. Cannot compile this graph.");
-		const FString CompatibleExecuteContextsMessage = FString::Printf(
-			CompatibleExecuteContextsFormat,
+		ReportErrorf(
+			TEXT("ExecuteContext types '%s' and '%s' are not compatible. Cannot compile this graph."),
 			*ExecuteContextStruct->GetStructCPPName(),
 			*TopLevelExecuteContextStruct->GetStructCPPName());
-		Settings.ASTSettings.Report(EMessageSeverity::Error, nullptr, CompatibleExecuteContextsMessage);
 		bEncounteredGraphError = true;
 		break;
 	}
 
+	// check if the top level matches the settings
+	if(TopLevelExecuteContextStruct)
+	{
+		if(!Settings.GetExecuteContextStruct()->IsChildOf(TopLevelExecuteContextStruct))
+		{
+			ReportErrorf(
+				TEXT("The graph contains an incompatible execute context '%s'. The compiler settings expect '%s' (or below)."),
+				*TopLevelExecuteContextStruct->GetStructCPPName(),
+				*Settings.GetExecuteContextStruct()->GetStructCPPName());
+			bEncounteredGraphError = true;
+		}
+	}
+
 	if(bEncounteredGraphError)
 	{
 		return false;
diff --git a/Engine/Source/Developer/RigVMDeveloper/Public/RigVMCompiler/RigVMAST.h b/Engine/Source/Developer/RigVMDeveloper/Public/RigVMCompiler/RigVMAST.h
--- a/Engine/Source/Developer/RigVMDeveloper/Public/RigVMCompiler/RigVMAST.h
+++ b/Engine/Source/Developer/RigVMDeveloper/Public/RigVMCompiler/RigVMAST.h
@@ -1252,24 +1252,29 @@ struct RIGVMDEVELOPER_API FRigVMParserASTSettings
 
 	FRigVMReportDelegate ReportDelegate;
 
+	UPROPERTY(VisibleAnywhere, Category = "AST", transient)
+	TObjectPtr<UScriptStruct> ExecuteContextStruct;
+
 	// static method to provide fast AST parse settings
-	static FRigVMParserASTSettings Fast()
+	static FRigVMParserASTSettings Fast(UScriptStruct* InExecuteContextStruct = nullptr)
 	{
 		FRigVMParserASTSettings Settings;
 		Settings.bFoldAssignments = false;
 		Settings.bFoldLiterals = false;
 		Settings.bFoldConstantBranches = false;
+		Settings.ExecuteContextStruct = InExecuteContextStruct ? InExecuteContextStruct : FRigVMExecuteContext::StaticStruct();
 		return Settings;
 	}
 
 	// static method to provide AST parse settings
 	// tuned for a fast executing runtime, but slow parse
-	static FRigVMParserASTSettings Optimized()
+	static FRigVMParserASTSettings Optimized(UScriptStruct* InExecuteContextStruct = nullptr)
 	{
 		FRigVMParserASTSettings Settings;
 		Settings.bFoldAssignments = true;
 		Settings.bFoldLiterals = true;
 		Settings.bFoldConstantBranches = true;
+		Settings.ExecuteContextStruct = InExecuteContextStruct ? InExecuteContextStruct : FRigVMExecuteContext::StaticStruct();
 		return Settings;
 	}
 
diff --git a/Engine/Source/Developer/RigVMDeveloper/Public/RigVMCompiler/RigVMCompiler.h b/Engine/Source/Developer/RigVMDeveloper/Public/RigVMCompiler/RigVMCompiler.h
--- a/Engine/Source/Developer/RigVMDeveloper/Public/RigVMCompiler/RigVMCompiler.h
+++ b/Engine/Source/Developer/RigVMDeveloper/Public/RigVMCompiler/RigVMCompiler.h
@@ -31,6 +31,11 @@ public:
 
 	FRigVMCompileSettings();
 
+	FRigVMCompileSettings(UScriptStruct* InExecuteContextScriptStruct);
+
+	UScriptStruct* GetExecuteContextStruct() const { return ASTSettings.ExecuteContextStruct; }
+	void SetExecuteContextStruct(UScriptStruct* InExecuteContextScriptStruct) { ASTSettings.ExecuteContextStruct = InExecuteContextScriptStruct; }
+
 	UPROPERTY(EditAnywhere, BlueprintReadWrite, Category = FRigVMCompileSettings)
 	bool SurpressInfoMessages;
 
@@ -52,21 +57,21 @@ public:
 	UPROPERTY()
 	bool SetupNodeInstructionIndex;
 
-	static FRigVMCompileSettings Fast()
+	static FRigVMCompileSettings Fast(UScriptStruct* InExecuteContextStruct = nullptr)
 	{
-		FRigVMCompileSettings Settings;
+		FRigVMCompileSettings Settings(InExecuteContextStruct);
 		Settings.EnablePinWatches = true;
 		Settings.IsPreprocessorPhase = false;
-		Settings.ASTSettings = FRigVMParserASTSettings::Fast();
+		Settings.ASTSettings = FRigVMParserASTSettings::Fast(InExecuteContextStruct);
 		return Settings;
 	}
 
-	static FRigVMCompileSettings Optimized()
+	static FRigVMCompileSettings Optimized(UScriptStruct* InExecuteContextStruct = nullptr)
 	{
-		FRigVMCompileSettings Settings;
+		FRigVMCompileSettings Settings(InExecuteContextStruct);
 		Settings.EnablePinWatches = false;
 		Settings.IsPreprocessorPhase = false;
-		Settings.ASTSettings = FRigVMParserASTSettings::Optimized();
+		Settings.ASTSettings = FRigVMParserASTSettings::Optimized(InExecuteContextStruct);
 		return Settings;
 	}
 
diff --git a/Engine/Source/Runtime/Core/Public/UObject/UE5MainStreamObjectVersion.h b/Engine/Source/Runtime/Core/Public/UObject/UE5MainStreamObjectVersion.h
--- a/Engine/Source/Runtime/Core/Public/UObject/UE5MainStreamObjectVersion.h
+++ b/Engine/Source/Runtime/Core/Public/UObject/UE5MainStreamObjectVersion.h
@@ -303,6 +303,9 @@ struct CORE_API FUE5MainStreamObjectVersion
 
 		// Store function information (and compilation data) in blueprint generated class
 		RigVMSaveFunctionAccessInModel,
+
+		// Store the RigVM execute context struct the VM uses in the archive
+		RigVMSerializeExecuteContextStruct,
 		
 		// -----<new versions can be added above this line>-------------------------------------------------
 		VersionPlusOne,
diff --git a/Engine/Source/Runtime/RigVM/Private/RigVMCore/RigVM.cpp b/Engine/Source/Runtime/RigVM/Private/RigVMCore/RigVM.cpp
--- a/Engine/Source/Runtime/RigVM/Private/RigVMCore/RigVM.cpp
+++ b/Engine/Source/Runtime/RigVM/Private/RigVMCore/RigVM.cpp
@@ -116,7 +116,7 @@ void URigVM::Serialize(FArchive& Ar)
 	{
 		return;
 	}
-
+	
 	// call into the super class to serialize any uproperty
 	if(Ar.IsObjectReferenceCollector() || Ar.IsCountingMemory())
 	{
@@ -146,6 +146,10 @@ void URigVM::Save(FArchive& Ar)
 	int32 RigVMUClassBasedStorageDefine = 0; // this used to 
 	Ar << RigVMUClassBasedStorageDefine;
 
+	// store the context public data object path
+	FString ExecuteContextPath = GetContextPublicDataStruct()->GetPathName();
+	Ar << ExecuteContextPath;
+
 	// we rely on Ar.IsIgnoringArchetypeRef for determining if we are currently performing
 	// CPFUO (Copy Properties for unrelated objects). During a reinstance pass we don't
 	// want to overwrite the bytecode and some other properties - since that's handled already
@@ -173,6 +177,22 @@ void URigVM::Load(FArchive& Ar)
 		Ar << RigVMUClassBasedStorageDefine;
 	}
 
+	if (Ar.CustomVer(FUE5MainStreamObjectVersion::GUID) >= FUE5MainStreamObjectVersion::RigVMSerializeExecuteContextStruct)
+	{
+		FString ExecuteContextPath;
+		Ar << ExecuteContextPath;
+
+		const FTopLevelAssetPath AssetPath(*ExecuteContextPath);
+		UScriptStruct* ExecuteContextScriptStruct = FindObject<UScriptStruct>(AssetPath);
+		if(ExecuteContextScriptStruct == nullptr)
+		{
+			Reset();
+			return;
+		}
+
+		SetContextPublicDataStruct(ExecuteContextScriptStruct);
+	}
+
 	if(RigVMUClassBasedStorageDefine == 1)
 	{
 		FRigVMMemoryContainer WorkMemoryStorage;
