Control Rig: Print python commands when adding function references of external functions

diff --git a/Engine/Source/Developer/RigVMDeveloper/Private/RigVMModel/RigVMController.cpp b/Engine/Source/Developer/RigVMDeveloper/Private/RigVMModel/RigVMController.cpp
--- a/Engine/Source/Developer/RigVMDeveloper/Private/RigVMModel/RigVMController.cpp
+++ b/Engine/Source/Developer/RigVMDeveloper/Private/RigVMModel/RigVMController.cpp
@@ -680,17 +680,15 @@ TArray<FString> URigVMController::GetAddNodePythonCommands(URigVMNode* Node) con
 					*RigVMPythonUtils::Vector2DToPythonString(RefNode->GetPosition()), 
 					*NodeName));
 		}
-		// todo UE-169886
-		// else
-		// {
-		// 	Commands.Add(FString::Printf(TEXT("function_blueprint = unreal.load_object(name = '%s', outer = None)"),
-		// 		*Library->GetOuter()->GetPathName()));
-		// 	Commands.Add(FString::Printf(TEXT("blueprint.get_controller_by_name('%s').add_function_reference_node(function_blueprint.get_local_function_library().find_function('%s'), %s, '%s')"),
-		// 				*GraphName,
-		// 				*NodeName,
-		// 				*RigVMPythonUtils::Vector2DToPythonString(LibraryNode->GetPosition()), 
-		// 				*NodeName));
-		// }
+		else
+		{
+			Commands.Add(FString::Printf(TEXT("blueprint.get_controller_by_name('%s').add_external_function_reference_node('%s', '%s', %s, '%s')"),
+						*GraphName,
+						*RefNode->GetReferencedFunctionHeader().LibraryPointer.HostObject.ToString(),
+						*RefNode->GetReferencedFunctionHeader().Name.ToString(),
+						*RigVMPythonUtils::Vector2DToPythonString(RefNode->GetPosition()), 
+						*NodeName));
+		}
 	}
 	else if (const URigVMCollapseNode* CollapseNode = Cast<URigVMCollapseNode>(Node))
 	{
@@ -11377,39 +11375,84 @@ URigVMFunctionReferenceNode* URigVMController::AddFunctionReferenceNodeFromDescr
 
 	if (bPrintPythonCommand)
 	{
-		// todo: UE-169886
-		// const FString GraphName = GetSanitizedGraphName(GetGraph()->GetGraphName());
-		// const FString FunctionDefinitionName = GetSanitizedNodeName(InFunctionDefinition.Name.ToString());
-		//
-		// URigVMLibraryNode* LibraryNode = InFunctionDefinition.LibraryNode.IsValid() ? InFunctionDefinition.LibraryNode.Get() : nullptr;
-		// if (LibraryNode && LibraryNode->GetLibrary() == GetGraph()->GetDefaultFunctionLibrary())
-		// {
-		//
-		// 	RigVMPythonUtils::Print(GetGraphOuterName(), 
-		// 		FString::Printf(TEXT("blueprint.get_controller_by_name('%s').add_function_reference_node(library.find_function('%s'), %s, '%s')"),
-		// 				*GraphName,
-		// 				*FunctionDefinitionName,
-		// 				*RigVMPythonUtils::Vector2DToPythonString(InNodePosition),
-		// 				*NodeName));
-		// }
-		// else
-		// {
-		// 	RigVMPythonUtils::Print(GetGraphOuterName(), 
-		// 		FString::Printf(TEXT("function_blueprint = unreal.load_object(name = '%s', outer = None)"),
-		// 		*InFunctionDefinition->GetLibrary()->GetOuter()->GetPathName()));
-		// 	RigVMPythonUtils::Print(GetGraphOuterName(), 
-		// 		FString::Printf(TEXT("blueprint.get_controller_by_name('%s').add_function_reference_node(function_blueprint.get_local_function_library().find_function('%s'), %s, '%s')"),
-		// 				*GraphName,
-		// 				*FunctionDefinitionName,
-		// 				*RigVMPythonUtils::Vector2DToPythonString(InFunctionDefinition->GetPosition()), 
-		// 				*FunctionDefinitionName));
-		// }
+		const FString GraphName = GetSanitizedGraphName(GetGraph()->GetGraphName());
+		const FString FunctionDefinitionName = GetSanitizedNodeName(InFunctionDefinition.Name.ToString());
+
+		bool bLocal = false;
+		if(IRigVMClientHost* ClientHost = GetImplementingOuter<IRigVMClientHost>())
+		{
+			if (InFunctionDefinition.LibraryPointer.HostObject == Cast<UObject>(ClientHost->GetRigVMGraphFunctionHost()))
+			{
+				bLocal = true;
+				RigVMPythonUtils::Print(GetGraphOuterName(), 
+					FString::Printf(TEXT("blueprint.get_controller_by_name('%s').add_function_reference_node(library.find_function('%s'), %s, '%s')"),
+							*GraphName,
+							*FunctionDefinitionName,
+							*RigVMPythonUtils::Vector2DToPythonString(InNodePosition),
+							*NodeName));
+			}
+		}
 		
+		if (!bLocal)
+		{
+			RigVMPythonUtils::Print(GetGraphOuterName(), 
+				FString::Printf(TEXT("blueprint.get_controller_by_name('%s').add_external_function_reference_node('%s', '%s', %s, '%s')"),
+						*GraphName,
+						*InFunctionDefinition.LibraryPointer.HostObject.ToString(),
+						*FunctionDefinitionName,
+						*RigVMPythonUtils::Vector2DToPythonString(InNodePosition), 
+						*NodeName));
+		}
 	}
 
 	return FunctionRefNode;
 }
 
+URigVMFunctionReferenceNode* URigVMController::AddExternalFunctionReferenceNode(const FString& InHostPath, const FName& InFunctionName, const FVector2D& InNodePosition, const FString& InNodeName, bool bSetupUndoRedo, bool bPrintPythonCommand)
+{
+	if (!IsValidGraph())
+	{
+		return nullptr;
+	}
+
+	if (!bIsTransacting && !IsGraphEditable())
+	{
+		return nullptr;
+	}
+
+	URigVMGraph* Graph = GetGraph();
+	check(Graph);
+
+	if (Graph->IsA<URigVMFunctionLibrary>())
+	{
+		ReportError(TEXT("Cannot add function reference nodes to function library graphs."));
+		return nullptr;
+	}
+
+	UObject* HostObject = StaticLoadObject(UObject::StaticClass(), NULL, *InHostPath, NULL, LOAD_None, NULL);
+	if (!HostObject)
+	{
+		ReportErrorf(TEXT("Failed to load the Host object %s."), *InHostPath);
+		return nullptr;
+	}
+
+	IRigVMGraphFunctionHost* FunctionHost = Cast<IRigVMGraphFunctionHost>(HostObject);
+	if (!FunctionHost)
+	{
+		ReportError(TEXT("Host object is not a IRigVMGraphFunctionHost."));
+		return nullptr;
+	}
+
+	FRigVMGraphFunctionData* Data = FunctionHost->GetRigVMGraphFunctionStore()->FindFunctionByName(InFunctionName);
+	if (!Data)
+	{
+		ReportErrorf(TEXT("Function %s not found in host %s."), *InFunctionName.ToString(), *InHostPath);
+		return nullptr;
+	}
+
+	return AddFunctionReferenceNodeFromDescription(Data->Header, InNodePosition, InNodeName, bSetupUndoRedo, bPrintPythonCommand);
+}
+
 bool URigVMController::SetRemappedVariable(URigVMFunctionReferenceNode* InFunctionRefNode,
 	const FName& InInnerVariableName, const FName& InOuterVariableName, bool bSetupUndoRedo)
 {
diff --git a/Engine/Source/Developer/RigVMDeveloper/Public/RigVMModel/RigVMController.h b/Engine/Source/Developer/RigVMDeveloper/Public/RigVMModel/RigVMController.h
--- a/Engine/Source/Developer/RigVMDeveloper/Public/RigVMModel/RigVMController.h
+++ b/Engine/Source/Developer/RigVMDeveloper/Public/RigVMModel/RigVMController.h
@@ -765,6 +765,8 @@ public:
 	UFUNCTION(BlueprintCallable, Category = RigVMController)
 	URigVMFunctionReferenceNode* AddFunctionReferenceNodeFromDescription(const FRigVMGraphFunctionHeader& InFunctionDefinition, const FVector2D& InNodePosition = FVector2D::ZeroVector, const
 	                                                                     FString& InNodeName = TEXT(""), bool bSetupUndoRedo = true, bool bPrintPythonCommand = false, bool bAllowPrivateFunctions = false);
+	UFUNCTION(BlueprintCallable, Category = RigVMController)
+	URigVMFunctionReferenceNode* AddExternalFunctionReferenceNode(const FString& InHostPath, const FName& InFunctionName, const FVector2D& InNodePosition = FVector2D::ZeroVector, const	                                                              FString& InNodeName = TEXT(""), bool bSetupUndoRedo = true, bool bPrintPythonCommand = false);
 	
 	UFUNCTION(BlueprintCallable, Category = RigVMController)
 	URigVMFunctionReferenceNode* AddFunctionReferenceNode(URigVMLibraryNode* InFunctionDefinition, const FVector2D& InNodePosition = FVector2D::ZeroVector, const FString& InNodeName = TEXT(""), bool bSetupUndoRedo = true, bool bPrintPythonCommand = false);
