Control Rig: Fix ensure in CR_BungeeMan

diff --git a/Engine/Source/Developer/RigVMDeveloper/Private/RigVMModel/RigVMController.cpp b/Engine/Source/Developer/RigVMDeveloper/Private/RigVMModel/RigVMController.cpp
--- a/Engine/Source/Developer/RigVMDeveloper/Private/RigVMModel/RigVMController.cpp
+++ b/Engine/Source/Developer/RigVMDeveloper/Private/RigVMModel/RigVMController.cpp
@@ -17319,10 +17319,15 @@ bool URigVMController::UpdateTemplateNodePinTypes(URigVMTemplateNode* InNode, bo
 				int32 PinIndex = Pins.Find(ArgumentPin);
 				WasReduced[PinIndex] = true;
 				FinalPinTypes[PinIndex] = Preferred.TypeIndex;
-				ReducedTypes = ReducedTypes.FilterByPredicate([Preferred, PinIndex](const TPair<int32, TArray<TRigVMTypeIndex>>& Permutation)
-				{
-					return Permutation.Value[PinIndex] == Preferred.TypeIndex;
-				});
+
+				// If the argument is hidden, the available PinTypes will all be INDEX_NONE
+				if (ArgumentPin->Direction != ERigVMPinDirection::Hidden)
+				{					
+					ReducedTypes = ReducedTypes.FilterByPredicate([Preferred, PinIndex](const TPair<int32, TArray<TRigVMTypeIndex>>& Permutation)
+					{
+						return Permutation.Value[PinIndex] == Preferred.TypeIndex;
+					});
+				}
 			}
 		}
 	}
