Control Rig: Fix cycle detected incorrectly

diff --git a/Engine/Source/Developer/RigVMDeveloper/Private/RigVMModel/RigVMPin.cpp b/Engine/Source/Developer/RigVMDeveloper/Private/RigVMModel/RigVMPin.cpp
--- a/Engine/Source/Developer/RigVMDeveloper/Private/RigVMModel/RigVMPin.cpp
+++ b/Engine/Source/Developer/RigVMDeveloper/Private/RigVMModel/RigVMPin.cpp
@@ -1937,14 +1937,11 @@ bool URigVMPin::CanLink(URigVMPin* InSourcePin, URigVMPin* InTargetPin, FString*
 					SourceNodes[SourceNodeIndex]->IsA<URigVMVariableNode>();
 				if (!bNodeCanLinkAnywhere)
 				{
-					if (URigVMUnitNode* UnitNode = Cast<URigVMUnitNode>(SourceNodes[SourceNodeIndex]))
+					// pure / immutable nodes can be connected to any input in any order.
+					// since a new link is going to change the abstract syntax tree 
+					if (!SourceNodes[SourceNodeIndex]->IsMutable())
 					{
-						// pure / immutable nodes can be connected to any input in any order.
-						// since a new link is going to change the abstract syntax tree 
-						if (!UnitNode->IsMutable())
-						{
-							bNodeCanLinkAnywhere = true;
-						}
+						bNodeCanLinkAnywhere = true;
 					}
 				}
 
