RigVM: Fix for cook - addressing fallout from 23264120

diff --git a/Engine/Source/Developer/RigVMDeveloper/Private/RigVMCompiler/RigVMCompiler.cpp b/Engine/Source/Developer/RigVMDeveloper/Private/RigVMCompiler/RigVMCompiler.cpp
--- a/Engine/Source/Developer/RigVMDeveloper/Private/RigVMCompiler/RigVMCompiler.cpp
+++ b/Engine/Source/Developer/RigVMDeveloper/Private/RigVMCompiler/RigVMCompiler.cpp
@@ -319,6 +319,7 @@ bool URigVMCompiler::Compile(TArray<URigVMGraph*> InGraphs, URigVMController* In
 		UserData.Add(FRigVMUserDataArray());
 	}
 
+	OutVM->SetContextPublicDataStruct(Settings.GetExecuteContextStruct());
 	OutVM->Reset();
 
 	TMap<FString, FRigVMOperand> LocalOperands;
