RigVM: Remove support for opaque attributes

diff --git a/Engine/Plugins/Animation/ControlRig/Source/ControlRigEditor/Private/ControlRigEditorModule.cpp b/Engine/Plugins/Animation/ControlRig/Source/ControlRigEditor/Private/ControlRigEditorModule.cpp
--- a/Engine/Plugins/Animation/ControlRig/Source/ControlRigEditor/Private/ControlRigEditorModule.cpp
+++ b/Engine/Plugins/Animation/ControlRig/Source/ControlRigEditor/Private/ControlRigEditorModule.cpp
@@ -1026,6 +1026,11 @@ void FControlRigEditorModule::GetTypeActions(UControlRigBlueprint* CRB, FBluepri
 			continue;
 		}
 
+		if(!Template.SupportsExecuteContextStruct(FControlRigExecuteContext::StaticStruct()))
+		{
+			continue;
+		}
+
 		FText NodeCategory = FText::FromString(Template.GetCategory());
 		FText MenuDesc = FText::FromName(Template.GetName());
 		FText ToolTip = Template.GetTooltipText();
@@ -1037,15 +1042,9 @@ void FControlRigEditorModule::GetTypeActions(UControlRigBlueprint* CRB, FBluepri
 
 	for (const FRigVMDispatchFactory* Factory : Registry.GetFactories())
 	{
-		// allow empty opaque arguments or the control rig notation
-		const TArray<TPair<FName,FString>>& OpaqueArguments = Factory->GetOpaqueArguments();
-		if(!OpaqueArguments.IsEmpty())
+		if(!Factory->SupportsExecuteContextStruct(FControlRigExecuteContext::StaticStruct()))
 		{
-			static const FRigDispatchFactory ControlRigFactory;
-			if(ControlRigFactory.GetOpaqueArguments() != OpaqueArguments)
-			{
-				continue;
-			}
+			continue;
 		}
 
 		const FRigVMTemplate* Template = Factory->GetTemplate();
@@ -1072,6 +1071,11 @@ void FControlRigEditorModule::GetTypeActions(UControlRigBlueprint* CRB, FBluepri
 			continue;
 		}
 
+		if(!Function.SupportsExecuteContextStruct(FControlRigExecuteContext::StaticStruct()))
+		{
+			continue;
+		}
+
 		// skip rig units which have a template
 		if (Function.GetTemplate())
 		{
@@ -1101,7 +1105,7 @@ void FControlRigEditorModule::GetTypeActions(UControlRigBlueprint* CRB, FBluepri
 		FText MenuDesc = FText::FromString(DisplayNameMetadata + MenuDescSuffixMetadata);
 		FText ToolTip = Struct->GetToolTipText();
 
-		UBlueprintNodeSpawner* NodeSpawner = UControlRigUnitNodeSpawner::CreateFromStruct(Struct, MenuDesc, NodeCategory, ToolTip);
+		UBlueprintNodeSpawner* NodeSpawner = UControlRigUnitNodeSpawner::CreateFromStruct(Struct, Function.GetMethodName(), MenuDesc, NodeCategory, ToolTip);
 		check(NodeSpawner != nullptr);
 		ActionRegistrar.AddBlueprintAction(ActionKey, NodeSpawner);
 	};
@@ -1881,11 +1885,11 @@ void FControlRigEditorModule::GetContextMenuActions(const UControlRigGraphSchema
 								StructOnScope = UnitNode->ConstructStructInstance(false /* default */);
 								StructMemory = (FRigUnit*)StructOnScope->GetStructMemory();
 
-								FRigUnitContext RigUnitContext;
-								RigUnitContext.Hierarchy = TemporaryHierarchy;
-								RigUnitContext.State = EControlRigState::Update;
+								//FRigUnitContext RigUnitContext;
+								//RigUnitContext.Hierarchy = TemporaryHierarchy;
+								//RigUnitContext.State = EControlRigState::Update;
 								
-								StructMemory->Execute(RigUnitContext);
+								StructMemory->Execute();
 							}
 						}
 
diff --git a/Engine/Plugins/Animation/ControlRig/Source/ControlRigEditor/Private/Graph/NodeSpawners/ControlRigUnitNodeSpawner.cpp b/Engine/Plugins/Animation/ControlRig/Source/ControlRigEditor/Private/Graph/NodeSpawners/ControlRigUnitNodeSpawner.cpp
--- a/Engine/Plugins/Animation/ControlRig/Source/ControlRigEditor/Private/Graph/NodeSpawners/ControlRigUnitNodeSpawner.cpp
+++ b/Engine/Plugins/Animation/ControlRig/Source/ControlRigEditor/Private/Graph/NodeSpawners/ControlRigUnitNodeSpawner.cpp
@@ -32,10 +32,11 @@
 
 #define LOCTEXT_NAMESPACE "ControlRigUnitNodeSpawner"
 
-UControlRigUnitNodeSpawner* UControlRigUnitNodeSpawner::CreateFromStruct(UScriptStruct* InStruct, const FText& InMenuDesc, const FText& InCategory, const FText& InTooltip)
+UControlRigUnitNodeSpawner* UControlRigUnitNodeSpawner::CreateFromStruct(UScriptStruct* InStruct, const FName& InMethodName, const FText& InMenuDesc, const FText& InCategory, const FText& InTooltip)
 {
 	UControlRigUnitNodeSpawner* NodeSpawner = NewObject<UControlRigUnitNodeSpawner>(GetTransientPackage());
 	NodeSpawner->StructTemplate = InStruct;
+	NodeSpawner->MethodName = InMethodName;
 	NodeSpawner->NodeClass = UControlRigGraphNode::StaticClass();
 
 	FBlueprintActionUiSpec& MenuSignature = NodeSpawner->DefaultMenuSignature;
@@ -154,13 +155,13 @@ UEdGraphNode* UControlRigUnitNodeSpawner::Invoke(UEdGraph* ParentGraph, FBinding
 #endif
 
 		UBlueprint* Blueprint = FBlueprintEditorUtils::FindBlueprintForGraph(ParentGraph);
-		NewNode = SpawnNode(ParentGraph, Blueprint, StructTemplate, Location);
+		NewNode = SpawnNode(ParentGraph, Blueprint, StructTemplate, MethodName, Location);
 	}
 
 	return NewNode;
 }
 
-UControlRigGraphNode* UControlRigUnitNodeSpawner::SpawnNode(UEdGraph* ParentGraph, UBlueprint* Blueprint, UScriptStruct* StructTemplate, FVector2D const Location)
+UControlRigGraphNode* UControlRigUnitNodeSpawner::SpawnNode(UEdGraph* ParentGraph, UBlueprint* Blueprint, UScriptStruct* StructTemplate, const FName& InMethodName, FVector2D const Location)
 {
 	UControlRigGraphNode* NewNode = nullptr;
 	UControlRigBlueprint* RigBlueprint = Cast<UControlRigBlueprint>(Blueprint);
@@ -218,7 +219,7 @@ UControlRigGraphNode* UControlRigUnitNodeSpawner::SpawnNode(UEdGraph* ParentGrap
 		FRigVMUnitNodeCreatedContext& UnitNodeCreatedContext = Controller->GetUnitNodeCreatedContext();
 		FRigVMUnitNodeCreatedContext::FScope ReasonScope(UnitNodeCreatedContext, ERigVMNodeCreatedReason::NodeSpawner);
 
-		if (URigVMUnitNode* ModelNode = Controller->AddUnitNode(StructTemplate, FRigUnit::GetMethodName(), Location, Name.ToString(), bIsUserFacingNode, !bIsTemplateNode))
+		if (URigVMUnitNode* ModelNode = Controller->AddUnitNode(StructTemplate, InMethodName, Location, Name.ToString(), bIsUserFacingNode, !bIsTemplateNode))
 		{
 			NewNode = Cast<UControlRigGraphNode>(RigGraph->FindNodeForModelNodeName(ModelNode->GetFName()));
 			check(NewNode);
diff --git a/Engine/Plugins/Animation/ControlRig/Source/ControlRigEditor/Private/Graph/NodeSpawners/ControlRigUnitNodeSpawner.h b/Engine/Plugins/Animation/ControlRig/Source/ControlRigEditor/Private/Graph/NodeSpawners/ControlRigUnitNodeSpawner.h
--- a/Engine/Plugins/Animation/ControlRig/Source/ControlRigEditor/Private/Graph/NodeSpawners/ControlRigUnitNodeSpawner.h
+++ b/Engine/Plugins/Animation/ControlRig/Source/ControlRigEditor/Private/Graph/NodeSpawners/ControlRigUnitNodeSpawner.h
@@ -32,7 +32,7 @@ public:
 	 * 
 	 * @return A newly allocated instance of this class.
 	 */
-	static UControlRigUnitNodeSpawner* CreateFromStruct(UScriptStruct* InStruct, const FText& InMenuDesc, const FText& InCategory, const FText& InTooltip);
+	static UControlRigUnitNodeSpawner* CreateFromStruct(UScriptStruct* InStruct, const FName& InMethodName, const FText& InMenuDesc, const FText& InCategory, const FText& InTooltip);
 
 	// UBlueprintNodeSpawner interface
 	virtual void Prime() override;
@@ -47,7 +47,10 @@ private:
 	UPROPERTY(Transient)
 	TObjectPtr<UScriptStruct> StructTemplate;
 
-	static UControlRigGraphNode* SpawnNode(UEdGraph* ParentGraph, UBlueprint* Blueprint, UScriptStruct* StructTemplate, FVector2D const Location);
+	UPROPERTY(Transient)
+	FName MethodName;
+
+	static UControlRigGraphNode* SpawnNode(UEdGraph* ParentGraph, UBlueprint* Blueprint, UScriptStruct* StructTemplate, const FName& InMethodName, FVector2D const Location);
 
 	friend class UEngineTestControlRig;
 	friend class FControlRigEditor;
diff --git a/Engine/Plugins/Animation/ControlRigSpline/Source/ControlRigSpline/Private/ControlRigSplineUnits.cpp b/Engine/Plugins/Animation/ControlRigSpline/Source/ControlRigSpline/Private/ControlRigSplineUnits.cpp
--- a/Engine/Plugins/Animation/ControlRigSpline/Source/ControlRigSpline/Private/ControlRigSplineUnits.cpp
+++ b/Engine/Plugins/Animation/ControlRigSpline/Source/ControlRigSpline/Private/ControlRigSplineUnits.cpp
@@ -12,7 +12,7 @@
 
 FRigUnit_ControlRigSplineFromPoints_Execute()
 {
-	switch (Context.State)
+	switch (ExecuteContext.UnitContext.State)
 	{
 		case EControlRigState::Init:
 		case EControlRigState::Update:
@@ -40,7 +40,7 @@ FRigUnit_ControlRigSplineFromPoints_Execute()
 
 FRigUnit_SetSplinePoints_Execute()
 {
-	if (Context.State == EControlRigState::Init)
+	if (ExecuteContext.UnitContext.State == EControlRigState::Init)
 	{
 		return;
 	}
@@ -65,7 +65,7 @@ FRigUnit_SetSplinePoints_Execute()
 		return;
 	}
 
-	switch (Context.State)
+	switch (ExecuteContext.UnitContext.State)
 	{
 		case EControlRigState::Init:
 		case EControlRigState::Update:
@@ -97,7 +97,7 @@ FRigUnit_PositionFromControlRigSpline_Execute()
 	}
 #endif
 	
-	switch (Context.State)
+	switch (ExecuteContext.UnitContext.State)
 	{
 		case EControlRigState::Init:
 		case EControlRigState::Update:
@@ -128,7 +128,7 @@ FRigUnit_TransformFromControlRigSpline_Execute()
 	}
 #endif	
 	
-	switch (Context.State)
+	switch (ExecuteContext.UnitContext.State)
 	{
 		case EControlRigState::Init:
 		case EControlRigState::Update:
@@ -178,7 +178,7 @@ FRigUnit_TangentFromControlRigSpline_Execute()
 	}
 #endif	
 	
-	switch (Context.State)
+	switch (ExecuteContext.UnitContext.State)
 	{
 		case EControlRigState::Init:
 		case EControlRigState::Update:
@@ -209,12 +209,12 @@ FRigUnit_TangentFromControlRigSpline_Execute()
 
 FRigUnit_DrawControlRigSpline_Execute()
 {
-	if (Context.State == EControlRigState::Init)
+	if (ExecuteContext.UnitContext.State == EControlRigState::Init)
 	{
 		return;
 	}
 
-	if (Context.DrawInterface == nullptr)
+	if (ExecuteContext.UnitContext.DrawInterface == nullptr)
 	{
 		return;
 	}
@@ -245,7 +245,7 @@ FRigUnit_DrawControlRigSpline_Execute()
 		T += Step;
 	}
 
-	Context.DrawInterface->Instructions.Add(Instruction);
+	ExecuteContext.UnitContext.DrawInterface->Instructions.Add(Instruction);
 }
 
 FRigUnit_GetLengthControlRigSpline_Execute()
@@ -265,7 +265,7 @@ FRigUnit_GetLengthControlRigSpline_Execute()
 	}
 #endif	
 	
-	switch (Context.State)
+	switch (ExecuteContext.UnitContext.State)
 	{
 		case EControlRigState::Init:
 		case EControlRigState::Update:
@@ -299,8 +299,7 @@ FRigUnit_FitChainToSplineCurve_Execute()
 		Weight,
 		bPropagateToChildren,
 		DebugSettings,
-		WorkData,
-		Context);
+		WorkData);
 }
 
 FRigVMStructUpgradeInfo FRigUnit_FitChainToSplineCurve::GetUpgradeInfo() const
@@ -359,7 +358,7 @@ FRigUnit_FitChainToSplineCurveItemArray_Execute()
 	TArray<float>& ItemRotationT = WorkData.ItemRotationT;
 	TArray<FTransform>& ItemLocalTransforms = WorkData.ItemLocalTransforms;
 
-	if (Context.State == EControlRigState::Init)
+	if (ExecuteContext.UnitContext.State == EControlRigState::Init)
 	{
 		CachedItems.Reset();
 		return;
@@ -714,7 +713,7 @@ FRigUnit_FitChainToSplineCurveItemArray_Execute()
 		}
 	}
 
-	if (Context.DrawInterface != nullptr && DebugSettings.bEnabled)
+	if (ExecuteContext.UnitContext.DrawInterface != nullptr && DebugSettings.bEnabled)
 	{
 		int32 Count = 64;
 		FControlRigDrawInstruction Instruction(EControlRigDrawSettings::LineStrip, DebugSettings.CurveColor, DebugSettings.Scale, DebugSettings.WorldOffset);
@@ -728,21 +727,21 @@ FRigUnit_FitChainToSplineCurveItemArray_Execute()
 			Instruction.Positions[Index] = Spline.PositionAtParam(T);
 			T += Step;
 		}
-		Context.DrawInterface->Instructions.Add(Instruction);
+		ExecuteContext.UnitContext.DrawInterface->Instructions.Add(Instruction);
 
 		for (auto Point : Spline.SplineData->GetControlPoints())
 		{
-			Context.DrawInterface->DrawPoint(DebugSettings.WorldOffset, Point, DebugSettings.Scale * 6, DebugSettings.CurveColor);
+			ExecuteContext.UnitContext.DrawInterface->DrawPoint(DebugSettings.WorldOffset, Point, DebugSettings.Scale * 6, DebugSettings.CurveColor);
 		}
 
-		Context.DrawInterface->DrawLineStrip(DebugSettings.WorldOffset, CurvePositions, DebugSettings.SegmentsColor, DebugSettings.Scale);
-		Context.DrawInterface->DrawPoints(DebugSettings.WorldOffset, CurvePositions, DebugSettings.Scale * 4.f, DebugSettings.SegmentsColor);
+		ExecuteContext.UnitContext.DrawInterface->DrawLineStrip(DebugSettings.WorldOffset, CurvePositions, DebugSettings.SegmentsColor, DebugSettings.Scale);
+		ExecuteContext.UnitContext.DrawInterface->DrawPoints(DebugSettings.WorldOffset, CurvePositions, DebugSettings.Scale * 4.f, DebugSettings.SegmentsColor);
 	}
 }
 
 FRigUnit_FitSplineCurveToChain_Execute()
 {
-	FRigUnit_FitSplineCurveToChainItemArray::StaticExecute(ExecuteContext, Items.Keys, Spline, Context);
+	FRigUnit_FitSplineCurveToChainItemArray::StaticExecute(ExecuteContext, Items.Keys, Spline);
 }
 
 FRigVMStructUpgradeInfo FRigUnit_FitSplineCurveToChain::GetUpgradeInfo() const
@@ -777,7 +776,7 @@ FRigUnit_FitSplineCurveToChainItemArray_Execute()
 	}
 #endif	
 
-	if (Context.State == EControlRigState::Init)
+	if (ExecuteContext.UnitContext.State == EControlRigState::Init)
 	{
 		return;
 	}
@@ -833,7 +832,7 @@ FRigUnit_ClosestParameterFromControlRigSpline_Execute()
 	}
 #endif	
 	
-	switch (Context.State)
+	switch (ExecuteContext.UnitContext.State)
 	{
 		case EControlRigState::Init:
 		case EControlRigState::Update:
@@ -905,7 +904,7 @@ FRigUnit_ParameterAtPercentage_Execute()
 	}
 #endif	
 	
-	switch (Context.State)
+	switch (ExecuteContext.UnitContext.State)
 	{
 	case EControlRigState::Init:
 	case EControlRigState::Update:
diff --git a/Engine/Plugins/Animation/ControlRigSpline/Source/ControlRigSpline/Public/ControlRigSplineUnits.h b/Engine/Plugins/Animation/ControlRigSpline/Source/ControlRigSpline/Public/ControlRigSplineUnits.h
--- a/Engine/Plugins/Animation/ControlRigSpline/Source/ControlRigSpline/Public/ControlRigSplineUnits.h
+++ b/Engine/Plugins/Animation/ControlRigSpline/Source/ControlRigSpline/Public/ControlRigSplineUnits.h
@@ -32,7 +32,7 @@ struct CONTROLRIGSPLINE_API FRigUnit_ControlRigSplineFromPoints : public FRigUni
 
 	/** Execute logic for this rig unit */
 	RIGVM_METHOD()
-	virtual void Execute(const FRigUnitContext& Context) override;
+	virtual void Execute() override;
 
 	UPROPERTY(meta = (Input))
 	TArray<FVector> Points;
@@ -68,7 +68,7 @@ struct CONTROLRIGSPLINE_API FRigUnit_SetSplinePoints : public FRigUnitMutable
 
 	/** Execute logic for this rig unit */
 	RIGVM_METHOD()
-	virtual void Execute(const FRigUnitContext& Context) override;
+	virtual void Execute() override;
 
 	UPROPERTY(meta = (Input))
 	TArray<FVector> Points;
@@ -93,7 +93,7 @@ struct CONTROLRIGSPLINE_API FRigUnit_PositionFromControlRigSpline : public FRigU
 
 	/** Execute logic for this rig unit */
 	RIGVM_METHOD()
-	virtual void Execute(const FRigUnitContext& Context) override;
+	virtual void Execute() override;
 
 	UPROPERTY(meta = (Input))
 	FControlRigSpline Spline;
@@ -123,7 +123,7 @@ struct CONTROLRIGSPLINE_API FRigUnit_TransformFromControlRigSpline : public FRig
 
 	/** Execute logic for this rig unit */
 	RIGVM_METHOD()
-	virtual void Execute(const FRigUnitContext& Context) override;
+	virtual void Execute() override;
 
 	UPROPERTY(meta = (Input))
 	FControlRigSpline Spline;
@@ -157,7 +157,7 @@ struct CONTROLRIGSPLINE_API FRigUnit_TangentFromControlRigSpline : public FRigUn
 
 	/** Execute logic for this rig unit */
 	RIGVM_METHOD()
-	virtual void Execute(const FRigUnitContext& Context) override;
+	virtual void Execute() override;
 
 	UPROPERTY(meta = (Input))
 	FControlRigSpline Spline;
@@ -186,7 +186,7 @@ struct CONTROLRIGSPLINE_API FRigUnit_DrawControlRigSpline : public FRigUnitMutab
 
 	/** Execute logic for this rig unit */
 	RIGVM_METHOD()
-	virtual void Execute(const FRigUnitContext& Context) override;
+	virtual void Execute() override;
 
 	UPROPERTY(meta = (Input))
 	FControlRigSpline Spline;
@@ -216,7 +216,7 @@ struct CONTROLRIGSPLINE_API FRigUnit_GetLengthControlRigSpline : public FRigUnit
 
 	/** Execute logic for this rig unit */
 	RIGVM_METHOD()
-	virtual void Execute(const FRigUnitContext& Context) override;
+	virtual void Execute() override;
 
 	UPROPERTY(meta = (Input))
 	FControlRigSpline Spline;
@@ -250,7 +250,7 @@ struct CONTROLRIGSPLINE_API FRigUnit_FitChainToSplineCurve : public FRigUnit_Hig
 	}
 
 	RIGVM_METHOD()
-	virtual void Execute(const FRigUnitContext& Context) override;
+	virtual void Execute() override;
 
 	/** 
 	 * The items to align
@@ -369,7 +369,7 @@ struct CONTROLRIGSPLINE_API FRigUnit_FitChainToSplineCurveItemArray : public FRi
 	}
 
 	RIGVM_METHOD()
-	virtual void Execute(const FRigUnitContext& Context) override;
+	virtual void Execute() override;
 
 	/** 
 	 * The items to align
@@ -474,7 +474,7 @@ struct CONTROLRIGSPLINE_API FRigUnit_FitSplineCurveToChain : public FRigUnit_Hig
 	}
 
 	RIGVM_METHOD()
-	virtual void Execute(const FRigUnitContext& Context) override;
+	virtual void Execute() override;
 
 	/** 
 	 * The items to align to
@@ -506,7 +506,7 @@ struct CONTROLRIGSPLINE_API FRigUnit_FitSplineCurveToChainItemArray : public FRi
 	}
 
 	RIGVM_METHOD()
-	virtual void Execute(const FRigUnitContext& Context) override;
+	virtual void Execute() override;
 
 	/** 
 	 * The items to align to
@@ -537,7 +537,7 @@ struct CONTROLRIGSPLINE_API FRigUnit_ClosestParameterFromControlRigSpline : publ
 
 	/** Execute logic for this rig unit */
 	RIGVM_METHOD()
-	virtual void Execute(const FRigUnitContext& Context) override;
+	virtual void Execute() override;
 
 	UPROPERTY(meta = (Input))
 	FControlRigSpline Spline;
@@ -566,7 +566,7 @@ struct CONTROLRIGSPLINE_API FRigUnit_ParameterAtPercentage : public FRigUnit_Con
 
 	/** Execute logic for this rig unit */
 	RIGVM_METHOD()
-	virtual void Execute(const FRigUnitContext& Context) override;
+	virtual void Execute() override;
 
 	UPROPERTY(meta = (Input))
 	FControlRigSpline Spline;
diff --git a/Engine/Plugins/Animation/RigLogic/Source/RigLogicModule/Private/RigUnit_RigLogic.cpp b/Engine/Plugins/Animation/RigLogic/Source/RigLogicModule/Private/RigUnit_RigLogic.cpp
--- a/Engine/Plugins/Animation/RigLogic/Source/RigLogicModule/Private/RigUnit_RigLogic.cpp
+++ b/Engine/Plugins/Animation/RigLogic/Source/RigLogicModule/Private/RigUnit_RigLogic.cpp
@@ -336,14 +336,14 @@ FRigUnit_RigLogic_Execute()
  	URigHierarchy* Hierarchy = ExecuteContext.Hierarchy;
 	if (Hierarchy)
 	{
-		switch (Context.State)
+		switch (ExecuteContext.UnitContext.State)
 		{
 			case EControlRigState::Init:
 			{
 				//const double startTime = FPlatformTime::Seconds();
 				if (!Data.SkelMeshComponent.IsValid())
 				{
-					Data.SkelMeshComponent = Context.DataSourceRegistry->RequestSource<USkeletalMeshComponent>(UControlRig::OwnerComponent);
+					Data.SkelMeshComponent = ExecuteContext.UnitContext.DataSourceRegistry->RequestSource<USkeletalMeshComponent>(UControlRig::OwnerComponent);
 					// In normal execution, Data.SkelMeshComponent will be nullptr at the beginning
 					// however, during unit testing we cannot fetch it from DataSourceRegistry 
 					// in that case, a mock version will be inserted into Data by unit test beforehand
diff --git a/Engine/Plugins/Animation/RigLogic/Source/RigLogicModule/Public/RigUnit_RigLogic.h b/Engine/Plugins/Animation/RigLogic/Source/RigLogicModule/Public/RigUnit_RigLogic.h
--- a/Engine/Plugins/Animation/RigLogic/Source/RigLogicModule/Public/RigUnit_RigLogic.h
+++ b/Engine/Plugins/Animation/RigLogic/Source/RigLogicModule/Public/RigUnit_RigLogic.h
@@ -128,7 +128,7 @@ public:
 #endif
 
 	RIGVM_METHOD()
-	void Execute(const FRigUnitContext& Context) override;
+	virtual void Execute() override;
 
 private:
 	static TSharedPtr<FSharedRigRuntimeContext> GetSharedRigRuntimeContext(USkeletalMesh* SkelMesh);
diff --git a/Engine/Plugins/Experimental/DataInterfaceGraph/Source/DataInterfaceGraph/Internal/DataInterfaceExecuteContext.h b/Engine/Plugins/Experimental/DataInterfaceGraph/Source/DataInterfaceGraph/Internal/DataInterfaceExecuteContext.h
--- a/Engine/Plugins/Experimental/DataInterfaceGraph/Source/DataInterfaceGraph/Internal/DataInterfaceExecuteContext.h
+++ b/Engine/Plugins/Experimental/DataInterfaceGraph/Source/DataInterfaceGraph/Internal/DataInterfaceExecuteContext.h
@@ -4,6 +4,8 @@
 
 #include "CoreMinimal.h"
 #include "RigVMCore/RigVMExecuteContext.h"
+#include "RigVMCore/RigVMStruct.h"
+#include "DataInterfaceUnitContext.h"
 #include "DataInterfaceExecuteContext.generated.h"
 
 class IDataInterface;
@@ -18,7 +20,14 @@ struct FDataInterfaceExecuteContext : public FRigVMExecuteContext
 {
 	GENERATED_BODY()
 
-	FDataInterfaceExecuteContext() = default;
+	FDataInterfaceExecuteContext()
+		: FRigVMExecuteContext()
+	, DataInterfaceContext(nullptr)
+	, Interface(nullptr)
+	, ResultPtr(nullptr)
+	, UnitContext()
+	{
+	}
 
 	const UE::DataInterface::FContext& GetContext() const
 	{
@@ -26,6 +35,11 @@ struct FDataInterfaceExecuteContext : public FRigVMExecuteContext
 		return *DataInterfaceContext;
 	}
 
+	const FDataInterfaceUnitContext& GetUnitContext() const
+	{
+		return UnitContext;
+	}
+
 	void SetResult(bool bInResult) const
 	{
 		check(ResultPtr);
@@ -50,7 +64,14 @@ struct FDataInterfaceExecuteContext : public FRigVMExecuteContext
 
 
 private:
-	const UE::DataInterface::FContext* DataInterfaceContext = nullptr;
-	const IDataInterface* Interface = nullptr;
-	bool* ResultPtr = nullptr;
-};
\ No newline at end of file
+	const UE::DataInterface::FContext* DataInterfaceContext;
+	const IDataInterface* Interface;
+	bool* ResultPtr;
+	FDataInterfaceUnitContext UnitContext;
+};
+
+USTRUCT(meta=(ExecuteContext="FDataInterfaceExecuteContext"))
+struct FRigUnit_DataInterfaceBase : public FRigVMStruct
+{
+	GENERATED_BODY()
+};
diff --git a/Engine/Plugins/Experimental/DataInterfaceGraph/Source/DataInterfaceGraph/Internal/DataInterfaceUnitContext.h b/Engine/Plugins/Experimental/DataInterfaceGraph/Source/DataInterfaceGraph/Internal/DataInterfaceUnitContext.h
--- a/Engine/Plugins/Experimental/DataInterfaceGraph/Source/DataInterfaceGraph/Internal/DataInterfaceUnitContext.h
+++ b/Engine/Plugins/Experimental/DataInterfaceGraph/Source/DataInterfaceGraph/Internal/DataInterfaceUnitContext.h
@@ -8,13 +8,21 @@
 
 struct FDataInterfaceUnitContext : public FRigUnitContext
 {
+	FDataInterfaceUnitContext()
+		: FRigUnitContext()
+		, Interface(nullptr)
+		, DataInterfaceContext(nullptr)
+		, bResult(nullptr)
+	{}
+
 	FDataInterfaceUnitContext(const IDataInterface* InInterface, const UE::DataInterface::FContext& InDataInterfaceContext, bool& bInResult)
-		: Interface(InInterface)
-		, DataInterfaceContext(InDataInterfaceContext)
-		, bResult(bInResult)
+		: FRigUnitContext()
+		, Interface(InInterface)
+		, DataInterfaceContext(&InDataInterfaceContext)
+		, bResult(&bInResult)
 	{}
 
-	const IDataInterface* Interface = nullptr;
-	const UE::DataInterface::FContext& DataInterfaceContext;
-	bool& bResult;
+	const IDataInterface* Interface;
+	const UE::DataInterface::FContext* DataInterfaceContext;
+	bool* bResult;
 };
diff --git a/Engine/Plugins/Experimental/DataInterfaceGraph/Source/DataInterfaceGraph/Internal/RigUnit_DataInterfaceBeginExecution.h b/Engine/Plugins/Experimental/DataInterfaceGraph/Source/DataInterfaceGraph/Internal/RigUnit_DataInterfaceBeginExecution.h
--- a/Engine/Plugins/Experimental/DataInterfaceGraph/Source/DataInterfaceGraph/Internal/RigUnit_DataInterfaceBeginExecution.h
+++ b/Engine/Plugins/Experimental/DataInterfaceGraph/Source/DataInterfaceGraph/Internal/RigUnit_DataInterfaceBeginExecution.h
@@ -10,12 +10,12 @@
  * Event for driving the skeleton hierarchy with variables and rig elements
  */
 USTRUCT(meta=(DisplayName="Execute Data Interface", Category="Events", NodeColor="1, 0, 0", Keywords="Begin,Update,Tick,Forward,Event"))
-struct DATAINTERFACEGRAPH_API FRigUnit_DataInterfaceBeginExecution : public FRigUnit
+struct DATAINTERFACEGRAPH_API FRigUnit_DataInterfaceBeginExecution : public FRigUnit_DataInterfaceBase
 {
 	GENERATED_BODY()
 
 	RIGVM_METHOD()
-	virtual void Execute(const FRigUnitContext& Context) override;
+	void Execute();
 
 	virtual FName GetEventName() const override { return EventName; }
 	virtual bool CanOnlyExistOnce() const override { return true; }
diff --git a/Engine/Plugins/Experimental/DataInterfaceGraph/Source/DataInterfaceGraph/Internal/RigUnit_DataInterfaceEndExecution.h b/Engine/Plugins/Experimental/DataInterfaceGraph/Source/DataInterfaceGraph/Internal/RigUnit_DataInterfaceEndExecution.h
--- a/Engine/Plugins/Experimental/DataInterfaceGraph/Source/DataInterfaceGraph/Internal/RigUnit_DataInterfaceEndExecution.h
+++ b/Engine/Plugins/Experimental/DataInterfaceGraph/Source/DataInterfaceGraph/Internal/RigUnit_DataInterfaceEndExecution.h
@@ -12,7 +12,7 @@ struct FRigUnitContext;
 
 /** Event for writing back calculated results to external variables */
 USTRUCT(meta=(DisplayName="End Execute Data Interface"))
-struct DATAINTERFACEGRAPH_API FRigUnit_DataInterfaceEndExecution : public FRigUnit
+struct DATAINTERFACEGRAPH_API FRigUnit_DataInterfaceEndExecution : public FRigUnit_DataInterfaceBase
 {
 	GENERATED_BODY()
 
@@ -36,11 +36,12 @@ struct DATAINTERFACEGRAPH_API FRigUnit_DataInterfaceEndExecution_Bool : public F
 	GENERATED_BODY()
 	
 	FRigUnit_DataInterfaceEndExecution_Bool()
-		: Result(false)
+		: FRigUnit_DataInterfaceEndExecution()
+		, Result(false)
 	{}
 	
 	RIGVM_METHOD()
-	virtual void Execute(const FRigUnitContext& Context) override;
+	void Execute();
 	
 	UPROPERTY(EditAnywhere, Category = Result, meta = (Input))
 	bool Result;
@@ -53,11 +54,12 @@ struct DATAINTERFACEGRAPH_API FRigUnit_DataInterfaceEndExecution_Float : public
 	GENERATED_BODY()
 
 	FRigUnit_DataInterfaceEndExecution_Float()
-	: Result(0.f)
+		: FRigUnit_DataInterfaceEndExecution()
+		, Result(0.f)
 	{}
 
 	RIGVM_METHOD()
-	virtual void Execute(const FRigUnitContext& Context) override;
+	void Execute();
 
 	UPROPERTY(EditAnywhere, Category = Result, meta = (Input))
 	float Result;
diff --git a/Engine/Plugins/Experimental/DataInterfaceGraph/Source/DataInterfaceGraph/Private/DataInterfaceGraph.cpp b/Engine/Plugins/Experimental/DataInterfaceGraph/Source/DataInterfaceGraph/Private/DataInterfaceGraph.cpp
--- a/Engine/Plugins/Experimental/DataInterfaceGraph/Source/DataInterfaceGraph/Private/DataInterfaceGraph.cpp
+++ b/Engine/Plugins/Experimental/DataInterfaceGraph/Source/DataInterfaceGraph/Private/DataInterfaceGraph.cpp
@@ -34,9 +34,8 @@ bool UDataInterfaceGraph::GetDataImpl(const UE::DataInterface::FContext& Context
 		FRigUnitContext RigUnitContext;
 		RigUnitContext.State = EControlRigState::Update;
 		FDataInterfaceUnitContext DataInterfaceUnitContext(this, Context, bResult);
-		void* AdditionalArguments[2] = { &RigUnitContext, &DataInterfaceUnitContext };
 		bResult &=
-			(RigVM->Execute(TArray<URigVMMemoryStorage*>(), AdditionalArguments, FRigUnit_DataInterfaceBeginExecution::EventName) != ERigVMExecuteResult::Failed);
+			(RigVM->Execute(TArray<URigVMMemoryStorage*>(), FRigUnit_DataInterfaceBeginExecution::EventName) != ERigVMExecuteResult::Failed);
 	}
 	
 	return bResult;
diff --git a/Engine/Plugins/Experimental/DataInterfaceGraph/Source/DataInterfaceGraph/Private/RigUnit_DataInterfaceEndExecution.cpp b/Engine/Plugins/Experimental/DataInterfaceGraph/Source/DataInterfaceGraph/Private/RigUnit_DataInterfaceEndExecution.cpp
--- a/Engine/Plugins/Experimental/DataInterfaceGraph/Source/DataInterfaceGraph/Private/RigUnit_DataInterfaceEndExecution.cpp
+++ b/Engine/Plugins/Experimental/DataInterfaceGraph/Source/DataInterfaceGraph/Private/RigUnit_DataInterfaceEndExecution.cpp
@@ -7,16 +7,10 @@
 
 FRigUnit_DataInterfaceEndExecution_Bool_Execute()
 {
-	if(Context.State == EControlRigState::Update)
-	{
-		SetResult(ExecuteContext, Result);
-	}
+	SetResult(ExecuteContext, Result);
 }
 
 FRigUnit_DataInterfaceEndExecution_Float_Execute()
 {
-	if(Context.State == EControlRigState::Update)
-	{
-		SetResult(ExecuteContext, Result);
-	}
+	SetResult(ExecuteContext, Result);
 }
diff --git a/Engine/Plugins/Experimental/DataInterfaceGraph/Source/DataInterfaceGraph/Private/RigUnit_DataInterfaceParameter.cpp b/Engine/Plugins/Experimental/DataInterfaceGraph/Source/DataInterfaceGraph/Private/RigUnit_DataInterfaceParameter.cpp
--- a/Engine/Plugins/Experimental/DataInterfaceGraph/Source/DataInterfaceGraph/Private/RigUnit_DataInterfaceParameter.cpp
+++ b/Engine/Plugins/Experimental/DataInterfaceGraph/Source/DataInterfaceGraph/Private/RigUnit_DataInterfaceParameter.cpp
@@ -16,18 +16,12 @@ bool FRigUnit_DataInterfaceParameter::GetParameterInternal(FName InName, const F
 
 FRigUnit_DataInterfaceParameter_Float_Execute()
 {
-	if(Context.State == EControlRigState::Update)
-	{
-		GetParameterInternal(Parameter, static_cast<const FDataInterfaceUnitContext&>(Context), &Result);
-	}
+	GetParameterInternal(Parameter, ExecuteContext.GetUnitContext(), &Result);
 }
 
 FRigUnit_DataInterfaceParameter_DataInterface_Execute()
 {
-	if(Context.State == EControlRigState::Update)
-	{
-		GetParameterInternal(Parameter, static_cast<const FDataInterfaceUnitContext&>(Context), &Result);
-	}
+	GetParameterInternal(Parameter, ExecuteContext.GetUnitContext(), &Result);
 }
 
 FRigUnit_DataInterface_Float_Execute()
@@ -35,18 +29,13 @@ FRigUnit_DataInterface_Float_Execute()
 	using namespace UE::DataInterface;
 
 	// @TODO: ensure that context arg is always present to avoid these checks here
-	// Could bind the RigVM to a specific execution context UStruct?
-	if(Context.State == EControlRigState::Update)
-	{
-		const FDataInterfaceExecuteContext& DataInterfaceExecuteContext = static_cast<const FDataInterfaceExecuteContext&>(ExecuteContext);
-		const FContext& DataInterfaceContext = DataInterfaceExecuteContext.GetContext();
+	const FContext& DataInterfaceContext = ExecuteContext.GetContext();
 
-		// Wrap the internal result we are going to be writing to 
-		TWrapParam<float> CallResult(&Result);
+	// Wrap the internal result we are going to be writing to 
+	TWrapParam<float> CallResult(&Result);
 
-		// Call the interface
-		DataInterfaceExecuteContext.SetResult(UE::DataInterface::GetDataSafe(DataInterface, DataInterfaceContext, CallResult));
-	}
+	// Call the interface
+	ExecuteContext.SetResult(UE::DataInterface::GetDataSafe(DataInterface, DataInterfaceContext, CallResult));
 }
 /*
 FRigUnit_DataInterface_Pose_Execute()
@@ -54,7 +43,6 @@ FRigUnit_DataInterface_Pose_Execute()
 	using namespace UE::DataInterface;
 
 	// @TODO: ensure that context arg is always present to avoid these checks here
-	// Could bind the RigVM to a specific execution context UStruct?
 	if(Context.State == EControlRigState::Update && ExecuteContext.OpaqueArguments.Num() > 1 && ExecuteContext.OpaqueArguments[1] != nullptr)
 	{
 		const FDataInterfaceUnitContext& DataInterfaceUnitContext = *static_cast<const FDataInterfaceUnitContext*>(ExecuteContext.OpaqueArguments[1]);
@@ -90,33 +78,28 @@ FRigUnit_TestFloatState_Execute()
 	using namespace UE::DataInterface;
 
 	// @TODO: ensure that context arg is always present to avoid these checks here
-	// Could bind the RigVM to a specific execution context UStruct?
-	if(Context.State == EControlRigState::Update)
-	{
-		const FDataInterfaceExecuteContext& DataInterfaceExecuteContext = static_cast<const FDataInterfaceExecuteContext&>(ExecuteContext);
-		const FContext& DataInterfaceContext = DataInterfaceExecuteContext.GetContext();
-
-		const TParam<FSpringDamperState> State = DataInterfaceContext.GetState<FSpringDamperState>(DataInterfaceExecuteContext.GetInterface(), 0);
-		const float DeltaTime = DataInterfaceContext.GetDeltaTime();
-
-		FKernel::Run(DataInterfaceContext,
-			[DeltaTime, &Result](FSpringDamperState& InOutState,
-						float InTargetValue,
-						float InTargetValueRate,
-						float InSmoothingTime,
-						float InDampingRatio)
-				{
-					FMath::SpringDamperSmoothing(
-						InOutState.Value,
-						InOutState.ValueRate,
-						InTargetValue,
-						InTargetValueRate,
-						DeltaTime,
-						InSmoothingTime,
-						InDampingRatio);
-
-					Result = InOutState.Value;
-				},
-				State, TargetValue, TargetValueRate, SmoothingTime, DampingRatio);
-	}
+	const FContext& DataInterfaceContext = ExecuteContext.GetContext();
+
+	const TParam<FSpringDamperState> State = DataInterfaceContext.GetState<FSpringDamperState>(ExecuteContext.GetInterface(), 0);
+	const float DeltaTime = DataInterfaceContext.GetDeltaTime();
+
+	FKernel::Run(DataInterfaceContext,
+		[DeltaTime, &Result](FSpringDamperState& InOutState,
+					float InTargetValue,
+					float InTargetValueRate,
+					float InSmoothingTime,
+					float InDampingRatio)
+			{
+				FMath::SpringDamperSmoothing(
+					InOutState.Value,
+					InOutState.ValueRate,
+					InTargetValue,
+					InTargetValueRate,
+					DeltaTime,
+					InSmoothingTime,
+					InDampingRatio);
+
+				Result = InOutState.Value;
+			},
+			State, TargetValue, TargetValueRate, SmoothingTime, DampingRatio);
 }
diff --git a/Engine/Plugins/Experimental/DataInterfaceGraph/Source/DataInterfaceGraph/Private/RigUnit_DataInterfaceParameter.h b/Engine/Plugins/Experimental/DataInterfaceGraph/Source/DataInterfaceGraph/Private/RigUnit_DataInterfaceParameter.h
--- a/Engine/Plugins/Experimental/DataInterfaceGraph/Source/DataInterfaceGraph/Private/RigUnit_DataInterfaceParameter.h
+++ b/Engine/Plugins/Experimental/DataInterfaceGraph/Source/DataInterfaceGraph/Private/RigUnit_DataInterfaceParameter.h
@@ -4,6 +4,7 @@
 
 #include "IDataInterface.h"
 #include "Units/RigUnit.h"
+#include "DataInterfaceGraph/Internal/DataInterfaceExecuteContext.h"
 #include "RigUnit_DataInterfaceParameter.generated.h"
 
 struct FDataInterfaceUnitContext;
@@ -11,7 +12,7 @@ class UDataInterfaceGraph;
 
 /** Unit for reading parameters from context */
 USTRUCT(meta=(DisplayName="Parameter"))
-struct FRigUnit_DataInterfaceParameter : public FRigUnit
+struct FRigUnit_DataInterfaceParameter : public FRigUnit_DataInterfaceBase
 {
 	GENERATED_BODY()
 
@@ -29,7 +30,7 @@ struct FRigUnit_DataInterfaceParameter_Float : public FRigUnit_DataInterfacePara
 	GENERATED_BODY()
 
 	RIGVM_METHOD()
-	virtual void Execute(const FRigUnitContext& Context) override;
+	void Execute();
 	
 protected:
 	UPROPERTY(EditAnywhere, Category = "Result", meta = (Output))
@@ -43,7 +44,7 @@ struct FRigUnit_DataInterfaceParameter_DataInterface : public FRigUnit_DataInter
 	GENERATED_BODY()
 
 	RIGVM_METHOD()
-	virtual void Execute(const FRigUnitContext& Context) override;
+	void Execute();
 	
 protected:
 	UPROPERTY(EditAnywhere, Category = "Result", meta = (Output))
@@ -142,7 +143,7 @@ struct FDataInterfaceParameter_DataInterface : public FDataInterfaceParameter
 
 /** Base unit for calling data interfaces from graphs */
 USTRUCT(meta=(DisplayName="Data Interface"))
-struct FRigUnit_DataInterface : public FRigUnit
+struct FRigUnit_DataInterface : public FRigUnit_DataInterfaceBase
 {
 	GENERATED_BODY()
 
@@ -158,7 +159,7 @@ struct FRigUnit_DataInterface_Float : public FRigUnit_DataInterface
 	GENERATED_BODY()
 
 	RIGVM_METHOD()
-	virtual void Execute(const FRigUnitContext& Context) override;
+	void Execute();
 	
 protected:
 	UPROPERTY(EditAnywhere, Category = "Data Interface", meta = (Input))
@@ -176,7 +177,7 @@ struct FRigUnit_DataInterface_Pose : public FRigUnit
 	GENERATED_BODY()
 
 	RIGVM_METHOD()
-	virtual void Execute(const FRigUnitContext& Context) override;
+	void Execute();
 	
 protected:
 	UPROPERTY(EditAnywhere, Category = "Pose Interface", meta = (Input))
@@ -191,7 +192,7 @@ protected:
 */
 /** Unit for getting a pose via a data interface */
 USTRUCT(meta=(DisplayName="Float Operator"))
-struct FRigUnit_FloatOperator : public FRigUnit
+struct FRigUnit_FloatOperator : public FRigUnit_DataInterfaceBase
 {
 	GENERATED_BODY()
 	
@@ -202,7 +203,7 @@ struct FRigUnit_FloatOperator : public FRigUnit
 	{}
 
 	RIGVM_METHOD()
-	virtual void Execute(const FRigUnitContext& Context) override;
+	void Execute();
 	
 protected:
 	UPROPERTY(EditAnywhere, Category = "Pose Interface", meta = (Input))
@@ -220,12 +221,12 @@ protected:
 
 /** Unit for getting a pose via a data interface */
 USTRUCT(meta=(DisplayName="Pose Operator"))
-struct FRigUnit_PoseOperator : public FRigUnit
+struct FRigUnit_PoseOperator : public FRigUnit_DataInterfaceBase
 {
 	GENERATED_BODY()
 
 	RIGVM_METHOD()
-	virtual void Execute(const FRigUnitContext& Context) override;
+	void Execute();
 	
 protected:
 	UPROPERTY(EditAnywhere, Category = "Pose Interface", meta = (Input))
@@ -238,7 +239,7 @@ protected:
 	FDataInterfaceParameters1 Parameters;
 
 	UPROPERTY(EditAnywhere, Transient, DisplayName = "Result", Category = "BeginExecution", meta = (Output))
-	FRigVMExecuteContext Result;
+	FDataInterfaceExecuteContext Result;
 };
 
 USTRUCT(BlueprintType)
@@ -258,12 +259,12 @@ struct FAnimSequenceParameters
 
 /** Unit for getting a pose via a data interface */
 USTRUCT(meta=(DisplayName="Anim Sequence"))
-struct FRigUnit_DataInterface_SequencePlayer : public FRigUnit
+struct FRigUnit_DataInterface_SequencePlayer : public FRigUnit_DataInterfaceBase
 {
 	GENERATED_BODY()
 
 	RIGVM_METHOD()
-	virtual void Execute(const FRigUnitContext& Context) override;
+	void Execute();
 	
 protected:
 	UPROPERTY(EditAnywhere, Category = "Data Interface", meta = (Input))
@@ -273,17 +274,17 @@ protected:
 	TScriptInterface<IDataInterface> Sequence = nullptr;
 	
 	UPROPERTY(EditAnywhere, Transient, DisplayName = "Result", Category = "BeginExecution", meta = (Output))
-	FRigVMExecuteContext Result;
+	FDataInterfaceExecuteContext Result;
 };
 
 /** Unit for getting a pose via a data interface */
 USTRUCT(meta=(DisplayName="Test Float State - Spring Damper Smoothing", Varying))
-struct FRigUnit_TestFloatState : public FRigUnit
+struct FRigUnit_TestFloatState : public FRigUnit_DataInterfaceBase
 {
 	GENERATED_BODY()
 
 	RIGVM_METHOD()
-	virtual void Execute(const FRigUnitContext& Context) override;
+	void Execute();
 	
 protected:
 	UPROPERTY(EditAnywhere, Category = "Parameters", meta=(Input))
diff --git a/Engine/Plugins/Experimental/DataInterfaceGraph/Source/DataInterfaceGraphUncookedOnly/Private/DataInterfaceGraphUncookedOnlyUtils.cpp b/Engine/Plugins/Experimental/DataInterfaceGraph/Source/DataInterfaceGraphUncookedOnly/Private/DataInterfaceGraphUncookedOnlyUtils.cpp
--- a/Engine/Plugins/Experimental/DataInterfaceGraph/Source/DataInterfaceGraphUncookedOnly/Private/DataInterfaceGraphUncookedOnlyUtils.cpp
+++ b/Engine/Plugins/Experimental/DataInterfaceGraph/Source/DataInterfaceGraphUncookedOnly/Private/DataInterfaceGraphUncookedOnlyUtils.cpp
@@ -45,17 +45,13 @@ void FUtils::Compile(UDataInterfaceGraph* InGraph)
 	void* InitContextPtr = &InitContext;
 	void* UpdateContextPtr = &UpdateContext;
 
-	TArray<FRigVMUserDataArray> UserData;
-	UserData.Add(FRigVMUserDataArray(&InitContextPtr, 1));
-	UserData.Add(FRigVMUserDataArray(&UpdateContextPtr, 1));
-
 	EditorData->CompileLog.Messages.Reset();
 	EditorData->CompileLog.NumErrors = EditorData->CompileLog.NumWarnings = 0;
 
 	URigVMCompiler* Compiler = URigVMCompiler::StaticClass()->GetDefaultObject<URigVMCompiler>();
 	Compiler->Settings = (EditorData->bCompileInDebugMode) ? FRigVMCompileSettings::Fast() : EditorData->VMCompileSettings;
 	URigVMController* RootController = EditorData->GetRigVMClient()->GetOrCreateController(EditorData->GetRigVMClient()->GetDefaultModel());
-	Compiler->Compile(EditorData->GetRigVMClient()->GetAllModels(false, false), RootController, InGraph->RigVM, InGraph->GetRigVMExternalVariables(), UserData, &EditorData->PinToOperandMap);
+	Compiler->Compile(EditorData->GetRigVMClient()->GetAllModels(false, false), RootController, InGraph->RigVM, InGraph->GetRigVMExternalVariables(), &EditorData->PinToOperandMap);
 
 	if (EditorData->bErrorsDuringCompilation)
 	{
diff --git a/Engine/Plugins/Experimental/FullBodyIK/Source/FullBodyIK/Private/RigUnit_FullbodyIK.cpp b/Engine/Plugins/Experimental/FullBodyIK/Source/FullBodyIK/Private/RigUnit_FullbodyIK.cpp
--- a/Engine/Plugins/Experimental/FullBodyIK/Source/FullBodyIK/Private/RigUnit_FullbodyIK.cpp
+++ b/Engine/Plugins/Experimental/FullBodyIK/Source/FullBodyIK/Private/RigUnit_FullbodyIK.cpp
@@ -215,7 +215,7 @@ FRigUnit_FullbodyIK_Execute()
 		AddEffectors(Hierarchy, Root, EffectorsView, LinkData, EffectorTargets, EffectorLinkIndices, LinkDataToHierarchyIndices, HierarchyToLinkDataMap, SolverProperty);
 	}
 
-	if (Context.State == EControlRigState::Update)
+	if (ExecuteContext.UnitContext.State == EControlRigState::Update)
 	{
 		DECLARE_SCOPE_HIERARCHICAL_COUNTER(TEXT("Update"))
 
@@ -348,7 +348,7 @@ FRigUnit_FullbodyIK_Execute()
 			///////////////////////////////////////////////////////////////////////////
 			// debug draw start
 			///////////////////////////////////////////////////////////////////////////
-			if (bDebugEnabled && Context.DrawInterface != nullptr)
+			if (bDebugEnabled && ExecuteContext.UnitContext.DrawInterface != nullptr)
 			{
 				const int32 DebugDataNum = DebugData.Num();
 				if (DebugData.Num() > 0)
@@ -397,12 +397,12 @@ FRigUnit_FullbodyIK_Execute()
 								if (Data.ParentLinkIndex != INDEX_NONE)
 								{
 									const FFBIKLinkData& ParentData = LocalLink[Data.ParentLinkIndex];
-									Context.DrawInterface->DrawLine(Offset, Data.GetPreviousTransform().GetLocation(), ParentData.GetPreviousTransform().GetLocation(), DrawColor, LineThickness);
+									ExecuteContext.UnitContext.DrawInterface->DrawLine(Offset, Data.GetPreviousTransform().GetLocation(), ParentData.GetPreviousTransform().GetLocation(), DrawColor, LineThickness);
 								}
 
 								if (DebugOption.bDrawDebugAxes)
 								{
-									Context.DrawInterface->DrawAxes(Offset, Data.GetPreviousTransform(), DebugOption.DrawSize);
+									ExecuteContext.UnitContext.DrawInterface->DrawAxes(Offset, Data.GetPreviousTransform(), DebugOption.DrawSize);
 								}
 							}
 						}
@@ -415,17 +415,17 @@ FRigUnit_FullbodyIK_Execute()
 								if (EffectorTarget.bPositionEnabled)
 								{
 									// draw effector target locations
-									Context.DrawInterface->DrawBox(Offset, FTransform(EffectorTarget.Position), FLinearColor::Yellow, DebugOption.DrawSize);
+									ExecuteContext.UnitContext.DrawInterface->DrawBox(Offset, FTransform(EffectorTarget.Position), FLinearColor::Yellow, DebugOption.DrawSize);
 								}
 
 								// draw effector link location
-								Context.DrawInterface->DrawBox(Offset, LocalLink[Iter.Key()].GetPreviousTransform(), FLinearColor::Green, DebugOption.DrawSize);
+								ExecuteContext.UnitContext.DrawInterface->DrawBox(Offset, LocalLink[Iter.Key()].GetPreviousTransform(), FLinearColor::Green, DebugOption.DrawSize);
 							}
 
 							for (int32 Index = 0; Index < DebugData[DebugIndex].TargetVectorSources.Num(); ++Index)
 							{
 								// draw arrow to the target
-								Context.DrawInterface->DrawLine(Offset, DebugData[DebugIndex].TargetVectorSources[Index].GetLocation(), 
+								ExecuteContext.UnitContext.DrawInterface->DrawLine(Offset, DebugData[DebugIndex].TargetVectorSources[Index].GetLocation(), 
 									DebugData[DebugIndex].TargetVectorSources[Index].GetLocation() + DebugData[DebugIndex].TargetVectors[Index], FLinearColor::Red);
 							}
 						}
@@ -448,7 +448,7 @@ FRigUnit_FullbodyIK_Execute()
 								{
 									FTransform ConstraintFrame = LinkData[*Found].GetTransform();
 									ConstraintFrame.ConcatenateRotation(FQuat(Constraints[Index].OffsetRotation));
-									Context.DrawInterface->DrawAxes(Offset, ConstraintFrame, 2.f);
+									ExecuteContext.UnitContext.DrawInterface->DrawAxes(Offset, ConstraintFrame, 2.f);
 								}
 							}
 						}
@@ -465,7 +465,7 @@ FRigUnit_FullbodyIK_Execute()
 							// base is parent transform but in their space, we can get there by inversing local ref rotation
 							FTransform BaseTransform = FTransform(LocalRefRotation).GetRelativeTransformReverse(RotationTransform);
 							BaseTransform.ConcatenateRotation(LimitConstraint.BaseFrameOffset);
-							Context.DrawInterface->DrawAxes(Offset, BaseTransform, 5.f, 1.f);
+							ExecuteContext.UnitContext.DrawInterface->DrawAxes(Offset, BaseTransform, 5.f, 1.f);
 
 							// current transform
 							const FQuat LocalRotation = BaseTransform.GetRotation().Inverse() * RotationTransform.GetRotation();
@@ -475,7 +475,7 @@ FRigUnit_FullbodyIK_Execute()
 							RotationTransform.SetLocation(BaseTransform.GetLocation());
 
 							// draw ref pose on their current transform
-							Context.DrawInterface->DrawAxes(Offset, RotationTransform, 10.f, 1.f);
+							ExecuteContext.UnitContext.DrawInterface->DrawAxes(Offset, RotationTransform, 10.f, 1.f);
 
 							FVector XAxis = BaseTransform.GetUnitAxis(EAxis::X);
 							FVector YAxis = BaseTransform.GetUnitAxis(EAxis::Y);
@@ -486,7 +486,7 @@ FRigUnit_FullbodyIK_Execute()
 								FTransform XAxisConeTM(YAxis, XAxis ^ YAxis, XAxis, BaseTransform.GetTranslation());
 								XAxisConeTM.SetRotation(FQuat(XAxis, 0.f) * XAxisConeTM.GetRotation());
 								XAxisConeTM.SetScale3D(FVector(30.f));
-								Context.DrawInterface->DrawCone(Offset, XAxisConeTM, LimitConstraint.Limit.X, 0.0f, 24, false, FLinearColor::Red, GEngine->ConstraintLimitMaterialX->GetRenderProxy());
+								ExecuteContext.UnitContext.DrawInterface->DrawCone(Offset, XAxisConeTM, LimitConstraint.Limit.X, 0.0f, 24, false, FLinearColor::Red, GEngine->ConstraintLimitMaterialX->GetRenderProxy());
 							}
 
 							if (LimitConstraint.bYLimitSet)
@@ -494,7 +494,7 @@ FRigUnit_FullbodyIK_Execute()
 								FTransform YAxisConeTM(ZAxis, YAxis ^ ZAxis, YAxis, BaseTransform.GetTranslation());
 								YAxisConeTM.SetRotation(FQuat(YAxis, 0.f) * YAxisConeTM.GetRotation());
 								YAxisConeTM.SetScale3D(FVector(30.f));
-								Context.DrawInterface->DrawCone(Offset, YAxisConeTM, LimitConstraint.Limit.Y, 0.0f, 24, false, FLinearColor::Green, GEngine->ConstraintLimitMaterialY->GetRenderProxy());
+								ExecuteContext.UnitContext.DrawInterface->DrawCone(Offset, YAxisConeTM, LimitConstraint.Limit.Y, 0.0f, 24, false, FLinearColor::Green, GEngine->ConstraintLimitMaterialY->GetRenderProxy());
 							}
 
 							if (LimitConstraint.bZLimitSet)
@@ -502,7 +502,7 @@ FRigUnit_FullbodyIK_Execute()
 								FTransform ZAxisConeTM(XAxis, ZAxis ^ XAxis, ZAxis, BaseTransform.GetTranslation());
 								ZAxisConeTM.SetRotation(FQuat(ZAxis, 0.f) * ZAxisConeTM.GetRotation());
 								ZAxisConeTM.SetScale3D(FVector(30.f));
-								Context.DrawInterface->DrawCone(Offset, ZAxisConeTM, LimitConstraint.Limit.Z, 0.0f, 24, false, FLinearColor::Blue, GEngine->ConstraintLimitMaterialZ->GetRenderProxy());
+								ExecuteContext.UnitContext.DrawInterface->DrawCone(Offset, ZAxisConeTM, LimitConstraint.Limit.Z, 0.0f, 24, false, FLinearColor::Blue, GEngine->ConstraintLimitMaterialZ->GetRenderProxy());
 							}
 						}
 
@@ -526,8 +526,8 @@ FRigUnit_FullbodyIK_Execute()
 							Positions.Add(JointTarget);
 							Positions.Add(RootTransform.GetLocation());
 
-							Context.DrawInterface->DrawLines(Offset, Positions, FLinearColor::Gray, 1.2f);
-							Context.DrawInterface->DrawLine(Offset, JointTransform.GetLocation(), JointTarget, FLinearColor::Red, 1.2f);
+							ExecuteContext.UnitContext.DrawInterface->DrawLines(Offset, Positions, FLinearColor::Gray, 1.2f);
+							ExecuteContext.UnitContext.DrawInterface->DrawLine(Offset, JointTransform.GetLocation(), JointTarget, FLinearColor::Red, 1.2f);
 						}
 					}
 				}
diff --git a/Engine/Plugins/Experimental/FullBodyIK/Source/PBIK/Private/RigUnit_PBIK.cpp b/Engine/Plugins/Experimental/FullBodyIK/Source/PBIK/Private/RigUnit_PBIK.cpp
--- a/Engine/Plugins/Experimental/FullBodyIK/Source/PBIK/Private/RigUnit_PBIK.cpp
+++ b/Engine/Plugins/Experimental/FullBodyIK/Source/PBIK/Private/RigUnit_PBIK.cpp
@@ -10,7 +10,7 @@ FRigUnit_PBIK_Execute()
 	DECLARE_SCOPE_HIERARCHICAL_COUNTER_RIGUNIT()
 	LLM_SCOPE_BYNAME(TEXT("Animation/FBIK"));
 
-	if (Context.State == EControlRigState::Init)
+	if (ExecuteContext.UnitContext.State == EControlRigState::Init)
 	{
 		BoneSettingToSolverBoneIndex.Reset();
 		Solver.Reset();
@@ -21,7 +21,7 @@ FRigUnit_PBIK_Execute()
 	}
 
 	// only updates from here on...
-	if (Context.State != EControlRigState::Update)
+	if (ExecuteContext.UnitContext.State != EControlRigState::Update)
 	{
 		return;
 	}
@@ -190,6 +190,6 @@ FRigUnit_PBIK_Execute()
 	}
 
 	// do all debug drawing
-	Debug.Draw(Context.DrawInterface, &Solver);
+	Debug.Draw(ExecuteContext.UnitContext.DrawInterface, &Solver);
 }
 
diff --git a/Engine/Plugins/Experimental/LiveLinkControlRig/Source/LiveLinkControlRig/Private/LiveLinkRigUnits.cpp b/Engine/Plugins/Experimental/LiveLinkControlRig/Source/LiveLinkControlRig/Private/LiveLinkRigUnits.cpp
--- a/Engine/Plugins/Experimental/LiveLinkControlRig/Source/LiveLinkControlRig/Private/LiveLinkRigUnits.cpp
+++ b/Engine/Plugins/Experimental/LiveLinkControlRig/Source/LiveLinkControlRig/Private/LiveLinkRigUnits.cpp
@@ -25,7 +25,7 @@ ILiveLinkClient* LiveLinkControlRigUtilities::TryGetLiveLinkClient()
 
 FRigUnit_LiveLinkEvaluteFrameAnimation_Execute()
 {
-	switch (Context.State)
+	switch (ExecuteContext.UnitContext.State)
 	{
 	case EControlRigState::Init:
 	case EControlRigState::Update:
@@ -45,7 +45,7 @@ FRigUnit_LiveLinkEvaluteFrameAnimation_Execute()
 		}
 
 		// Debug Drawing if required
-		if (bDrawDebug && Context.DrawInterface != nullptr)
+		if (bDrawDebug && ExecuteContext.UnitContext.DrawInterface != nullptr)
 		{
 			/*
 			FRigBoneHierarchy DebugHierarchy;
@@ -81,7 +81,7 @@ FRigUnit_LiveLinkEvaluteFrameAnimation_Execute()
 
 FRigUnit_LiveLinkGetTransformByName_Execute()
 {
-	switch (Context.State)
+	switch (ExecuteContext.UnitContext.State)
 	{
 	case EControlRigState::Init:
 	case EControlRigState::Update:
@@ -114,7 +114,7 @@ FRigUnit_LiveLinkGetTransformByName_Execute()
 
 FRigUnit_LiveLinkGetParameterValueByName_Execute()
 {
-	switch (Context.State)
+	switch (ExecuteContext.UnitContext.State)
 	{
 	case EControlRigState::Init:
 	case EControlRigState::Update:
@@ -138,7 +138,7 @@ FRigUnit_LiveLinkGetParameterValueByName_Execute()
 
 FRigUnit_LiveLinkEvaluteFrameTransform_Execute()
 {
-	switch (Context.State)
+	switch (ExecuteContext.UnitContext.State)
 	{
 	case EControlRigState::Init:
 	case EControlRigState::Update:
@@ -158,7 +158,7 @@ FRigUnit_LiveLinkEvaluteFrameTransform_Execute()
 		}
 
 		// Debug Drawing if required
-		if (bDrawDebug && Context.DrawInterface != nullptr)
+		if (bDrawDebug && ExecuteContext.UnitContext.DrawInterface != nullptr)
 		{
 			// Rather than drawing as an axis guide, draw as a hierarchy with a dummy root bone which makes it clearer what the origin of the transform
 			// space is
diff --git a/Engine/Source/Developer/RigVMDeveloper/Private/RigVMCompiler/RigVMAST.cpp b/Engine/Source/Developer/RigVMDeveloper/Private/RigVMCompiler/RigVMAST.cpp
--- a/Engine/Source/Developer/RigVMDeveloper/Private/RigVMCompiler/RigVMAST.cpp
+++ b/Engine/Source/Developer/RigVMDeveloper/Private/RigVMCompiler/RigVMAST.cpp
@@ -713,7 +713,7 @@ void FRigVMParserASTSettings::Report(EMessageSeverity::Type InSeverity, UObject*
 
 const TArray<FRigVMASTProxy> FRigVMParserAST::EmptyProxyArray;
 
-FRigVMParserAST::FRigVMParserAST(TArray<URigVMGraph*> InGraphs, URigVMController* InController, const FRigVMParserASTSettings& InSettings, const TArray<FRigVMExternalVariable>& InExternalVariables, const TArray<FRigVMUserDataArray>& InRigVMUserData)
+FRigVMParserAST::FRigVMParserAST(TArray<URigVMGraph*> InGraphs, URigVMController* InController, const FRigVMParserASTSettings& InSettings, const TArray<FRigVMExternalVariable>& InExternalVariables)
 	: LibraryNodeBeingCompiled(nullptr)
 {
 	DECLARE_SCOPE_HIERARCHICAL_COUNTER_FUNC()
@@ -819,7 +819,7 @@ FRigVMParserAST::FRigVMParserAST(TArray<URigVMGraph*> InGraphs, URigVMController
 		bContinueToFoldConstantBranches = false;
 		if (LibraryNodeBeingCompiled == nullptr)
 		{
-			if (FoldConstantValuesToLiterals(InGraphs, InController, InExternalVariables, InRigVMUserData))
+			if (FoldConstantValuesToLiterals(InGraphs, InController, InExternalVariables))
 			{
 				bContinueToFoldConstantBranches = true;
 			}
@@ -1882,7 +1882,7 @@ void FRigVMParserAST::FoldAssignments()
 
 	RemoveExpressions(ExpressionsToRemove);
 }
-bool FRigVMParserAST::FoldConstantValuesToLiterals(TArray<URigVMGraph*> InGraphs, URigVMController* InController, const TArray<FRigVMExternalVariable>& InExternalVariables, const TArray<FRigVMUserDataArray>& InRigVMUserData)
+bool FRigVMParserAST::FoldConstantValuesToLiterals(TArray<URigVMGraph*> InGraphs, URigVMController* InController, const TArray<FRigVMExternalVariable>& InExternalVariables)
 {
 	DECLARE_SCOPE_HIERARCHICAL_COUNTER_FUNC()
 
@@ -1891,11 +1891,6 @@ bool FRigVMParserAST::FoldConstantValuesToLiterals(TArray<URigVMGraph*> InGraphs
 		return false;
 	}
 
-	if (InRigVMUserData.Num() == 0)
-	{
-		return false;
-	}
-
 	// loop over all call externs and figure out if they are a non-const node
 	// with one or more const pins. we then build a temporary VM to run the part of the 
 	// graph, and pull out the required values - we then bake the value into a literal 
@@ -2075,16 +2070,13 @@ bool FRigVMParserAST::FoldConstantValuesToLiterals(TArray<URigVMGraph*> InGraphs
 	TempCompiler->Settings.EnablePinWatches = false;
 	TempCompiler->Settings.ASTSettings = FRigVMParserASTSettings::Fast(Settings.ExecuteContextStruct);
 
-	TempCompiler->Compile(InGraphs, InController, TempVM, InExternalVariables, InRigVMUserData, &Operands, TempAST);
+	TempCompiler->Compile(InGraphs, InController, TempVM, InExternalVariables, &Operands, TempAST);
 
 	TArray<URigVMMemoryStorage*> Memory;
 	Memory.Add(TempVM->GetWorkMemory());
 	Memory.Add(TempVM->GetLiteralMemory());
 	Memory.Add(TempVM->GetDebugMemory());
-	for (const FRigVMUserDataArray& RigVMUserData : InRigVMUserData)
-	{
-		TempVM->Execute(Memory, RigVMUserData);
-	}
+	TempVM->Execute(Memory);
 
 	// copy the values out of the temp VM and set them on the cached value
 	for (const FRigVMASTProxy& PinToComputeProxy : PinsToCompute)
diff --git a/Engine/Source/Developer/RigVMDeveloper/Private/RigVMCompiler/RigVMCodeGenerator.cpp b/Engine/Source/Developer/RigVMDeveloper/Private/RigVMCompiler/RigVMCodeGenerator.cpp
--- a/Engine/Source/Developer/RigVMDeveloper/Private/RigVMCompiler/RigVMCodeGenerator.cpp
+++ b/Engine/Source/Developer/RigVMDeveloper/Private/RigVMCompiler/RigVMCodeGenerator.cpp
@@ -51,7 +51,7 @@ static constexpr TCHAR RigVM_CopyOpAssignFormat[] = TEXT("\t{0} = {1}{2};");
 static constexpr TCHAR RigVM_IncrementOpFormat[] = TEXT("\t{0}++;");
 static constexpr TCHAR RigVM_DecrementOpFormat[] = TEXT("\t{0}--;");
 static constexpr TCHAR RigVM_EqualsOpFormat[] = TEXT("\t{0} = {1} == {2};");
-static constexpr TCHAR RigVM_InvokeEntryFormat[] = TEXT("\tif (InEntryName == EntryName_{0}) return ExecuteEntry_{0}({1});");
+static constexpr TCHAR RigVM_InvokeEntryFormat[] = TEXT("\tif (InEntryName == EntryName_{0}) return ExecuteEntry_{0}();");
 static constexpr TCHAR RigVM_InvokeEntryByNameFormat[] = TEXT("\tERigVMExecuteResult EntryResult = InvokeEntryByName(InEntryName{0});\r\n\tSetInstructionIndex(INDEX_NONE);\r\n\tStopProfiling();\r\n\treturn ExecuteResult;");
 static constexpr TCHAR RigVM_InvokeEntryByNameFormat2[] = TEXT("\tif(!InvokeEntryByName({0}{1})) return ERigVMExecuteResult::Failed;");
 static constexpr TCHAR RigVM_CanExecuteEntryFormat[] = TEXT("\tif(!CanExecuteEntry(InEntryName, false)) { return ERigVMExecuteResult::Failed; }");
@@ -59,7 +59,7 @@ static constexpr TCHAR RigVM_EntryExecuteGuardFormat[] = TEXT("\tFEntryExecuteGu
 static constexpr TCHAR RigVM_PublicContextGuardFormat[] = TEXT("\tTGuardValue<{0}> PublicContextGuard(Context.GetPublicData<{0}>(), PublicContext);");
 static constexpr TCHAR RigVM_EntryNameFormat[] = TEXT("EntryName_{0}");
 static constexpr TCHAR RigVM_SetExecuteContextStructFormat[] = TEXT("\tSetContextPublicDataStruct({0}::StaticStruct());");
-static constexpr TCHAR RigVM_UpdateContextFormat[] = TEXT("\t{0}& PublicContext = UpdateContext<{0}>(AdditionalArguments, {1}, InEntryName);");
+static constexpr TCHAR RigVM_UpdateContextFormat[] = TEXT("\t{0}& PublicContext = UpdateContext<{0}>({1}, InEntryName);");
 static constexpr TCHAR RigVM_TrueFormat[] = TEXT("true");  
 static constexpr TCHAR RigVM_FalseFormat[] = TEXT("false");
 static constexpr TCHAR RigVM_SingleUnderscoreFormat[] = TEXT("_");
@@ -71,9 +71,6 @@ static constexpr TCHAR RigVM_IsValidArrayIndexFormat[] = TEXT("IsValidArrayIndex
 static constexpr TCHAR RigVM_TemporaryArrayIndexFormat[] = TEXT("\tTemporaryArrayIndex = {0};");
 static constexpr TCHAR RigVM_StartProfilingFormat[] = TEXT("\tStartProfiling");
 static constexpr TCHAR RigVM_ExecuteReachedExitFormat[] = TEXT("\tBroadcastExecutionReachedExit();");
-static constexpr TCHAR RigVM_AdditionalArgumentNameFormat[] = TEXT("AdditionalArgument_{0}_{1}");
-static constexpr TCHAR RigVM_FullOpaqueArgumentFormat[] = TEXT("\t{0} {1} = *({2}*)Context.OpaqueArguments[{3}];");
-static constexpr TCHAR RigVM_ParameterOpaqueArgumentFormat[] = TEXT("{0} {1}");
 static constexpr TCHAR RigVM_InstructionLabelFormat[] = TEXT("\tInstruction{0}Label:");
 static constexpr TCHAR RigVM_SetInstructionIndexFormat[] = TEXT("\tSetInstructionIndex({0});");
 static constexpr TCHAR RigVM_ContextFormat[] = TEXT("Context");
@@ -114,12 +111,12 @@ static constexpr TCHAR RigVM_GetVMHashFormat[] = TEXT("\tvirtual uint32 GetVMHas
 static constexpr TCHAR RigVM_GetEntryNamesFormat[] = TEXT("\tvirtual const TArray<FName>& GetEntryNames() const override\r\n\t{\r\n\t\tstatic const TArray<FName> StaticEntryNames = { {0} };\r\n\t\treturn StaticEntryNames;\r\n\t}");
 static constexpr TCHAR RigVM_DeclareUpdateExternalVariablesFormat[] = TEXT("\tvirtual void UpdateExternalVariables() override;");
 static constexpr TCHAR RigVM_DeclareInvokeEntryByNameFormat[] = TEXT("\tERigVMExecuteResult InvokeEntryByName(const FName& InEntryName{0});");
-static constexpr TCHAR RigVM_DeclareInitializeFormat[] = TEXT("\tvirtual bool Initialize(TArrayView<URigVMMemoryStorage*> Memory, TArrayView<void*> AdditionalArguments, bool bInitializeMemory = true) override;");
-static constexpr TCHAR RigVM_DefineInitializeFormat[] = TEXT("bool U{0}::Initialize(TArrayView<URigVMMemoryStorage*> Memory, TArrayView<void*> AdditionalArguments, bool bInitializeMemory)\r\n{");
-static constexpr TCHAR RigVM_DeclareExecuteFormat[] = TEXT("\tvirtual ERigVMExecuteResult Execute(TArrayView<URigVMMemoryStorage*> Memory, TArrayView<void*> AdditionalArguments, const FName& InEntryName) override;");
+static constexpr TCHAR RigVM_DeclareInitializeFormat[] = TEXT("\tvirtual bool Initialize(TArrayView<URigVMMemoryStorage*> Memory, TArrayView<void*> bool bInitializeMemory = true) override;");
+static constexpr TCHAR RigVM_DefineInitializeFormat[] = TEXT("bool U{0}::Initialize(TArrayView<URigVMMemoryStorage*> Memory, TArrayView<void*> bool bInitializeMemory)\r\n{");
+static constexpr TCHAR RigVM_DeclareExecuteFormat[] = TEXT("\tvirtual ERigVMExecuteResult Execute(TArrayView<URigVMMemoryStorage*> Memory, TArrayView<void*> const FName& InEntryName) override;");
 static constexpr TCHAR RigVM_DefineUpdateExternalVariablesFormat[] = TEXT("void U{0}::UpdateExternalVariables()\r\n{");
 static constexpr TCHAR RigVM_DefineInvokeEntryByNameFormat[] = TEXT("ERigVMExecuteResult U{0}::InvokeEntryByName(const FName& InEntryName{1})\r\n{");
-static constexpr TCHAR RigVM_DefineExecuteFormat[] = TEXT("ERigVMExecuteResult U{0}::Execute(TArrayView<URigVMMemoryStorage*> Memory, TArrayView<void*> AdditionalArguments, const FName& InEntryName)\r\n{");
+static constexpr TCHAR RigVM_DefineExecuteFormat[] = TEXT("ERigVMExecuteResult U{0}::Execute(TArrayView<URigVMMemoryStorage*> Memory, TArrayView<void*> const FName& InEntryName)\r\n{");
 static constexpr TCHAR RigVM_DeclareExecuteEntryFormat[] = TEXT("\tERigVMExecuteResult ExecuteEntry_{0}({1});");
 static constexpr TCHAR RigVM_DefineExecuteEntryFormat[] = TEXT("ERigVMExecuteResult U{0}::ExecuteEntry_{1}({2})\r\n{");
 static constexpr TCHAR RigVM_DeclareExecuteGroupFormat[] = TEXT("\tERigVMExecuteResult ExecuteGroup_{0}_{1}({2});");
@@ -599,32 +596,6 @@ FString FRigVMCodeGenerator::DumpInstructions(int32 InInstructionGroup, bool bLo
 		{
 			Lines.Add(Format(RigVM_SetExecuteContextStructFormat, ExecuteContextType));
 			Lines.Add(Format(RigVM_UpdateContextFormat, ExecuteContextType, Instructions.Num()));
-
-			// let's get the additional arguments for all sub groups
-			TArray<FString> OpaqueArgumentHit;
-			for(const FInstructionGroup& EntryGroup : InstructionGroups)
-			{
-				if(EntryGroup.Entry.IsEmpty() || EntryGroup.Depth > 0)
-				{
-					continue;
-				}
-
-				for(const FOpaqueArgument& OpaqueArgument : EntryGroup.OpaqueArguments)
-				{
-					if(OpaqueArgumentHit.Contains(OpaqueArgument.Get<0>()))
-					{
-						continue;
-					}
-					const int32 ArgumentIndex = OpaqueArgumentHit.Add(OpaqueArgument.Get<0>());
-					FString ArgumentTypeNoRef = OpaqueArgument.Get<2>();
-					if (ArgumentTypeNoRef.EndsWith(TEXT("&")))
-					{
-						ArgumentTypeNoRef.LeftChopInline(1);
-					}
-					Lines.Add(Format(RigVM_FullOpaqueArgumentFormat, *OpaqueArgument.Get<2>(), *OpaqueArgument.Get<1>(), *ArgumentTypeNoRef, ArgumentIndex));
-				}
-			}
-
 			Lines.Add(Format(RigVM_InvokeEntryByNameFormat, *GetEntryParameters()));
 
 			return DumpLines(Lines, bLog);
@@ -641,23 +612,7 @@ FString FRigVMCodeGenerator::DumpInstructions(int32 InInstructionGroup, bool bLo
 			{
 				const FRigVMByteCodeEntry& Entry = ByteCode.GetEntry(EntryIndex);
 				const FString EntryName = Entry.GetSanitizedName();
-				FString OpaqueArguments;
-
-				// find the entry's group and provide the needed arguments
-				for(const FInstructionGroup& EntryGroup : InstructionGroups)
-				{
-					if(EntryGroup.Depth <= 0 && EntryGroup.Entry == EntryName)
-					{
-						TArray<FString> ArgumentNames = {RigVM_ContextPublicFormat};
-						for(const FOpaqueArgument& OpaqueArgument : EntryGroup.OpaqueArguments)
-						{
-							ArgumentNames.Add(OpaqueArgument.Get<1>());
-						}
-						OpaqueArguments = FString::Join(ArgumentNames, RigVM_CommaSeparator);
-						break;
-					}
-				}
-				Lines.Add(Format(RigVM_InvokeEntryFormat, *EntryName, *OpaqueArguments));
+				Lines.Add(Format(RigVM_InvokeEntryFormat, *EntryName));
 			}
 		}
 	}
@@ -720,13 +675,6 @@ FString FRigVMCodeGenerator::DumpInstructions(int32 InInstructionGroup, bool bLo
 
 			FString Parameters;
 			TArray<FString> ParameterArray = {RigVM_ContextPublicFormat};
-			if(!ChildGroup.OpaqueArguments.IsEmpty())
-			{
-				for(const FOpaqueArgument& OpaqueArgument : ChildGroup.OpaqueArguments)
-				{
-					ParameterArray.Add(OpaqueArgument.Get<1>());
-				}
-			}
 			Parameters = FString::Join(ParameterArray, RigVM_CommaSeparator);
 			Lines.Add(Format(RigVM_InvokeExecuteGroupFormat, *ChildGroup.Entry, ChildGroupIndex, *Parameters));
 		}
@@ -871,20 +819,6 @@ FString FRigVMCodeGenerator::DumpInstructions(const FString& InPrefix, const TAr
 					Arguments.Add(GetOperandName(Operand, bSliced));
 				}
 
-				for(const FRigVMFunctionArgument& Argument : Function->GetArguments())
-				{
-					if (Function->IsAdditionalArgument(Argument))
-					{
-						const FOpaqueArgument* OpaqueArgument = InGroup.OpaqueArguments.FindByPredicate([Argument](const FOpaqueArgument& InArgument) -> bool
-						{
-							return InArgument.Get<0>() == Argument.Name;
-						});
-
-						check(OpaqueArgument);
-						Arguments.Add(OpaqueArgument->Get<1>());
-					}
-				}
-
 				const FString JoinedArguments = FString::Join(Arguments, RigVM_CommaSeparator);
 				if(Function->Struct)
 				{
@@ -1273,13 +1207,6 @@ FString FRigVMCodeGenerator::DumpHeader(bool bLog)
 
 		FString Parameters;
 		TArray<FString> ParameterArray = {Format(RigVM_ContextPublicParameterFormat, ExecuteContextType)};
-		if(!Group.OpaqueArguments.IsEmpty())
-		{
-			for(const FOpaqueArgument& OpaqueArgument : Group.OpaqueArguments)
-			{
-				ParameterArray.Add(Format(RigVM_ParameterOpaqueArgumentFormat, *OpaqueArgument.Get<2>(), *OpaqueArgument.Get<1>()));
-			}
-		}
 		Parameters = FString::Join(ParameterArray, RigVM_CommaSeparator);
 		
 		if(Group.Depth == 0)
@@ -1376,13 +1303,6 @@ FString FRigVMCodeGenerator::DumpSource(bool bLog)
 
 		FString Parameters;
 		TArray<FString> ParameterArray = {Format(RigVM_ContextPublicParameterFormat, ExecuteContextType)};
-		if(!Group.OpaqueArguments.IsEmpty())
-		{
-			for(const FOpaqueArgument& OpaqueArgument : Group.OpaqueArguments)
-			{
-				ParameterArray.Add(Format(RigVM_ParameterOpaqueArgumentFormat, *OpaqueArgument.Get<2>(), *OpaqueArgument.Get<1>()));
-			}
-		}
 		Parameters = FString::Join(ParameterArray, RigVM_CommaSeparator);
 
 		Lines.Emplace();
@@ -1910,22 +1830,6 @@ void FRigVMCodeGenerator::ParseInstructionGroups()
 						ParseInclude(Function->Factory->GetScriptStruct(), Function->GetMethodName());
 					}
 
-					// check all of the arguments that may not be part of the struct definition
-					for(const FRigVMFunctionArgument& Argument : Function->GetArguments())
-					{
-						if (Function->IsAdditionalArgument(Argument))
-						{
-							if(!Group.OpaqueArguments.ContainsByPredicate([Argument](const FOpaqueArgument& OpaqueArgument) -> bool
-							{
-								return Argument.Name == OpaqueArgument.Get<0>();
-							}))
-							{
-								const FString ArgumentName = Format(RigVM_AdditionalArgumentNameFormat, Group.OpaqueArguments.Num(), Argument.Name);
-								FOpaqueArgument OpaqueArgument(Argument.Name, ArgumentName, Argument.Type);
-								Group.OpaqueArguments.AddUnique(OpaqueArgument);
-							}
-						}
-					}
 					break;
 				}
 				case ERigVMOpCode::JumpAbsolute:
@@ -2477,10 +2381,6 @@ FString FRigVMCodeGenerator::GetEntryParameters() const
 		if(EntryGroup.Depth <= 0 && EntryGroup.Entry == EntryName)
 		{
 			TArray<FString> ArgumentNames = {RigVM_ContextPublicFormat};
-			for(const FOpaqueArgument& OpaqueArgument : EntryGroup.OpaqueArguments)
-			{
-				ArgumentNames.Add(OpaqueArgument.Get<1>());
-			}
 			Parameters = FString::Join(ArgumentNames, RigVM_CommaSeparator);
 			break;
 		}
diff --git a/Engine/Source/Developer/RigVMDeveloper/Private/RigVMCompiler/RigVMCompiler.cpp b/Engine/Source/Developer/RigVMDeveloper/Private/RigVMCompiler/RigVMCompiler.cpp
--- a/Engine/Source/Developer/RigVMDeveloper/Private/RigVMCompiler/RigVMCompiler.cpp
+++ b/Engine/Source/Developer/RigVMDeveloper/Private/RigVMCompiler/RigVMCompiler.cpp
@@ -256,7 +256,7 @@ URigVMCompiler::URigVMCompiler()
 {
 }
 
-bool URigVMCompiler::Compile(TArray<URigVMGraph*> InGraphs, URigVMController* InController, URigVM* OutVM, const TArray<FRigVMExternalVariable>& InExternalVariables, const TArray<FRigVMUserDataArray>& InRigVMUserData, TMap<FString, FRigVMOperand>* OutOperands, TSharedPtr<FRigVMParserAST> InAST, FRigVMFunctionCompilationData* OutFunctionCompilationData)
+bool URigVMCompiler::Compile(TArray<URigVMGraph*> InGraphs, URigVMController* InController, URigVM* OutVM, const TArray<FRigVMExternalVariable>& InExternalVariables, TMap<FString, FRigVMOperand>* OutOperands, TSharedPtr<FRigVMParserAST> InAST, FRigVMFunctionCompilationData* OutFunctionCompilationData)
 {
 	double CompilationTime = 0;
 	FDurationTimer CompileTimer(CompilationTime);
@@ -279,11 +279,21 @@ bool URigVMCompiler::Compile(TArray<URigVMGraph*> InGraphs, URigVMController* In
 		return false;;
 	}
 
+	// also during traverse - find all known execute contexts
+	// for functions / dispatches / templates.
+	// we only allow compatible execute context structs within a VM
+	TArray<UStruct*> ValidExecuteContextStructs = FRigVMTemplate::GetSuperStructs(Settings.GetExecuteContextStruct());
+	TArray<FString> ValidExecuteContextStructNames;
+	Algo::Transform(ValidExecuteContextStructs, ValidExecuteContextStructNames, [](const UStruct* InStruct)
+	{
+		return CastChecked<UScriptStruct>(InStruct)->GetStructCPPName();
+	});
+
 	for(URigVMGraph* Graph : InGraphs)
 	{
 		if(Graph->GetExecuteContextStruct())
 		{
-			if(!Settings.GetExecuteContextStruct()->IsChildOf(Graph->GetExecuteContextStruct()))
+			if(!ValidExecuteContextStructs.Contains(Graph->GetExecuteContextStruct()))
 			{
 				ReportErrorf(
 					TEXT("Compiler settings' ExecuteContext (%s) is not compatible with '%s' graph's ExecuteContext (%s). Cannot compile."),
@@ -313,12 +323,6 @@ bool URigVMCompiler::Compile(TArray<URigVMGraph*> InGraphs, URigVMController* In
 
 	DECLARE_SCOPE_HIERARCHICAL_COUNTER_FUNC()
 
-	TArray<FRigVMUserDataArray> UserData = InRigVMUserData;
-	if (UserData.Num() == 0)
-	{
-		UserData.Add(FRigVMUserDataArray());
-	}
-
 	OutVM->SetContextPublicDataStruct(Settings.GetExecuteContextStruct());
 	OutVM->Reset();
 
@@ -439,11 +443,6 @@ bool URigVMCompiler::Compile(TArray<URigVMGraph*> InGraphs, URigVMController* In
 	TArray<URigVMGraph*> VisitedGraphs;
 	VisitedGraphs.Append(InGraphs);
 
-	// also during traverse - find all known execute context and argument notations
-	// for functions / dispatches / templates.
-	// we only allow compatible execute context structs within a VM
-	TArray<UScriptStruct*> ExecuteContextStructs;
-	
 	const FRigVMRegistry& Registry = FRigVMRegistry::Get();
 
 	for(int32 GraphIndex=0; GraphIndex<VisitedGraphs.Num(); GraphIndex++)
@@ -640,11 +639,45 @@ bool URigVMCompiler::Compile(TArray<URigVMGraph*> InGraphs, URigVMController* In
 				}
 			}
 
+			auto ReportIncompatibleExecuteContextString = [&] (const FString InExecuteContextName)
+			{
+				static constexpr TCHAR Format[] = TEXT("ExecuteContext '%s' on node '%s' is not compatible with '%s' provided by the compiler settings."); 
+				ReportErrorf(
+					Format,
+					*InExecuteContextName,
+					*ModelNode->GetNodePath(),
+					*Settings.GetExecuteContextStruct()->GetStructCPPName());
+				bEncounteredGraphError = true;
+			};
+
+			auto ReportIncompatibleExecuteContext = [&] (const UScriptStruct* InExecuteContext)
+			{
+				ReportIncompatibleExecuteContextString(InExecuteContext->GetStructCPPName());
+			};
+
+			FString ExecuteContextMetaData;
+			if(const URigVMUnitNode* UnitNode = Cast<URigVMUnitNode>(ModelNode))
+			{
+				if(UScriptStruct* Struct = UnitNode->GetScriptStruct())
+				{
+					if(Struct->GetStringMetaDataHierarchical(FRigVMStruct::ExecuteContextName, &ExecuteContextMetaData))
+					{
+						if(!ValidExecuteContextStructNames.Contains(ExecuteContextMetaData))
+						{
+							ReportIncompatibleExecuteContextString(ExecuteContextMetaData);
+						}
+					}
+				}
+			}
+
 			if(const URigVMDispatchNode* DispatchNode = Cast<URigVMDispatchNode>(ModelNode))
 			{
 				if(const FRigVMDispatchFactory* Factory = DispatchNode->GetFactory())
 				{
-					ExecuteContextStructs.AddUnique(Factory->GetExecuteContextStruct());
+					if(!ValidExecuteContextStructs.Contains(Factory->GetExecuteContextStruct()))
+					{
+						ReportIncompatibleExecuteContext(Factory->GetExecuteContextStruct());
+					}
 				}
 			}
 			else if(const URigVMTemplateNode* TemplateNode = Cast<URigVMTemplateNode>(ModelNode))
@@ -665,7 +698,10 @@ bool URigVMCompiler::Compile(TArray<URigVMGraph*> InGraphs, URigVMController* In
 							{
 								if(StructProperty->Struct->IsChildOf(FRigVMExecuteContext::StaticStruct()))
 								{
-									ExecuteContextStructs.AddUnique(StructProperty->Struct);
+									if(!ValidExecuteContextStructs.Contains(StructProperty->Struct))
+									{
+										ReportIncompatibleExecuteContext(StructProperty->Struct);
+									}
 								}
 							}
 						}
@@ -683,7 +719,10 @@ bool URigVMCompiler::Compile(TArray<URigVMGraph*> InGraphs, URigVMController* In
 								const FRigVMTemplateArgumentType& Type = Registry.GetType(Argument->TypeIndex);
 								if(UScriptStruct* ExecuteContextStruct = Cast<UScriptStruct>(Type.CPPTypeObject))
 								{
-									ExecuteContextStructs.AddUnique(ExecuteContextStruct);
+									if(!ValidExecuteContextStructs.Contains(ExecuteContextStruct))
+									{
+										ReportIncompatibleExecuteContext(ExecuteContextStruct);
+									}
 								}
 							}
 						}
@@ -701,43 +740,6 @@ bool URigVMCompiler::Compile(TArray<URigVMGraph*> InGraphs, URigVMController* In
 		}
 	}
 
-	// check if the execute context structs are compatible
-	UScriptStruct* TopLevelExecuteContextStruct = FRigVMExecuteContext::StaticStruct();
-	for(UScriptStruct* ExecuteContextStruct : ExecuteContextStructs)
-	{
-		if(TopLevelExecuteContextStruct->IsChildOf(ExecuteContextStruct))
-		{
-			continue;
-		}
-		
-		if(ExecuteContextStruct->IsChildOf(TopLevelExecuteContextStruct))
-		{
-			TopLevelExecuteContextStruct = ExecuteContextStruct;
-			continue;
-		}
-
-		// at this point the execute context structs are unrelated
-		ReportErrorf(
-			TEXT("ExecuteContext types '%s' and '%s' are not compatible. Cannot compile this graph."),
-			*ExecuteContextStruct->GetStructCPPName(),
-			*TopLevelExecuteContextStruct->GetStructCPPName());
-		bEncounteredGraphError = true;
-		break;
-	}
-
-	// check if the top level matches the settings
-	if(TopLevelExecuteContextStruct)
-	{
-		if(!Settings.GetExecuteContextStruct()->IsChildOf(TopLevelExecuteContextStruct))
-		{
-			ReportErrorf(
-				TEXT("The graph contains an incompatible execute context '%s'. The compiler settings expect '%s' (or below)."),
-				*TopLevelExecuteContextStruct->GetStructCPPName(),
-				*Settings.GetExecuteContextStruct()->GetStructCPPName());
-			bEncounteredGraphError = true;
-		}
-	}
-
 	if(bEncounteredGraphError)
 	{
 		return false;
@@ -758,7 +760,7 @@ bool URigVMCompiler::Compile(TArray<URigVMGraph*> InGraphs, URigVMController* In
 	WorkData.AST = InAST;
 	if (!WorkData.AST.IsValid())
 	{
-		WorkData.AST = MakeShareable(new FRigVMParserAST(InGraphs, InController, Settings.ASTSettings, InExternalVariables, UserData));
+		WorkData.AST = MakeShareable(new FRigVMParserAST(InGraphs, InController, Settings.ASTSettings, InExternalVariables));
 		for(URigVMGraph* Graph : InGraphs)
 		{
 			Graph->RuntimeAST = WorkData.AST;
@@ -770,9 +772,8 @@ bool URigVMCompiler::Compile(TArray<URigVMGraph*> InGraphs, URigVMController* In
 	ensure(WorkData.AST.IsValid());
 
 	WorkData.VM = OutVM;
-	WorkData.ExecuteContextStruct = InGraphs[0]->GetExecuteContextStruct();
+	WorkData.ExecuteContextStruct = Settings.GetExecuteContextStruct();
 	WorkData.PinPathToOperand = OutOperands;
-	WorkData.RigVMUserData = UserData[0];
 	WorkData.bSetupMemory = true;
 	WorkData.ProxySources = &WorkData.AST->SharedOperandPins;
 
@@ -1089,11 +1090,9 @@ bool URigVMCompiler::CompileFunction(const URigVMLibraryNode* InLibraryNode, URi
 		ExternalVariables = InController->GetExternalVariablesDelegate.Execute(InLibraryNode->GetContainedGraph());
 	}
 	TMap<FString, FRigVMOperand> Operands;
-	TArray<FRigVMUserDataArray> UserData;
-	UserData.Add(FRigVMUserDataArray());
 	
 	URigVM* TempVM = NewObject<URigVM>(InLibraryNode->GetContainedGraph());
-	const bool bSuccess = Compile({InLibraryNode->GetContainedGraph()}, InController, TempVM, ExternalVariables, UserData, &Operands, nullptr, OutFunctionCompilationData);
+	const bool bSuccess = Compile({InLibraryNode->GetContainedGraph()}, InController, TempVM, ExternalVariables, &Operands, nullptr, OutFunctionCompilationData);
 	TempVM->ClearMemory();
 	TempVM->Rename(nullptr, GetTransientPackage(), REN_ForceNoResetLoaders | REN_DoNotDirty | REN_DontCreateRedirectors | REN_NonTransactional);
 	TempVM->MarkAsGarbage();
diff --git a/Engine/Source/Developer/RigVMDeveloper/Private/RigVMModel/RigVMClient.cpp b/Engine/Source/Developer/RigVMDeveloper/Private/RigVMModel/RigVMClient.cpp
--- a/Engine/Source/Developer/RigVMDeveloper/Private/RigVMModel/RigVMClient.cpp
+++ b/Engine/Source/Developer/RigVMDeveloper/Private/RigVMModel/RigVMClient.cpp
@@ -267,6 +267,8 @@ void FRigVMClient::AddModel(URigVMGraph* InModel, bool bCreateController)
 		Models.Add(InModel);
 	}
 
+	InModel->SetExecuteContextStruct(GetExecuteContextStruct());
+
 	if(bCreateController)
 	{
 		CreateController(InModel);
diff --git a/Engine/Source/Developer/RigVMDeveloper/Public/RigVMCompiler/RigVMAST.h b/Engine/Source/Developer/RigVMDeveloper/Public/RigVMCompiler/RigVMAST.h
--- a/Engine/Source/Developer/RigVMDeveloper/Public/RigVMCompiler/RigVMAST.h
+++ b/Engine/Source/Developer/RigVMDeveloper/Public/RigVMCompiler/RigVMAST.h
@@ -1301,7 +1301,7 @@ public:
 	// default constructor
 	// @param InGraph The graph / model to parse
 	// @param InSettings The parse settings to use
-	FRigVMParserAST(TArray<URigVMGraph*> InGraphs, URigVMController* InController = nullptr, const FRigVMParserASTSettings& InSettings = FRigVMParserASTSettings::Fast(), const TArray<FRigVMExternalVariable>& InExternalVariables = TArray<FRigVMExternalVariable>(), const TArray<FRigVMUserDataArray>& InRigVMUserData = TArray<FRigVMUserDataArray>());
+	FRigVMParserAST(TArray<URigVMGraph*> InGraphs, URigVMController* InController = nullptr, const FRigVMParserASTSettings& InSettings = FRigVMParserASTSettings::Fast(), const TArray<FRigVMExternalVariable>& InExternalVariables = TArray<FRigVMExternalVariable>());
 
 	// default destructor
 	~FRigVMParserAST();
@@ -1483,7 +1483,7 @@ private:
 	void FoldAssignments();
 
 	// helper function to fold constant values into literals
-	bool FoldConstantValuesToLiterals(TArray<URigVMGraph*> InGraphs, URigVMController* InController, const TArray<FRigVMExternalVariable>& InExternalVariables, const TArray<FRigVMUserDataArray>& InRigVMUserData);
+	bool FoldConstantValuesToLiterals(TArray<URigVMGraph*> InGraphs, URigVMController* InController, const TArray<FRigVMExternalVariable>& InExternalVariables);
 
 	// helper function to inline all contributing nodes of the graph
 	void Inline(TArray<URigVMGraph*> InGraphs);
diff --git a/Engine/Source/Developer/RigVMDeveloper/Public/RigVMCompiler/RigVMCodeGenerator.h b/Engine/Source/Developer/RigVMDeveloper/Public/RigVMCompiler/RigVMCodeGenerator.h
--- a/Engine/Source/Developer/RigVMDeveloper/Public/RigVMCompiler/RigVMCodeGenerator.h
+++ b/Engine/Source/Developer/RigVMDeveloper/Public/RigVMCompiler/RigVMCodeGenerator.h
@@ -41,7 +41,6 @@ private:
 	typedef TMap<FString, FString> FStringMap;
 	typedef TArray<FRigVMPropertyDescription> FRigVMPropertyDescriptionArray;
 	typedef TTuple<FString, FString> FMappedType;
-	typedef TTuple<FString, FString, FString> FOpaqueArgument;
 	typedef TFunction<FString(const FString&)> TStructConstGenerator;
 
 	enum ERigVMNativizedPropertyType
@@ -69,7 +68,6 @@ private:
 		int32 ParentGroup;
 		TArray<int32> ChildGroups;
 		TArray<int32> RequiredLabels;
-		TArray<FOpaqueArgument> OpaqueArguments;
 
 		FInstructionGroup()
 			: Entry()
@@ -87,7 +85,6 @@ private:
 			, ParentGroup(InOther.ParentGroup)
 			, ChildGroups(InOther.ChildGroups)
 			, RequiredLabels(InOther.RequiredLabels)
-			, OpaqueArguments(InOther.OpaqueArguments)
 		{}
 	};
 
diff --git a/Engine/Source/Developer/RigVMDeveloper/Public/RigVMCompiler/RigVMCompiler.h b/Engine/Source/Developer/RigVMDeveloper/Public/RigVMCompiler/RigVMCompiler.h
--- a/Engine/Source/Developer/RigVMDeveloper/Public/RigVMCompiler/RigVMCompiler.h
+++ b/Engine/Source/Developer/RigVMDeveloper/Public/RigVMCompiler/RigVMCompiler.h
@@ -93,7 +93,6 @@ public:
 	bool bSetupMemory;
 	URigVM* VM;
 	UScriptStruct* ExecuteContextStruct;
-	FRigVMUserDataArray RigVMUserData;
 	TMap<FString, FRigVMOperand>* PinPathToOperand;
 	TMap<const FRigVMVarExprAST*, FRigVMOperand> ExprToOperand;
 	TMap<const FRigVMExprAST*, bool> ExprComplete;
@@ -173,10 +172,10 @@ public:
 	UFUNCTION(BlueprintCallable, Category = FRigVMCompiler)
 	bool Compile(TArray<URigVMGraph*> InGraphs, URigVMController* InController, URigVM* OutVM)
 	{
-		return Compile(InGraphs, InController, OutVM, TArray<FRigVMExternalVariable>(), TArray<FRigVMUserDataArray>(), nullptr);
+		return Compile(InGraphs, InController, OutVM, TArray<FRigVMExternalVariable>(), nullptr);
 	}
 
-	bool Compile(TArray<URigVMGraph*> InGraphs, URigVMController* InController, URigVM* OutVM, const TArray<FRigVMExternalVariable>& InExternalVariables, const TArray<FRigVMUserDataArray>& InRigVMUserData, TMap<FString, FRigVMOperand>* OutOperands, TSharedPtr<FRigVMParserAST> InAST = TSharedPtr<FRigVMParserAST>(), FRigVMFunctionCompilationData* OutFunctionCompilationData = nullptr);
+	bool Compile(TArray<URigVMGraph*> InGraphs, URigVMController* InController, URigVM* OutVM, const TArray<FRigVMExternalVariable>& InExternalVariables, TMap<FString, FRigVMOperand>* OutOperands, TSharedPtr<FRigVMParserAST> InAST = TSharedPtr<FRigVMParserAST>(), FRigVMFunctionCompilationData* OutFunctionCompilationData = nullptr);
 
 	bool CompileFunction(const URigVMLibraryNode* InLibraryNode, URigVMController* InController, FRigVMFunctionCompilationData* OutFunctionCompilationData);
 
diff --git a/Engine/Source/Developer/RigVMDeveloper/Public/RigVMModel/RigVMClient.h b/Engine/Source/Developer/RigVMDeveloper/Public/RigVMModel/RigVMClient.h
--- a/Engine/Source/Developer/RigVMDeveloper/Public/RigVMModel/RigVMClient.h
+++ b/Engine/Source/Developer/RigVMDeveloper/Public/RigVMModel/RigVMClient.h
@@ -83,12 +83,14 @@ public:
 	GENERATED_BODY()
 
 	FRigVMClient()
-		: FunctionLibrary(nullptr)
+		: ExecuteContextStruct(nullptr)
+		, FunctionLibrary(nullptr)
 		, bSuspendNotifications(false)
 		, bIgnoreModelNotifications(false)
 		, OuterClientHost(nullptr)
 		, OuterClientPropertyName(NAME_None)
 	{
+		ExecuteContextStruct = FRigVMExecuteContext::StaticStruct();
 	}
 
 	void SetOuterClientHost(UObject* InOuterClientHost, const FName& InOuterClientHostPropertyName);
@@ -112,6 +114,8 @@ public:
 	URigVMFunctionLibrary* GetFunctionLibrary() const { return FunctionLibrary; }
 	URigVMFunctionLibrary* GetOrCreateFunctionLibrary(bool bSetupUndoRedo, const FObjectInitializer* ObjectInitializer = nullptr, bool bCreateController = true);
 	TArray<FName> GetEntryNames() const;
+	UScriptStruct* GetExecuteContextStruct() const { return ExecuteContextStruct; }
+	void SetExecuteContextStruct(UScriptStruct* InExecuteContextStruct) { ExecuteContextStruct = InExecuteContextStruct; }
 
 	URigVMGraph* AddModel(const FName& InName, bool bSetupUndoRedo, const FObjectInitializer* ObjectInitializer = nullptr, bool bCreateController = true);
 	void AddModel(URigVMGraph* InModel, bool bCreateController);
@@ -166,6 +170,9 @@ private:
 
 	URigVMController* CreateController(const URigVMGraph* InModel);
 
+	UPROPERTY(transient)
+	TObjectPtr<UScriptStruct> ExecuteContextStruct;
+
 	UPROPERTY()
 	TArray<TObjectPtr<URigVMGraph>> Models;
 
diff --git a/Engine/Source/Programs/Shared/EpicGames.UHT/Exporters/CodeGen/UhtHeaderCodeGeneratorCppFile.cs b/Engine/Source/Programs/Shared/EpicGames.UHT/Exporters/CodeGen/UhtHeaderCodeGeneratorCppFile.cs
--- a/Engine/Source/Programs/Shared/EpicGames.UHT/Exporters/CodeGen/UhtHeaderCodeGeneratorCppFile.cs
+++ b/Engine/Source/Programs/Shared/EpicGames.UHT/Exporters/CodeGen/UhtHeaderCodeGeneratorCppFile.cs
@@ -556,17 +556,6 @@ namespace EpicGames.UHT.Exporters.CodeGen
 								.Append(parameter.TypeOriginal())
 								.Append("\"));\r\n");
 						}
-						foreach (UhtRigVMParameter parameter in methodInfo.Parameters)
-						{
-							builder
-								.Append("\t\t")
-								.AppendArgumentsName(scriptStruct, methodInfo)
-								.Append(".Emplace(TEXT(\"")
-								.Append(parameter.NameOriginal())
-								.Append("\"), TEXT(\"")
-								.Append(parameter.TypeOriginal())
-								.Append("\"));\r\n");
-						}
 						builder
 							.Append("\t\tFRigVMRegistry::Get().Register(TEXT(\"")
 							.Append(scriptStruct.SourceName)
@@ -702,6 +691,12 @@ namespace EpicGames.UHT.Exporters.CodeGen
 			// The static method is implemented by the user using a macro.
 			if (scriptStruct.RigVMStructInfo != null)
 			{
+				string constPrefix = "";
+				if (scriptStruct.RigVMStructInfo.ExecuteContextMember == String.Empty)
+				{
+					constPrefix = "const ";
+				}
+				
 				foreach (UhtRigVMMethodInfo methodInfo in scriptStruct.RigVMStructInfo.Methods)
 				{
 					builder.Append("\r\n");
@@ -712,9 +707,7 @@ namespace EpicGames.UHT.Exporters.CodeGen
 						.Append(scriptStruct.SourceName)
 						.Append("::")
 						.Append(methodInfo.Name)
-						.Append('(')
-						.AppendParameterDecls(methodInfo.Parameters, false, ",\r\n\t\t")
-						.Append(")\r\n");
+						.Append("()\r\n");
 					builder.Append("{\r\n");
 					if (scriptStruct.RigVMStructInfo.ExecuteContextMember == String.Empty)
 					{
@@ -726,6 +719,27 @@ namespace EpicGames.UHT.Exporters.CodeGen
 					}
 
 					builder.Append("\tTemporaryExecuteContext.Initialize();\r\n");
+					builder.Append('\t');
+					if (methodInfo.ReturnType != "void")
+					{
+						builder.Append("return ");
+					}
+					builder
+						.Append(methodInfo.Name)
+						.Append("(TemporaryExecuteContext);\r\n")
+						.Append("}\r\n\r\n");
+
+					builder
+						.Append(methodInfo.ReturnType)
+						.Append(' ')
+						.Append(scriptStruct.SourceName)
+						.Append("::")
+						.Append(methodInfo.Name)
+						.Append('(')
+						.Append(constPrefix)
+						.Append(scriptStruct.RigVMStructInfo.ExecuteContextType)
+						.Append("& InExecuteContext)\r\n");
+					builder.Append("{\r\n");
 
 					bool wroteLine = false;
 					foreach (UhtRigVMParameter parameter in scriptStruct.RigVMStructInfo.Members)
@@ -745,9 +759,8 @@ namespace EpicGames.UHT.Exporters.CodeGen
 
 					//COMPATIBILITY-TODO - Replace spaces with \t
 					builder.Append('\t').Append(methodInfo.ReturnPrefix()).Append("Static").Append(methodInfo.Name).Append("(\r\n");
-					builder.Append("\t\tTemporaryExecuteContext");
+					builder.Append("\t\tInExecuteContext");
 					builder.AppendParameterNames(scriptStruct.RigVMStructInfo.Members, true, ",\r\n\t\t", true);
-					builder.AppendParameterNames(methodInfo.Parameters, true, ",\r\n\t\t");
 					builder.Append("\r\n");
 					builder.Append("\t);\r\n");
 					builder.Append("}\r\n");
diff --git a/Engine/Source/Programs/Shared/EpicGames.UHT/Exporters/CodeGen/UhtHeaderCodeGeneratorHFile.cs b/Engine/Source/Programs/Shared/EpicGames.UHT/Exporters/CodeGen/UhtHeaderCodeGeneratorHFile.cs
--- a/Engine/Source/Programs/Shared/EpicGames.UHT/Exporters/CodeGen/UhtHeaderCodeGeneratorHFile.cs
+++ b/Engine/Source/Programs/Shared/EpicGames.UHT/Exporters/CodeGen/UhtHeaderCodeGeneratorHFile.cs
@@ -4,6 +4,7 @@ using System;
 using System.Collections.Generic;
 using System.IO;
 using System.Linq;
+using System.Net.NetworkInformation;
 using System.Text;
 using EpicGames.Core;
 using EpicGames.UHT.Types;
@@ -129,9 +130,13 @@ namespace EpicGames.UHT.Exporters.CodeGen
 				{
 					builder.Append("#define ").Append(scriptStruct.SourceName).Append('_').Append(methodInfo.Name).Append("() \\\r\n");
 					builder.Append('\t').Append(methodInfo.ReturnType).Append(' ').Append(scriptStruct.SourceName).Append("::Static").Append(methodInfo.Name).Append("( \\\r\n");
-					builder.Append("\t\t").Append(scriptStruct.RigVMStructInfo.ExecuteContextType).Append("& ").Append(RigVMExecuteContextParamName);
+					builder.Append("\t\t");
+					if (scriptStruct.RigVMStructInfo.ExecuteContextMember == String.Empty)
+					{
+						builder.Append("const ");
+					}
+					builder.Append(scriptStruct.RigVMStructInfo.ExecuteContextType).Append("& ").Append(RigVMExecuteContextParamName);
 					builder.AppendParameterDecls(scriptStruct.RigVMStructInfo.Members, true, ", \\\r\n\t\t", true, false);
-					builder.AppendParameterDecls(methodInfo.Parameters, true, ", \\\r\n\t\t", false, false);
 					builder.Append(" \\\r\n");
 					builder.Append("\t)\r\n");
 				}
@@ -154,33 +159,33 @@ namespace EpicGames.UHT.Exporters.CodeGen
 					// declare the static method as well as the stub method
 					if (scriptStruct.RigVMStructInfo != null)
 					{
+						string constPrefix = "";
+						if (scriptStruct.RigVMStructInfo.ExecuteContextMember == String.Empty)
+						{
+							constPrefix = "const ";
+						}
+						
 						foreach (UhtRigVMMethodInfo methodInfo in scriptStruct.RigVMStructInfo.Methods)
 						{
+							builder.Append('\t').Append(methodInfo.ReturnType).Append(' ').Append(methodInfo.Name).Append('(').Append(constPrefix).Append(scriptStruct.RigVMStructInfo.ExecuteContextType).Append("& InExecuteContext); \\\r\n");
 							builder.Append("\tstatic ").Append(methodInfo.ReturnType).Append(" Static").Append(methodInfo.Name).Append("( \\\r\n");
-							builder.Append("\t\t").Append(scriptStruct.RigVMStructInfo.ExecuteContextType).Append("& ").Append(RigVMExecuteContextParamName);
+							builder.Append("\t\t");
+							if (scriptStruct.RigVMStructInfo.ExecuteContextMember == String.Empty)
+							{
+								builder.Append("const ");
+							}
+							builder.Append(scriptStruct.RigVMStructInfo.ExecuteContextType).Append("& ").Append(RigVMExecuteContextParamName);
 							builder.AppendParameterDecls(scriptStruct.RigVMStructInfo.Members, true, ", \\\r\n\t\t", true, false);
-							builder.AppendParameterDecls(methodInfo.Parameters, true, ", \\\r\n\t\t", false, false);
 							builder.Append(" \\\r\n");
 							builder.Append("\t); \\\r\n");
 
 							builder.Append("\tFORCEINLINE_DEBUGGABLE static ").Append(methodInfo.ReturnType).Append(" RigVM").Append(methodInfo.Name).Append("( \\\r\n");
-							builder.Append("\t\t").Append(RigVMExecuteContextDeclaration).Append(", \\\r\n");
+							builder.Append("\t\t");
+							builder.Append(RigVMExecuteContextDeclaration).Append(", \\\r\n");
 							builder.Append("\t\tFRigVMMemoryHandleArray RigVMMemoryHandles \\\r\n");
 							builder.Append("\t) \\\r\n");
 							builder.Append("\t{ \\\r\n");
 
-							// implement inline stub method body
-							if (methodInfo.Parameters.Count > 0)
-							{
-								for (int parameterIndex = 0; parameterIndex < methodInfo.Parameters.Count; parameterIndex++)
-								{
-									UhtRigVMParameter parameter = methodInfo.Parameters[parameterIndex];
-									builder.Append("\t\t").Append(parameter.Declaration()).Append(" = *(").Append(parameter.TypeNoRef())
-										.Append("*)RigVMExecuteContext.OpaqueArguments[").Append(parameterIndex).Append("]; \\\r\n");
-								}
-								builder.Append("\t\t \\\r\n");
-							}
-
 							if (scriptStruct.RigVMStructInfo.Members.Count > 0)
 							{
 								int operandIndex = 0;
@@ -278,7 +283,6 @@ namespace EpicGames.UHT.Exporters.CodeGen
 							builder.Append("\t\t").Append(methodInfo.ReturnPrefix()).Append("Static").Append(methodInfo.Name).Append("( \\\r\n");
 							builder.Append("\t\t\tRigVMExecuteContext.GetPublicData<").Append(scriptStruct.RigVMStructInfo.ExecuteContextType).Append(">()");
 							builder.AppendParameterNames(scriptStruct.RigVMStructInfo.Members, true, ", \\\r\n\t\t\t", false);
-							builder.AppendParameterNames(methodInfo.Parameters, true, ", \\\r\n\t\t\t");
 							builder.Append(" \\\r\n");
 							builder.Append("\t\t); \\\r\n");
 							builder.Append("\t} \\\r\n");
diff --git a/Engine/Source/Programs/Shared/EpicGames.UHT/Parsers/UhtScriptStructParser.cs b/Engine/Source/Programs/Shared/EpicGames.UHT/Parsers/UhtScriptStructParser.cs
--- a/Engine/Source/Programs/Shared/EpicGames.UHT/Parsers/UhtScriptStructParser.cs
+++ b/Engine/Source/Programs/Shared/EpicGames.UHT/Parsers/UhtScriptStructParser.cs
@@ -219,6 +219,11 @@ namespace EpicGames.UHT.Parsers
 
 				if (!isGetUpgradeInfo && !isGetNextAggregateName)
 				{
+					if (methodInfo.Parameters.Count > 0)
+					{
+						topScope.TokenReader.LogError($"RIGVM_METHOD {scriptStruct.SourceName}::{methodInfo.Name} has {methodInfo.Parameters.Count} parameters. Since 5.2 parameters are no longer allowed for RIGVM_METHOD functions.");
+						methodInfo.Parameters.Clear();
+					}
 					structInfo.Methods.Add(methodInfo);
 				}
 			}
diff --git a/Engine/Source/Programs/Shared/EpicGames.UHT/Types/UhtScriptStruct.cs b/Engine/Source/Programs/Shared/EpicGames.UHT/Types/UhtScriptStruct.cs
--- a/Engine/Source/Programs/Shared/EpicGames.UHT/Types/UhtScriptStruct.cs
+++ b/Engine/Source/Programs/Shared/EpicGames.UHT/Types/UhtScriptStruct.cs
@@ -916,6 +916,11 @@ namespace EpicGames.UHT.Types
 		{
 			if (RigVMStructInfo != null)
 			{
+				if (this.MetaData.TryGetValueHierarchical("ExecuteContext", out string? ExecuteContextMetadata))
+				{
+					RigVMStructInfo.ExecuteContextType = ExecuteContextMetadata;
+				}
+
 				for (UhtStruct? current = this; current != null; current = current.SuperStruct)
 				{
 					foreach (UhtProperty property in current.Properties)
