Control Rig: Remove bNotify argument from RigVMController functions

diff --git a/Engine/Plugins/Animation/ControlRig/Source/ControlRigDeveloper/Private/ControlRigBlueprint.cpp b/Engine/Plugins/Animation/ControlRig/Source/ControlRigDeveloper/Private/ControlRigBlueprint.cpp
--- a/Engine/Plugins/Animation/ControlRig/Source/ControlRigDeveloper/Private/ControlRigBlueprint.cpp
+++ b/Engine/Plugins/Animation/ControlRig/Source/ControlRigDeveloper/Private/ControlRigBlueprint.cpp
@@ -890,7 +890,8 @@ void UControlRigBlueprint::PostLoad()
 			for(URigVMNode* Node : GraphToValidate->GetNodes())
 			{
 				URigVMController* Controller = GetOrCreateController(GraphToValidate);
-				Controller->RemoveUnusedOrphanedPins(Node, false);
+				FRigVMControllerNotifGuard NotifGuard(Controller, true);
+				Controller->RemoveUnusedOrphanedPins(Node);				
 			}
 				
 			for(URigVMNode* Node : GraphToValidate->GetNodes())
@@ -904,9 +905,8 @@ void UControlRigBlueprint::PostLoad()
 						{
 							if(URigVMController* Controller = GetController(GraphToValidate))
 							{
-								Controller->SuspendNotifications(true);
+								FRigVMControllerNotifGuard(Controller, true);
 								Controller->SetPinDefaultValue(Pin->GetPinPath(), TEXT("Null"), false, false, false);
-								Controller->SuspendNotifications(false);
 							}
 						}
 					}
@@ -1307,11 +1307,12 @@ void UControlRigBlueprint::RefreshAllModels(EControlRigBlueprintLoadType InLoadT
 			URigVMController* Controller = GetOrCreateController(GraphToDetach);
 			// temporarily disable default value validation during load time, serialized values should always be accepted
 			TGuardValue<bool> PerGraphDisablePinDefaultValueValidation(Controller->bValidatePinDefaults, false);
+			FRigVMControllerNotifGuard(Controller, true);
 			Controller->DetachLinksFromPinObjects();
 			TArray<URigVMNode*> Nodes = GraphToDetach->GetNodes();
 			for (URigVMNode* Node : Nodes)
 			{
-				Controller->RepopulatePinsOnNode(Node, true, false, true);
+				Controller->RepopulatePinsOnNode(Node, true, true);
 			}
 		}
 		SetupPinRedirectorsForBackwardsCompatibility();
@@ -1322,7 +1323,8 @@ void UControlRigBlueprint::RefreshAllModels(EControlRigBlueprintLoadType InLoadT
 		URigVMController* Controller = GetOrCreateController(GraphToDetach);
 		// at this stage, allow all links to be reattached,
 		// RecomputeAllTemplateFilteredPermutations() later should break any invalid links
-		Controller->ReattachLinksToPinObjects(true /* follow redirectors */, nullptr, false, true, true);
+		FRigVMControllerNotifGuard(Controller, true);
+		Controller->ReattachLinksToPinObjects(true /* follow redirectors */, nullptr, true, true);
 	}
 
 	if(bIsPostLoad)
@@ -1394,11 +1396,11 @@ void UControlRigBlueprint::RefreshAllModels(EControlRigBlueprintLoadType InLoadT
 		URigVMController* Controller = GetOrCreateController(GraphToClean);
 		TGuardValue<bool> RecomputeGuard(Controller->bSuspendRecomputingOuterTemplateFilters, true);
 		TGuardValue<bool> GuardEditGraph(GraphToClean->bEditable, true);
-		Controller->SuspendNotifications(true);
+		FRigVMControllerNotifGuard(Controller, true);
 		
 		for(URigVMNode* ModelNode : GraphToClean->GetNodes())
 		{
-			Controller->RemoveUnusedOrphanedPins(ModelNode, false);
+			Controller->RemoveUnusedOrphanedPins(ModelNode);
 		}
 
 		if(bIsPostLoad)
@@ -1535,7 +1537,6 @@ void UControlRigBlueprint::RefreshAllModels(EControlRigBlueprintLoadType InLoadT
 		}
 #endif
 
-		Controller->SuspendNotifications(false);
 	}
 }
 
@@ -4095,7 +4096,11 @@ void UControlRigBlueprint::PatchVariableNodesOnLoad()
 	{
 		TGuardValue<bool> GuardNotifsSelf(bSuspendModelNotificationsForSelf, true);
 
-		GetOrCreateController()->ReattachLinksToPinObjects();
+		{
+			URigVMController* Controller = GetOrCreateController();
+			FRigVMControllerNotifGuard(Controller, true);
+			Controller->ReattachLinksToPinObjects();
+		}
 
 		check(GetDefaultModel());
 
@@ -4149,7 +4154,7 @@ void UControlRigBlueprint::PatchRigElementKeyCacheOnLoad()
 		{
 			URigVMController* Controller = GetOrCreateController(Graph);
 			TGuardValue<bool> DisablePinDefaultValueValidation(Controller->bValidatePinDefaults, false);
-			Controller->SuspendNotifications(true);
+			FRigVMControllerNotifGuard(Controller, true);
 			for (URigVMNode* Node : Graph->GetNodes())
 			{
 				if (URigVMUnitNode* UnitNode = Cast<URigVMUnitNode>(Node))
@@ -4176,7 +4181,6 @@ void UControlRigBlueprint::PatchRigElementKeyCacheOnLoad()
 					}
 				}
 			}
-			Controller->SuspendNotifications(false);
 		}
 	}
 }
@@ -4361,12 +4365,11 @@ void UControlRigBlueprint::PatchPropagateToChildren()
 		
 		if (URigVMController* Controller = GetOrCreateController(Graph))
 		{
-			Controller->SuspendNotifications(true);
+			FRigVMControllerNotifGuard(Controller, true);
 			for (const URigVMPin* Pin: PinsToUpdate)
 			{
 				Controller->SetPinDefaultValue(Pin->GetPinPath(), TEXT("True"), false, false, false);
 			}
-			Controller->SuspendNotifications(false);
 		}
 	}
 }
@@ -4381,7 +4384,11 @@ void UControlRigBlueprint::PatchParameterNodesOnLoad()
 		
 		TGuardValue<bool> GuardNotifsSelf(bSuspendModelNotificationsForSelf, true);
 
-		GetOrCreateController()->ReattachLinksToPinObjects();
+		{
+			URigVMController* Controller = GetOrCreateController();
+			FRigVMControllerNotifGuard(Controller, true);
+			Controller->ReattachLinksToPinObjects();
+		}
 
 		check(GetDefaultModel());
 
@@ -4513,7 +4520,7 @@ void UControlRigBlueprint::PatchLinksWithCast()
 				if(!SourcePin->IsLinkedTo(TargetPin))
 				{
 					const TArray<URigVMLink*> LinksToReattach = {Link};
-					Controller->ReattachLinksToPinObjects(true, &LinksToReattach, true, false, false);
+					Controller->ReattachLinksToPinObjects(true, &LinksToReattach, false, false);
 				}
 			}
 
diff --git a/Engine/Plugins/Animation/ControlRig/Source/ControlRigEditor/Private/Editor/ControlRigEditor.cpp b/Engine/Plugins/Animation/ControlRig/Source/ControlRigEditor/Private/Editor/ControlRigEditor.cpp
--- a/Engine/Plugins/Animation/ControlRig/Source/ControlRigEditor/Private/Editor/ControlRigEditor.cpp
+++ b/Engine/Plugins/Animation/ControlRig/Source/ControlRigEditor/Private/Editor/ControlRigEditor.cpp
@@ -5004,12 +5004,13 @@ void FControlRigEditor::OnWrappedPropertyChangedChainEvent(UDetailsViewWrapperOb
 				GetControlRigBlueprint()->RequestAutoVMRecompilation();
 			}
 			else if (PropertyPath == TEXT("DefaultValue"))
-			{			
+			{
+				FRigVMControllerNotifGuard(Controller, true);
 				for (FRigVMGraphVariableDescription& Variable : Graph->GetLocalVariables())
 				{
 					if (Variable.Name == VariableDescription.Name)
 					{
-						Controller->SetLocalVariableDefaultValue(Variable.Name, VariableDescription.DefaultValue, true, true, false);
+						Controller->SetLocalVariableDefaultValue(Variable.Name, VariableDescription.DefaultValue, true, true);
 						break;
 					}
 				}
@@ -5083,7 +5084,7 @@ void FControlRigEditor::OnRequestLocalizeFunctionDialog(URigVMLibraryNode* InFun
 FRigVMController_BulkEditResult FControlRigEditor::OnRequestBulkEditDialog(UControlRigBlueprint* InBlueprint, URigVMController* InController,
 	URigVMLibraryNode* InFunction, ERigVMControllerBulkEditType InEditType)
 {
-	const TArray<FAssetData> FirstLevelReferenceAssets = InController->GetAffectedAssets(InEditType, false, true);
+	const TArray<FAssetData> FirstLevelReferenceAssets = InController->GetAffectedAssets(InEditType, false);
 	if(FirstLevelReferenceAssets.Num() == 0)
 	{
 		return FRigVMController_BulkEditResult();
diff --git a/Engine/Plugins/Animation/ControlRig/Source/ControlRigEditor/Private/Editor/SControlRigFunctionBulkEditWidget.cpp b/Engine/Plugins/Animation/ControlRig/Source/ControlRigEditor/Private/Editor/SControlRigFunctionBulkEditWidget.cpp
--- a/Engine/Plugins/Animation/ControlRig/Source/ControlRigEditor/Private/Editor/SControlRigFunctionBulkEditWidget.cpp
+++ b/Engine/Plugins/Animation/ControlRig/Source/ControlRigEditor/Private/Editor/SControlRigFunctionBulkEditWidget.cpp
@@ -62,7 +62,7 @@ void SControlRigFunctionBulkEditWidget::OnAssetsActivated(const TArray<FAssetDat
 
 void SControlRigFunctionBulkEditWidget::LoadAffectedAssets()
 {
-	TArray<FAssetData> FirstLevelReferenceAssets = Controller->GetAffectedAssets(EditType, false, true);
+	TArray<FAssetData> FirstLevelReferenceAssets = Controller->GetAffectedAssets(EditType, false);
 
 	{
 		FScopedSlowTask SlowTask((float)FirstLevelReferenceAssets.Num(), LOCTEXT("LoadingAffectedAssets", "Loading Affected Assets"));
@@ -87,7 +87,7 @@ void SControlRigFunctionBulkEditWidget::LoadAffectedAssets()
 		
 		SlowTask.MakeDialog();
 
-		Controller->GetAffectedAssets(EditType, true, true);
+		Controller->GetAffectedAssets(EditType, true);
 	}
 
 	Controller->OnBulkEditProgressDelegate.Unbind();
@@ -111,7 +111,7 @@ TSharedRef<SWidget> SControlRigFunctionBulkEditWidget::MakeAssetViewForReference
 	AssetPickerConfig.InitialAssetViewType = EAssetViewType::List;
 	AssetPickerConfig.OnGetCustomSourceAssets.BindLambda([this](const FARFilter& SourceFilter, TArray<FAssetData>& AddedAssets)
 	{
-		AddedAssets.Append(Controller->GetAffectedAssets(EditType, false, true));
+		AddedAssets.Append(Controller->GetAffectedAssets(EditType, false));
 	});
 	
 	return ContentBrowserModule.Get().CreateAssetPicker(AssetPickerConfig);
diff --git a/Engine/Plugins/Animation/ControlRig/Source/ControlRigEditor/Private/Graph/NodeSpawners/ControlRigFunctionRefNodeSpawner.cpp b/Engine/Plugins/Animation/ControlRig/Source/ControlRigEditor/Private/Graph/NodeSpawners/ControlRigFunctionRefNodeSpawner.cpp
--- a/Engine/Plugins/Animation/ControlRig/Source/ControlRigEditor/Private/Graph/NodeSpawners/ControlRigFunctionRefNodeSpawner.cpp
+++ b/Engine/Plugins/Animation/ControlRig/Source/ControlRigEditor/Private/Graph/NodeSpawners/ControlRigFunctionRefNodeSpawner.cpp
@@ -385,9 +385,8 @@ UControlRigGraphNode* UControlRigFunctionRefNodeSpawner::SpawnNode(UEdGraph* Par
 			{
 				// If the package is a template, do not remove the ControlRigGraphNode
 				// We might be spawning a node to populate the PROTO_ context menu for function declarations.
-				Controller->SuspendNotifications(true);
+				FRigVMControllerNotifGuard(Controller, true);
 				Controller->RemoveNode(ModelNode, false);
-				Controller->SuspendNotifications(false);
 			}
 
 			if (bIsUserFacingNode)
diff --git a/Engine/Source/Developer/RigVMDeveloper/Private/RigVMCompiler/RigVMCompiler.cpp b/Engine/Source/Developer/RigVMDeveloper/Private/RigVMCompiler/RigVMCompiler.cpp
--- a/Engine/Source/Developer/RigVMDeveloper/Private/RigVMCompiler/RigVMCompiler.cpp
+++ b/Engine/Source/Developer/RigVMDeveloper/Private/RigVMCompiler/RigVMCompiler.cpp
@@ -477,7 +477,7 @@ bool URigVMCompiler::Compile(TArray<URigVMGraph*> InGraphs, URigVMController* In
 				bEncounteredGraphError = true;
 			}
 
-			if(!InController->RemoveUnusedOrphanedPins(ModelNode, true))
+			if(!InController->RemoveUnusedOrphanedPins(ModelNode))
 			{
 				static const FString LinkedMessage = TEXT("Node @@ uses pins that no longer exist. Please rewire the links and re-compile.");
 				Settings.ASTSettings.Report(EMessageSeverity::Error, ModelNode, LinkedMessage);
diff --git a/Engine/Source/Developer/RigVMDeveloper/Private/RigVMModel/RigVMController.cpp b/Engine/Source/Developer/RigVMDeveloper/Private/RigVMModel/RigVMController.cpp
--- a/Engine/Source/Developer/RigVMDeveloper/Private/RigVMModel/RigVMController.cpp
+++ b/Engine/Source/Developer/RigVMDeveloper/Private/RigVMModel/RigVMController.cpp
@@ -928,7 +928,10 @@ URigVMUnitNode* URigVMController::AddUnitNode(UScriptStruct* InScriptStruct, con
 
 	FString ExportedDefaultValue;
 	CreateDefaultValueForStructIfRequired(InScriptStruct, ExportedDefaultValue);
-	AddPinsForStruct(InScriptStruct, Node, nullptr, ERigVMPinDirection::Invalid, ExportedDefaultValue, true);
+	{
+		TGuardValue<bool> SuspendNotifications(bSuspendNotifications, true);
+		AddPinsForStruct(InScriptStruct, Node, nullptr, ERigVMPinDirection::Invalid, ExportedDefaultValue, true);
+	}
 
 	Graph->Nodes.Add(Node);
 	if (!bSuspendNotifications)
@@ -1287,7 +1290,10 @@ URigVMVariableNode* URigVMController::AddVariableNode(const FName& InVariableNam
 	{
 		FString DefaultValue = InDefaultValue;
 		CreateDefaultValueForStructIfRequired(ValuePin->GetScriptStruct(), DefaultValue);
-		AddPinsForStruct(ValuePin->GetScriptStruct(), Node, ValuePin, ValuePin->Direction, DefaultValue, false);
+		{
+			TGuardValue<bool> SuspendNotifications(bSuspendNotifications, true);
+			AddPinsForStruct(ValuePin->GetScriptStruct(), Node, ValuePin, ValuePin->Direction, DefaultValue, false);
+		}
 	}
 	else if (!InDefaultValue.IsEmpty() && InDefaultValue != TEXT("()"))
 	{
@@ -2452,6 +2458,7 @@ URigVMRerouteNode* URigVMController::AddRerouteNodeOnPin(const FString& InPinPat
 
 	if (ValuePin->IsStruct())
 	{
+		TGuardValue<bool> SuspendNotifications(bSuspendNotifications, true);
 		AddPinsForStruct(ValuePin->GetScriptStruct(), Node, ValuePin, ValuePin->Direction, FString(), false);
 	}
 
@@ -3488,6 +3495,7 @@ TArray<FName> URigVMController::ImportNodesFromText(const FString& InText, bool
 		{
 			if (URigVMUnitNode* UnitNode = Cast<URigVMUnitNode>(SubNodes[SubNodeIndex]))
 			{
+				TGuardValue<bool> SuspendNotifications(bSuspendNotifications, true);
 				RepopulatePinsOnNode(UnitNode);
 			}
 		}
@@ -3507,6 +3515,7 @@ TArray<FName> URigVMController::ImportNodesFromText(const FString& InText, bool
 				{
 					FRigVMControllerGraphGuard GraphGuard(this, CollapseNode->GetContainedGraph(), false);
 					TGuardValue<bool> GuardEditGraph(CollapseNode->ContainedGraph->bEditable, true);
+					TGuardValue<bool> SuspendNotifications(bSuspendNotifications, true);
 					ReattachLinksToPinObjects();
 				}
 				
@@ -4390,6 +4399,7 @@ URigVMCollapseNode* URigVMController::CollapseNodes(const TArray<URigVMNode*>& I
 
 		if (CollapsedPin->IsStruct())
 		{
+			TGuardValue<bool> SuspendNotifications(bSuspendNotifications, true);
 			AddPinsForStruct(CollapsedPin->GetScriptStruct(), CollapseNode, CollapsedPin, CollapsedPin->GetDirection(), FString(), false);
 		}
 
@@ -4411,7 +4421,10 @@ URigVMCollapseNode* URigVMController::CollapseNodes(const TArray<URigVMNode*>& I
 		EntryNode = NewObject<URigVMFunctionEntryNode>(CollapseNode->ContainedGraph, TEXT("Entry"));
 		CollapseNode->ContainedGraph->Nodes.Add(EntryNode);
 		EntryNode->Position = -Diagonal * 0.5f - FVector2D(250.f, 0.f);
-		RefreshFunctionPins(EntryNode, false);
+		{
+			TGuardValue<bool> SuspendNotifications(bSuspendNotifications, true);
+			RefreshFunctionPins(EntryNode);
+		}
 
 		Notify(ERigVMGraphNotifType::NodeAdded, EntryNode);
 
@@ -4420,7 +4433,10 @@ URigVMCollapseNode* URigVMController::CollapseNodes(const TArray<URigVMNode*>& I
 			ReturnNode = NewObject<URigVMFunctionReturnNode>(CollapseNode->ContainedGraph, TEXT("Return"));
 			CollapseNode->ContainedGraph->Nodes.Add(ReturnNode);
 			ReturnNode->Position = FVector2D(Diagonal.X, -Diagonal.Y) * 0.5f + FVector2D(300.f, 0.f);
-			RefreshFunctionPins(ReturnNode, false);
+			{
+				TGuardValue<bool> SuspendNotifications(bSuspendNotifications, true);
+				RefreshFunctionPins(ReturnNode);
+			}
 
 			Notify(ERigVMGraphNotifType::NodeAdded, ReturnNode);
 		}
@@ -5722,6 +5738,7 @@ URigVMCollapseNode* URigVMController::PromoteFunctionReferenceNodeToCollapseNode
 	
 		{
 			FRigVMControllerGraphGuard Guard(this, CollapseNode->GetContainedGraph(), false);
+			TGuardValue<bool> SuspendNotifications(bSuspendNotifications, true);
 			ReattachLinksToPinObjects();
 
 			for (URigVMNode* Node : CollapseNode->GetContainedGraph()->GetNodes())
@@ -5793,7 +5810,7 @@ void URigVMController::SetReferencedFunction(URigVMFunctionReferenceNode* InFunc
 	GetGraph()->Notify(ERigVMGraphNotifType::NodeReferenceChanged, InFunctionRefNode);
 }
 
-void URigVMController::RefreshFunctionPins(URigVMNode* InNode, bool bNotify)
+void URigVMController::RefreshFunctionPins(URigVMNode* InNode)
 {
 	if (InNode == nullptr)
 	{
@@ -5806,9 +5823,9 @@ void URigVMController::RefreshFunctionPins(URigVMNode* InNode, bool bNotify)
 	if (EntryNode || ReturnNode)
 	{
 		TArray<URigVMLink*> Links = InNode->GetLinks();
-		DetachLinksFromPinObjects(&Links, bNotify);
-		RepopulatePinsOnNode(InNode, false, bNotify);
-		ReattachLinksToPinObjects(false, &Links, bNotify);
+		DetachLinksFromPinObjects(&Links);
+		RepopulatePinsOnNode(InNode, false);
+		ReattachLinksToPinObjects(false, &Links);
 	}
 }
 
@@ -7345,7 +7362,7 @@ bool URigVMController::SetPinDefaultValue(const FString& InPinPath, const FStrin
 	return true;
 }
 
-bool URigVMController::SetPinDefaultValue(URigVMPin* InPin, const FString& InDefaultValue, bool bResizeArrays, bool bSetupUndoRedo, bool bMergeUndoAction, bool bNotify)
+bool URigVMController::SetPinDefaultValue(URigVMPin* InPin, const FString& InDefaultValue, bool bResizeArrays, bool bSetupUndoRedo, bool bMergeUndoAction)
 {
 	if (!bIsTransacting && !IsGraphEditable())
 	{
@@ -7362,8 +7379,6 @@ bool URigVMController::SetPinDefaultValue(URigVMPin* InPin, const FString& InDef
 		ensure(!InDefaultValue.IsEmpty());
 	}
 
-	TGuardValue<bool> Guard(bSuspendNotifications, !bNotify);
-
 	URigVMGraph* Graph = GetGraph();
 	check(Graph);
  
@@ -8390,7 +8405,10 @@ URigVMPin* URigVMController::InsertArrayPin(URigVMPin* ArrayPin, int32 InIndex,
 		{
 			FString DefaultValue = InDefaultValue;
 			CreateDefaultValueForStructIfRequired(ScriptStruct, DefaultValue);
-			AddPinsForStruct(ScriptStruct, Pin->GetNode(), Pin, Pin->Direction, DefaultValue, false);
+			{
+				TGuardValue<bool> SuspendNotifications(bSuspendNotifications, true);
+				AddPinsForStruct(ScriptStruct, Pin->GetNode(), Pin, Pin->Direction, DefaultValue, false);
+			}
 		}
 	}
 	else if (Pin->IsArray())
@@ -8466,7 +8484,7 @@ bool URigVMController::RemoveArrayPin(const FString& InArrayElementPinPath, bool
 	}
 
 	int32 IndexToRemove = ArrayElementPin->GetPinIndex();
-	if (!RemovePin(ArrayElementPin, bSetupUndoRedo, true))
+	if (!RemovePin(ArrayElementPin, bSetupUndoRedo))
 	{
 		return false;
 	}
@@ -8502,7 +8520,7 @@ bool URigVMController::RemoveArrayPin(const FString& InArrayElementPinPath, bool
 	return true;
 }
 
-bool URigVMController::RemovePin(URigVMPin* InPinToRemove, bool bSetupUndoRedo, bool bNotify)
+bool URigVMController::RemovePin(URigVMPin* InPinToRemove, bool bSetupUndoRedo)
 {
 	if (!bIsTransacting && !IsGraphEditable())
 	{
@@ -8530,13 +8548,13 @@ bool URigVMController::RemovePin(URigVMPin* InPinToRemove, bool bSetupUndoRedo,
 	TArray<URigVMPin*> SubPins = InPinToRemove->GetSubPins();
 	for (URigVMPin* SubPin : SubPins)
 	{
-		if (!RemovePin(SubPin, bSetupUndoRedo, bNotify))
+		if (!RemovePin(SubPin, bSetupUndoRedo))
 		{
 			return false;
 		}
 	}
 
-	if (bNotify)
+	if (!bSuspendNotifications)
 	{
 		Notify(ERigVMGraphNotifType::PinRemoved, InPinToRemove);
 	}
@@ -8629,7 +8647,7 @@ bool URigVMController::SetArrayPinSize(const FString& InArrayPinPath, int32 InSi
 
 	if (!InDefaultValue.IsEmpty() && InDefaultValue != TEXT("()"))
 	{
-		SetPinDefaultValue(Pin, InDefaultValue, false, bSetupUndoRedo, true, !bSuspendNotifications);
+		SetPinDefaultValue(Pin, InDefaultValue, false, bSetupUndoRedo, true);
 	}
 
 	if (bSetupUndoRedo)
@@ -9515,11 +9533,11 @@ bool URigVMController::AddLink(URigVMPin* OutputPin, URigVMPin* InputPin, bool b
 			Notify(ERigVMGraphNotifType::InteractionBracketOpened, nullptr);
 			if(OutputPin->GetNode()->IsA<URigVMRerouteNode>())
 			{
-				SetPinDefaultValue(OutputPin, InputPin->GetDefaultValue(), true, bSetupUndoRedo, false, true);
+				SetPinDefaultValue(OutputPin, InputPin->GetDefaultValue(), true, bSetupUndoRedo, false);
 			}
 			if(InputPin->GetNode()->IsA<URigVMRerouteNode>())
 			{
-				SetPinDefaultValue(OutputPin, OutputPin->GetDefaultValue(), true, bSetupUndoRedo, false, true);
+				SetPinDefaultValue(OutputPin, OutputPin->GetDefaultValue(), true, bSetupUndoRedo, false);
 			}
 			Notify(ERigVMGraphNotifType::InteractionBracketClosed, nullptr);
 		}
@@ -9834,8 +9852,8 @@ bool URigVMController::BreakLink(URigVMPin* OutputPin, URigVMPin* InputPin, bool
 
 			// each time a link is removed, existing orphaned pins
 			// may become unused and thus can be removed
-			RemoveUnusedOrphanedPins(OutputPin->GetNode(), true);
-			RemoveUnusedOrphanedPins(InputPin->GetNode(), true);
+			RemoveUnusedOrphanedPins(OutputPin->GetNode());
+			RemoveUnusedOrphanedPins(InputPin->GetNode());
 			
 			if (!bIsTransacting && !bSuspendRecomputingTemplateFilters)
 			{
@@ -10122,7 +10140,10 @@ FName URigVMController::AddExposedPin(const FName& InPinName, ERigVMPinDirection
 
 		FString DefaultValue = InDefaultValue;
 		CreateDefaultValueForStructIfRequired(Pin->GetScriptStruct(), DefaultValue);
-		AddPinsForStruct(Pin->GetScriptStruct(), LibraryNode, Pin, Pin->Direction, DefaultValue, false);
+		{
+			TGuardValue<bool> SuspendNotifications(bSuspendNotifications, true);
+			AddPinsForStruct(Pin->GetScriptStruct(), LibraryNode, Pin, Pin->Direction, DefaultValue, false);
+		}
 	}
 
 	FRigVMControllerCompileBracketScope CompileScope(this);
@@ -10148,7 +10169,10 @@ FName URigVMController::AddExposedPin(const FName& InPinName, ERigVMPinDirection
 	{
 		EntryNode = NewObject<URigVMFunctionEntryNode>(Graph, TEXT("Entry"));
 		Graph->Nodes.Add(EntryNode);
-		RefreshFunctionPins(EntryNode, false);
+		{
+			TGuardValue<bool> SuspendNotifications(bSuspendNotifications, true);
+			RefreshFunctionPins(EntryNode);
+		}
 		EntryNode->InitializeFilteredPermutations();
 
 		Notify(ERigVMGraphNotifType::NodeAdded, EntryNode);
@@ -10159,14 +10183,17 @@ FName URigVMController::AddExposedPin(const FName& InPinName, ERigVMPinDirection
 	{
 		ReturnNode = NewObject<URigVMFunctionReturnNode>(Graph, TEXT("Return"));
 		Graph->Nodes.Add(ReturnNode);
-		RefreshFunctionPins(ReturnNode, false);
+		{
+			TGuardValue<bool> SuspendNotifications(bSuspendNotifications, true);
+			RefreshFunctionPins(ReturnNode);
+		}
 		ReturnNode->InitializeFilteredPermutations();
 
 		Notify(ERigVMGraphNotifType::NodeAdded, ReturnNode);
 	}
 
-	RefreshFunctionPins(EntryNode, true);
-	RefreshFunctionPins(ReturnNode, true);
+	RefreshFunctionPins(EntryNode);
+	RefreshFunctionPins(ReturnNode);
 
 	// If the type is not a wildcard, add pin as template argument
 	if (InCPPType != RigVMTypeUtils::GetWildCardCPPType() && InCPPType != RigVMTypeUtils::GetWildCardArrayCPPType())
@@ -10407,6 +10434,7 @@ bool URigVMController::AddArgumentForPin(URigVMPin* InPin, URigVMPin* InToLinkPi
 		// Temporarily disconnect connected links
 		if (CurTemplate.NumArguments() > 0)
 		{
+			TGuardValue<bool> SuspendNotifications(bSuspendNotifications, true);
 			DetachLinksFromPinObjects(&ConnectedLinks);
 		}
 	
@@ -10448,6 +10476,7 @@ bool URigVMController::AddArgumentForPin(URigVMPin* InPin, URigVMPin* InToLinkPi
 		// Reconnect  links
 		if (CurTemplate.NumArguments() > 0)
 		{
+			TGuardValue<bool> SuspendNotifications(bSuspendNotifications, true);
 			ReattachLinksToPinObjects(false, &ConnectedLinks);
 		}
 		
@@ -10650,7 +10679,7 @@ bool URigVMController::RemoveExposedPin(const FName& InPinName, bool bSetupUndoR
 		TGuardValue<bool> GuardRecomputeOuterFilters(bSuspendRecomputingOuterTemplateFilters, true);
 		{
 			FRigVMControllerGraphGuard GraphGuard(this, LibraryNode->GetGraph(), bSetupUndoRedo);
-			bSuccessfullyRemovedPin = RemovePin(Pin, bSetupUndoRedo, true);
+			bSuccessfullyRemovedPin = RemovePin(Pin, bSetupUndoRedo);
 		}
 
 		TArray<URigVMVariableNode*> NodesToRemove;
@@ -10669,8 +10698,8 @@ bool URigVMController::RemoveExposedPin(const FName& InPinName, bool bSetupUndoR
 			RemoveNode(NodesToRemove[i], bSetupUndoRedo);
 		}
 
-		RefreshFunctionPins(Graph->GetEntryNode(), true);
-		RefreshFunctionPins(Graph->GetReturnNode(), true);
+		RefreshFunctionPins(Graph->GetEntryNode());
+		RefreshFunctionPins(Graph->GetReturnNode());
 		RefreshFunctionReferences(LibraryNode, false);
 	}
 
@@ -11145,7 +11174,7 @@ bool URigVMController::ChangeExposedPinType(const FName& InPinName, const FStrin
 
 			if (bSuccessChangingType)
 			{
-				RemoveUnusedOrphanedPins(LibraryNode, true);
+				RemoveUnusedOrphanedPins(LibraryNode);
 			}
 		}
 		if (!bSuccessChangingType)
@@ -11163,10 +11192,10 @@ bool URigVMController::ChangeExposedPinType(const FName& InPinName, const FStrin
 		if (InterfaceNode)
 		{
 			const TArray<URigVMLink*> Links = InterfaceNode->GetLinks();
-			DetachLinksFromPinObjects(&Links, true);
-			RepopulatePinsOnNode(InterfaceNode, true, true, bSetupOrphanPins);
-			ReattachLinksToPinObjects(true, &Links, true, bSetupOrphanPins);
-			RemoveUnusedOrphanedPins(InterfaceNode, true);
+			DetachLinksFromPinObjects(&Links);
+			RepopulatePinsOnNode(InterfaceNode, true, bSetupOrphanPins);
+			ReattachLinksToPinObjects(true, &Links, bSetupOrphanPins);
+			RemoveUnusedOrphanedPins(InterfaceNode);
 		}
 	}
 
@@ -11178,7 +11207,7 @@ bool URigVMController::ChangeExposedPinType(const FName& InPinName, const FStrin
 			{
 				FRigVMControllerGraphGuard GraphGuard(this, ReferenceNode->GetGraph(), bSetupUndoRedo);
 				ChangePinType(ReferencedNodePin, InCPPType, InCPPTypeObjectPath, bSetupUndoRedo, bSetupOrphanPins);
-				RemoveUnusedOrphanedPins(ReferenceNode, true);
+				RemoveUnusedOrphanedPins(ReferenceNode);
 			}
         });
 	}
@@ -11193,7 +11222,7 @@ bool URigVMController::ChangeExposedPinType(const FName& InPinName, const FStrin
 				if (ValuePin)
 				{
 					ChangePinType(ValuePin, InCPPType, InCPPTypeObjectPath, bSetupUndoRedo, bSetupOrphanPins);
-					RemoveUnusedOrphanedPins(VariableNode, true);
+					RemoveUnusedOrphanedPins(VariableNode);
 				}
 			}
 		}
@@ -11278,8 +11307,8 @@ bool URigVMController::SetExposedPinIndex(const FName& InPinName, int32 InNewInd
 		Notify(ERigVMGraphNotifType::PinIndexChanged, Pin);
 	}
 
-	RefreshFunctionPins(LibraryNode->GetEntryNode(), true);
-	RefreshFunctionPins(LibraryNode->GetReturnNode(), true);
+	RefreshFunctionPins(LibraryNode->GetEntryNode());
+	RefreshFunctionPins(LibraryNode->GetReturnNode());
 	RefreshFunctionReferences(LibraryNode, false);
 	
 	if (bSetupUndoRedo)
@@ -11356,7 +11385,10 @@ URigVMFunctionReferenceNode* URigVMController::AddFunctionReferenceNodeFromDescr
 
 	FRigVMControllerCompileBracketScope CompileScope(this);
 	
-	RepopulatePinsOnNode(FunctionRefNode, false, false);
+	{
+		TGuardValue<bool> SuspendNotifications(bSuspendNotifications, true);
+		RepopulatePinsOnNode(FunctionRefNode, false);
+	}
 	//UpdateTemplateNodePinTypes(FunctionRefNode, false); // We will need to add this if we ever allow functions to be templates
 
 	Notify(ERigVMGraphNotifType::NodeAdded, FunctionRefNode);
@@ -11621,7 +11653,10 @@ URigVMLibraryNode* URigVMController::AddFunctionToLibrary(const FName& InFunctio
 		{
 			EntryNode = NewObject<URigVMFunctionEntryNode>(CollapseNode->ContainedGraph, TEXT("Entry"));
 			CollapseNode->ContainedGraph->Nodes.Add(EntryNode);
-			RefreshFunctionPins(EntryNode, false);
+			{
+				TGuardValue<bool> SuspendNotifications(bSuspendNotifications, true);
+				RefreshFunctionPins(EntryNode);
+			}
 			EntryNode->Position = FVector2D(-250.f, 0.f);
 			Notify(ERigVMGraphNotifType::NodeAdded, EntryNode);
 		}
@@ -11630,7 +11665,10 @@ URigVMLibraryNode* URigVMController::AddFunctionToLibrary(const FName& InFunctio
 		{
 			ReturnNode = NewObject<URigVMFunctionReturnNode>(CollapseNode->ContainedGraph, TEXT("Return"));
 			CollapseNode->ContainedGraph->Nodes.Add(ReturnNode);
-			RefreshFunctionPins(ReturnNode, false);
+			{
+				TGuardValue<bool> SuspendNotifications(bSuspendNotifications, true);
+				RefreshFunctionPins(ReturnNode);
+			}
 			ReturnNode->Position = FVector2D(250.f, 0.f);
 			Notify(ERigVMGraphNotifType::NodeAdded, ReturnNode);
 		}
@@ -12311,7 +12349,7 @@ bool URigVMController::SetLocalVariableTypeFromObjectPath(const FName& InVariabl
 	return SetLocalVariableType(InVariableName, InCPPType, CPPTypeObject, bSetupUndoRedo, bPrintPythonCommand);
 }
 
-bool URigVMController::SetLocalVariableDefaultValue(const FName& InVariableName, const FString& InDefaultValue, bool bSetupUndoRedo, bool bPrintPythonCommand, bool bNotify)
+bool URigVMController::SetLocalVariableDefaultValue(const FName& InVariableName, const FString& InDefaultValue, bool bSetupUndoRedo, bool bPrintPythonCommand)
 {
 	if (!IsValidGraph())
 	{
@@ -12454,7 +12492,7 @@ bool URigVMController::PerformUserWorkflow(const FRigVMUserWorkflow& InWorkflow,
 	return bSuccess;
 }
 
-TArray<TSoftObjectPtr<URigVMFunctionReferenceNode>> URigVMController::GetAffectedReferences(ERigVMControllerBulkEditType InEditType, bool bForceLoad, bool bNotify)
+TArray<TSoftObjectPtr<URigVMFunctionReferenceNode>> URigVMController::GetAffectedReferences(ERigVMControllerBulkEditType InEditType, bool bForceLoad)
 {
 	TArray<TSoftObjectPtr<URigVMFunctionReferenceNode>> FunctionReferencePtrs;
 	
@@ -12490,7 +12528,7 @@ TArray<TSoftObjectPtr<URigVMFunctionReferenceNode>> URigVMController::GetAffecte
 
 		if(bForceLoad)
 		{
-			if(OnBulkEditProgressDelegate.IsBound() && bNotify)
+			if(OnBulkEditProgressDelegate.IsBound() && !bSuspendNotifications)
 			{
 				OnBulkEditProgressDelegate.Execute(FunctionReferencePtr, InEditType, ERigVMControllerBulkEditProgress::BeginLoad, FunctionReferenceIndex, FunctionReferencePtrs.Num());
 			}
@@ -12500,7 +12538,7 @@ TArray<TSoftObjectPtr<URigVMFunctionReferenceNode>> URigVMController::GetAffecte
 				FunctionReferencePtr.LoadSynchronous();
 			}
 
-			if(OnBulkEditProgressDelegate.IsBound() && bNotify)
+			if(OnBulkEditProgressDelegate.IsBound() && !bSuspendNotifications)
 			{
 				OnBulkEditProgressDelegate.Execute(FunctionReferencePtr, InEditType, ERigVMControllerBulkEditProgress::FinishedLoad, FunctionReferenceIndex, FunctionReferencePtrs.Num());
 			}
@@ -12524,7 +12562,8 @@ TArray<TSoftObjectPtr<URigVMFunctionReferenceNode>> URigVMController::GetAffecte
 				if(URigVMLibraryNode* AffectedFunction = AffectedFunctionReferenceNode->FindFunctionForNode())
 				{
 					FRigVMControllerGraphGuard GraphGuard(this, AffectedFunction->GetContainedGraph(), false);
-					TArray<TSoftObjectPtr<URigVMFunctionReferenceNode>> AffectedFunctionReferencePtrs = GetAffectedReferences(InEditType, bForceLoad, false);
+					TGuardValue<bool> SuspendNotifications(bSuspendNotifications, true);
+					TArray<TSoftObjectPtr<URigVMFunctionReferenceNode>> AffectedFunctionReferencePtrs = GetAffectedReferences(InEditType, bForceLoad);
 					for(TSoftObjectPtr<URigVMFunctionReferenceNode> AffectedFunctionReferencePtr : AffectedFunctionReferencePtrs)
 					{
 						const FString Key = AffectedFunctionReferencePtr.ToSoftObjectPath().ToString();
@@ -12544,7 +12583,7 @@ TArray<TSoftObjectPtr<URigVMFunctionReferenceNode>> URigVMController::GetAffecte
 	return FunctionReferencePtrs;
 }
 
-TArray<FAssetData> URigVMController::GetAffectedAssets(ERigVMControllerBulkEditType InEditType, bool bForceLoad, bool bNotify)
+TArray<FAssetData> URigVMController::GetAffectedAssets(ERigVMControllerBulkEditType InEditType, bool bForceLoad)
 {
 	TArray<FAssetData> Assets;
 
@@ -12555,7 +12594,7 @@ TArray<FAssetData> URigVMController::GetAffectedAssets(ERigVMControllerBulkEditT
 		return Assets;
 	}
 
-	TArray<TSoftObjectPtr<URigVMFunctionReferenceNode>> FunctionReferencePtrs = GetAffectedReferences(InEditType, bForceLoad, bNotify);
+	TArray<TSoftObjectPtr<URigVMFunctionReferenceNode>> FunctionReferencePtrs = GetAffectedReferences(InEditType, bForceLoad);
 	TMap<FString, int32> VisitedAssets;
 
 	URigVMGraph* Graph = GetGraph();
@@ -12744,7 +12783,10 @@ URigVMRerouteNode* URigVMController::AddFreeRerouteNode(bool bShowAsFullNode, co
 	{
 		FString DefaultValue = InDefaultValue;
 		CreateDefaultValueForStructIfRequired(ValuePin->GetScriptStruct(), DefaultValue);
-		AddPinsForStruct(ValuePin->GetScriptStruct(), Node, ValuePin, ValuePin->Direction, DefaultValue, false);
+		{
+			TGuardValue<bool> SuspendNotifications(bSuspendNotifications, true);
+			AddPinsForStruct(ValuePin->GetScriptStruct(), Node, ValuePin, ValuePin->Direction, DefaultValue, false);
+		}
 	}
 	else if (!InDefaultValue.IsEmpty() && InDefaultValue != TEXT("()"))
 	{
@@ -12847,6 +12889,7 @@ URigVMIfNode* URigVMController::AddIfNode(const FString& InCPPType, const FName&
 
 	if (TruePin->IsStruct())
 	{
+		TGuardValue<bool> SuspendNotifications(bSuspendNotifications, true);
 		AddPinsForStruct(TruePin->GetScriptStruct(), Node, TruePin, TruePin->Direction, FString(), false);
 	}
 
@@ -12860,6 +12903,7 @@ URigVMIfNode* URigVMController::AddIfNode(const FString& InCPPType, const FName&
 
 	if (FalsePin->IsStruct())
 	{
+		TGuardValue<bool> SuspendNotifications(bSuspendNotifications, true);
 		AddPinsForStruct(FalsePin->GetScriptStruct(), Node, FalsePin, FalsePin->Direction, FString(), false);
 	}
 
@@ -12872,6 +12916,7 @@ URigVMIfNode* URigVMController::AddIfNode(const FString& InCPPType, const FName&
 
 	if (ResultPin->IsStruct())
 	{
+		TGuardValue<bool> SuspendNotifications(bSuspendNotifications, true);
 		AddPinsForStruct(ResultPin->GetScriptStruct(), Node, ResultPin, ResultPin->Direction, FString(), false);
 	}
 
@@ -12991,6 +13036,7 @@ URigVMSelectNode* URigVMController::AddSelectNode(const FString& InCPPType, cons
 
 	if (ResultPin->IsStruct())
 	{
+		TGuardValue<bool> SuspendNotifications(bSuspendNotifications, true);
 		AddPinsForStruct(ResultPin->GetScriptStruct(), Node, ResultPin, ResultPin->Direction, FString(), false);
 	}
 
@@ -13402,7 +13448,10 @@ URigVMArrayNode* URigVMController::AddArrayNode(ERigVMOpCode InOpCode, const FSt
 				{
 					FString DefaultValue;
 					InController->CreateDefaultValueForStructIfRequired(Pin->GetScriptStruct(), DefaultValue);
-					InController->AddPinsForStruct(Struct, InNode, Pin, InDirection, DefaultValue, true, false);
+					{
+						TGuardValue<bool> SuspendNotifications(InController->bSuspendNotifications, true);
+						InController->AddPinsForStruct(Struct, InNode, Pin, InDirection, DefaultValue, true);
+					}
 				}
 			}
 
@@ -14169,7 +14218,7 @@ bool URigVMController::CanAddFunctionRefForDefinition(const FRigVMGraphFunctionH
 	return true;
 }
 
-void URigVMController::AddPinsForStruct(UStruct* InStruct, URigVMNode* InNode, URigVMPin* InParentPin, ERigVMPinDirection InPinDirection, const FString& InDefaultValue, bool bAutoExpandArrays, bool bNotify)
+void URigVMController::AddPinsForStruct(UStruct* InStruct, URigVMNode* InNode, URigVMPin* InParentPin, ERigVMPinDirection InPinDirection, const FString& InDefaultValue, bool bAutoExpandArrays)
 {
 	if(!ShouldStructBeUnfolded(InStruct))
 	{
@@ -14221,8 +14270,10 @@ void URigVMController::AddPinsForStruct(UStruct* InStruct, URigVMNode* InNode, U
 						DefaultValue = *DefaultValuePtr;
 					}
 					CreateDefaultValueForStructIfRequired(StructProperty->Struct, DefaultValue);
-
-					AddPinsForStruct(StructProperty->Struct, InNode, Pin, Pin->GetDirection(), DefaultValue, bAutoExpandArrays);
+					{
+						TGuardValue<bool> SuspendNotifications(bSuspendNotifications, true);
+						AddPinsForStruct(StructProperty->Struct, InNode, Pin, Pin->GetDirection(), DefaultValue, bAutoExpandArrays);
+					}
 				}
 				else if(DefaultValuePtr != nullptr)
 				{
@@ -14258,7 +14309,7 @@ void URigVMController::AddPinsForStruct(UStruct* InStruct, URigVMNode* InNode, U
 				Pin->DefaultValue = DefaultValue;
 			}
 
-			if (bNotify)
+			if (!bSuspendNotifications)
 			{
 				Notify(ERigVMGraphNotifType::PinAdded, Pin);
 			}
@@ -14302,7 +14353,10 @@ void URigVMController::AddPinsForArray(FArrayProperty* InArrayProperty, URigVMNo
 				{
 					CreateDefaultValueForStructIfRequired(ScriptStruct, DefaultValue);
 				}
-				AddPinsForStruct(StructProperty->Struct, InNode, Pin, Pin->Direction, DefaultValue, bAutoExpandArrays);
+				{
+					TGuardValue<bool> SuspendNotifications(bSuspendNotifications, true);
+					AddPinsForStruct(StructProperty->Struct, InNode, Pin, Pin->Direction, DefaultValue, bAutoExpandArrays);
+				}
 			}
 			else if (!DefaultValue.IsEmpty())
 			{
@@ -14381,14 +14435,15 @@ void URigVMController::AddPinsForTemplate(const FRigVMTemplate* InTemplate, cons
 			{
 				DefaultValue = TemplateNode->GetInitialDefaultValueForPin(Pin->GetFName());
 			}
-			
+
+			TGuardValue<bool> SuspendNotifications(bSuspendNotifications, true);
 			if(UScriptStruct* ScriptStruct = Cast<UScriptStruct>(Pin->CPPTypeObject))
 			{
-				AddPinsForStruct(ScriptStruct, Pin->GetNode(), Pin, Pin->Direction, DefaultValue, false, false);
+				AddPinsForStruct(ScriptStruct, Pin->GetNode(), Pin, Pin->Direction, DefaultValue, false);
 			}
 			else if(!DefaultValue.IsEmpty())
 			{
-				SetPinDefaultValue(Pin, DefaultValue, true, false, false, false);
+				SetPinDefaultValue(Pin, DefaultValue, true, false, false);
 			}
 		}
 	}
@@ -14661,11 +14716,10 @@ FProperty* URigVMController::FindPropertyForPin(const FString& InPinPath)
 	return nullptr;
 }
 
-int32 URigVMController::DetachLinksFromPinObjects(const TArray<URigVMLink*>* InLinks, bool bNotify)
+int32 URigVMController::DetachLinksFromPinObjects(const TArray<URigVMLink*>* InLinks)
 {
 	URigVMGraph* Graph = GetGraph();
 	check(Graph);
-	TGuardValue<bool> EventuallySuspendNotifs(bSuspendNotifications, !bNotify);
 
 	TArray<URigVMLink*> Links;
 	if (InLinks)
@@ -14708,7 +14762,7 @@ int32 URigVMController::DetachLinksFromPinObjects(const TArray<URigVMLink*>* InL
 			{
 				FRigVMControllerGraphGuard GraphGuard(this, CollapseNode->GetContainedGraph(), false);
 				TGuardValue<bool> GuardEditGraph(CollapseNode->ContainedGraph->bEditable, true);
-				DetachLinksFromPinObjects(InLinks, bNotify);
+				DetachLinksFromPinObjects(InLinks);
 			}
 		}
 	}
@@ -14716,11 +14770,10 @@ int32 URigVMController::DetachLinksFromPinObjects(const TArray<URigVMLink*>* InL
 	return Links.Num();
 }
 
-int32 URigVMController::ReattachLinksToPinObjects(bool bFollowCoreRedirectors, const TArray<URigVMLink*>* InLinks, bool bNotify, bool bSetupOrphanedPins, bool bAllowNonArgumentLinks)
+int32 URigVMController::ReattachLinksToPinObjects(bool bFollowCoreRedirectors, const TArray<URigVMLink*>* InLinks, bool bSetupOrphanedPins, bool bAllowNonArgumentLinks)
 {
 	URigVMGraph* Graph = GetGraph();
 	check(Graph);
-	TGuardValue<bool> EventuallySuspendNotifs(bSuspendNotifications, !bNotify);
 	FScopeLock Lock(&PinPathCoreRedirectorsLock);
 
 	bool bReplacingAllLinks = false;
@@ -14904,7 +14957,7 @@ int32 URigVMController::ReattachLinksToPinObjects(bool bFollowCoreRedirectors, c
 			{
 				FRigVMControllerGraphGuard GraphGuard(this, CollapseNode->GetContainedGraph(), false);
 				TGuardValue<bool> GuardEditGraph(CollapseNode->ContainedGraph->bEditable, true);
-				ReattachLinksToPinObjects(bFollowCoreRedirectors, nullptr, bNotify, bSetupOrphanedPins, bAllowNonArgumentLinks);
+				ReattachLinksToPinObjects(bFollowCoreRedirectors, nullptr, bSetupOrphanedPins, bAllowNonArgumentLinks);
 			}
 		}
 	}
@@ -15104,7 +15157,7 @@ bool URigVMController::ShouldRedirectPin(const FString& InOldPinPath, FString& I
 	return false;
 }
 
-void URigVMController::RepopulatePinsOnNode(URigVMNode* InNode, bool bFollowCoreRedirectors, bool bNotify, bool bSetupOrphanedPins)
+void URigVMController::RepopulatePinsOnNode(URigVMNode* InNode, bool bFollowCoreRedirectors, bool bSetupOrphanedPins)
 {
 	if (InNode == nullptr)
 	{
@@ -15126,7 +15179,6 @@ void URigVMController::RepopulatePinsOnNode(URigVMNode* InNode, bool bFollowCore
 	URigVMArrayNode* ArrayNode = Cast<URigVMArrayNode>(InNode);
 	URigVMDispatchNode* DispatchNode = Cast<URigVMDispatchNode>(InNode);
 
-	TGuardValue<bool> EventuallySuspendNotifs(bSuspendNotifications, !bNotify);
 	FScopeLock Lock(&PinPathCoreRedirectorsLock);
 
 	URigVMGraph* Graph = GetGraph();
@@ -15169,7 +15221,7 @@ void URigVMController::RepopulatePinsOnNode(URigVMNode* InNode, bool bFollowCore
 			return;
 		}
 
-		RemovePinsDuringRepopulate(UnitNode, UnitNode->Pins, bNotify, bSetupOrphanedPins);
+		RemovePinsDuringRepopulate(UnitNode, UnitNode->Pins, bSetupOrphanedPins);
 
 		if (ScriptStruct == nullptr)
 		{
@@ -15192,7 +15244,7 @@ void URigVMController::RepopulatePinsOnNode(URigVMNode* InNode, bool bFollowCore
 
 		FString ExportedDefaultValue;
 		CreateDefaultValueForStructIfRequired(ScriptStruct, ExportedDefaultValue);
-		AddPinsForStruct(ScriptStruct, UnitNode, nullptr, ERigVMPinDirection::Invalid, ExportedDefaultValue, false, bNotify);
+		AddPinsForStruct(ScriptStruct, UnitNode, nullptr, ERigVMPinDirection::Invalid, ExportedDefaultValue, false);
 	}
 	else if (DispatchNode)
 	{
@@ -15202,7 +15254,7 @@ void URigVMController::RepopulatePinsOnNode(URigVMNode* InNode, bool bFollowCore
 			PinTypeMap.Add(Pin->GetFName(), Pin->GetTypeIndex());
 		}
 		
-		RemovePinsDuringRepopulate(DispatchNode, DispatchNode->Pins, bNotify, bSetupOrphanedPins);
+		RemovePinsDuringRepopulate(DispatchNode, DispatchNode->Pins, bSetupOrphanedPins);
 
 		const FRigVMTemplate* Template = DispatchNode->GetTemplate();
 		AddPinsForTemplate(Template, PinTypeMap, DispatchNode);
@@ -15291,7 +15343,7 @@ void URigVMController::RepopulatePinsOnNode(URigVMNode* InNode, bool bFollowCore
 			}
 		}
 		
-		RemovePinsDuringRepopulate(InNode, ValuePin->SubPins, bNotify, bSetupOrphanedPins);
+		RemovePinsDuringRepopulate(InNode, ValuePin->SubPins, bSetupOrphanedPins);
 
 		if (ValuePin->IsStruct())
 		{
@@ -15310,7 +15362,10 @@ void URigVMController::RepopulatePinsOnNode(URigVMNode* InNode, bool bFollowCore
 
 			FString ExportedDefaultValue;
 			CreateDefaultValueForStructIfRequired(ScriptStruct, ExportedDefaultValue);
-			AddPinsForStruct(ScriptStruct, InNode, ValuePin, ValuePin->Direction, ExportedDefaultValue, false);
+			{
+				TGuardValue<bool> SuspendNotifications(bSuspendNotifications, true);
+				AddPinsForStruct(ScriptStruct, InNode, ValuePin, ValuePin->Direction, ExportedDefaultValue, false);
+			}
 		}
 	}
 	else if (EntryNode || ReturnNode)
@@ -15318,7 +15373,7 @@ void URigVMController::RepopulatePinsOnNode(URigVMNode* InNode, bool bFollowCore
 		if (URigVMLibraryNode* LibraryNode = Cast<URigVMLibraryNode>(InNode->GetGraph()->GetOuter()))
 		{
 			bool bIsEntryNode = EntryNode != nullptr;
-			RemovePinsDuringRepopulate(InNode, InNode->Pins, bNotify, bSetupOrphanedPins);
+			RemovePinsDuringRepopulate(InNode, InNode->Pins, bSetupOrphanedPins);
 
 			TArray<URigVMPin*> SortedLibraryPins;
 
@@ -15375,6 +15430,7 @@ void URigVMController::RepopulatePinsOnNode(URigVMNode* InNode, bool bFollowCore
 
 				if (ExposedPin->IsStruct())
 				{
+					TGuardValue<bool> SuspendNotifications(bSuspendNotifications, true);
 					AddPinsForStruct(ExposedPin->GetScriptStruct(), InNode, ExposedPin, ExposedPin->GetDirection(), FString(), false);
 				}
 
@@ -15403,7 +15459,7 @@ void URigVMController::RepopulatePinsOnNode(URigVMNode* InNode, bool bFollowCore
 			NewRootPinInfos.Add(TTuple<URigVMPin*, FName>(NewRootPin, RootPin->GetFName()));
 		}
 		
-		RemovePinsDuringRepopulate(InNode, InNode->Pins, bNotify, bSetupOrphanedPins);
+		RemovePinsDuringRepopulate(InNode, InNode->Pins, bSetupOrphanedPins);
 		
 		for (TTuple<URigVMPin*, FName> NewRootPinInfo : NewRootPinInfos)
 		{
@@ -15415,6 +15471,7 @@ void URigVMController::RepopulatePinsOnNode(URigVMNode* InNode, bool bFollowCore
 		{
 			if (Pin->IsStruct())
 			{
+				TGuardValue<bool> SuspendNotifications(bSuspendNotifications, true);
 				AddPinsForStruct(Pin->GetScriptStruct(), InNode, Pin, Pin->GetDirection(), FString(), false);
 			}
 			Notify(ERigVMGraphNotifType::PinAdded, Pin);
@@ -15423,7 +15480,8 @@ void URigVMController::RepopulatePinsOnNode(URigVMNode* InNode, bool bFollowCore
 		if (CollapseNode->GetOuter()->IsA<URigVMFunctionLibrary>())
 		{
 			// no need to notify since the function library graph is invisible anyway
-			RemoveUnusedOrphanedPins(CollapseNode, false);
+			TGuardValue<bool> SuspendNotifications(bSuspendNotifications, true);
+			RemoveUnusedOrphanedPins(CollapseNode);
 		}
 
 		FRigVMControllerGraphGuard GraphGuard(this, CollapseNode->GetContainedGraph(), false);
@@ -15433,7 +15491,7 @@ void URigVMController::RepopulatePinsOnNode(URigVMNode* InNode, bool bFollowCore
 		TArray<URigVMNode*> ContainedNodes = CollapseNode->GetContainedNodes();
 		for (URigVMNode* ContainedNode : ContainedNodes)
 		{
-			RepopulatePinsOnNode(ContainedNode, bFollowCoreRedirectors, bNotify, bSetupOrphanedPins);
+			RepopulatePinsOnNode(ContainedNode, bFollowCoreRedirectors, bSetupOrphanedPins);
 		}
 	}
 	else if (FunctionRefNode)
@@ -15443,7 +15501,7 @@ void URigVMController::RepopulatePinsOnNode(URigVMNode* InNode, bool bFollowCore
 		// we want to make sure notify the graph of a potential name change
 		// when repopulating the function ref node
 		Notify(ERigVMGraphNotifType::NodeRenamed, FunctionRefNode);
-		RemovePinsDuringRepopulate(InNode, InNode->Pins, bNotify, bSetupOrphanedPins);
+		RemovePinsDuringRepopulate(InNode, InNode->Pins, bSetupOrphanedPins);
 
 		TMap<FString, FPinState> ReferencedPinStates;
 		if (Header.LibraryPointer.LibraryNode.ResolveObject())
@@ -15462,6 +15520,7 @@ void URigVMController::RepopulatePinsOnNode(URigVMNode* InNode, bool bFollowCore
 
 			if (NewPin->IsStruct())
 			{
+				TGuardValue<bool> SuspendNotifications(bSuspendNotifications, true);
 				AddPinsForStruct(NewPin->GetScriptStruct(), InNode, NewPin, NewPin->GetDirection(), FString(), false);
 			}
 
@@ -15489,7 +15548,7 @@ void URigVMController::RepopulatePinsOnNode(URigVMNode* InNode, bool bFollowCore
 			NewRootPinInfos.Add(TTuple<URigVMPin*, FName>(NewRootPin, RootPin->GetFName()));
 		}
 		
-		RemovePinsDuringRepopulate(InNode, InNode->Pins, bNotify, bSetupOrphanedPins);
+		RemovePinsDuringRepopulate(InNode, InNode->Pins, bSetupOrphanedPins);
 		
 		for (TTuple<URigVMPin*, FName> NewRootPinInfo : NewRootPinInfos)
 		{
@@ -15501,6 +15560,7 @@ void URigVMController::RepopulatePinsOnNode(URigVMNode* InNode, bool bFollowCore
 		{
 			if (Pin->IsStruct())
 			{
+				TGuardValue<bool> SuspendNotifications(bSuspendNotifications, true);
 				AddPinsForStruct(Pin->GetScriptStruct(), InNode, Pin, Pin->GetDirection(), FString(), false);
 			}
 			Notify(ERigVMGraphNotifType::PinAdded, Pin);
@@ -15520,7 +15580,7 @@ void URigVMController::RepopulatePinsOnNode(URigVMNode* InNode, bool bFollowCore
 	}
 }
 
-void URigVMController::RemovePinsDuringRepopulate(URigVMNode* InNode, TArray<URigVMPin*>& InPins, bool bNotify, bool bSetupOrphanedPins)
+void URigVMController::RemovePinsDuringRepopulate(URigVMNode* InNode, TArray<URigVMPin*>& InPins, bool bSetupOrphanedPins)
 {
 	TArray<URigVMPin*> Pins = InPins;
 	for (URigVMPin* Pin : Pins)
@@ -15549,14 +15609,14 @@ void URigVMController::RemovePinsDuringRepopulate(URigVMNode* InNode, TArray<URi
 					RenameObject(RootPin, *OrphanedName, nullptr);
 					InNode->Pins.Remove(RootPin);
 
-					if(bNotify)
+					if(!bSuspendNotifications)
 					{
 						Notify(ERigVMGraphNotifType::PinRemoved, RootPin);
 					}
 
 					InNode->OrphanedPins.Add(RootPin);
 
-					if(bNotify)
+					if(!bSuspendNotifications)
 					{
 						Notify(ERigVMGraphNotifType::PinAdded, RootPin);
 					}
@@ -15569,7 +15629,7 @@ void URigVMController::RemovePinsDuringRepopulate(URigVMNode* InNode, TArray<URi
 				
 					OrphanedRootPin->GetNode()->OrphanedPins.Add(OrphanedRootPin);
 
-					if(bNotify)
+					if(!bSuspendNotifications)
 					{
 						Notify(ERigVMGraphNotifType::PinAdded, OrphanedRootPin);
 					}
@@ -15590,13 +15650,13 @@ void URigVMController::RemovePinsDuringRepopulate(URigVMNode* InNode, TArray<URi
 	{
 		if(!Pin->IsOrphanPin())
 		{
-			RemovePin(Pin, false, bNotify);
+			RemovePin(Pin, false);
 		}
 	}
 	InPins.Reset();
 }
 
-bool URigVMController::RemoveUnusedOrphanedPins(URigVMNode* InNode, bool bNotify)
+bool URigVMController::RemoveUnusedOrphanedPins(URigVMNode* InNode)
 {
 	if(!InNode->HasOrphanedPins())
 	{
@@ -15613,7 +15673,7 @@ bool URigVMController::RemoveUnusedOrphanedPins(URigVMNode* InNode, bool bNotify
 
 		if(NumSourceLinks + NumTargetLinks == 0)
 		{
-			RemovePin(OrphanedPin, false, bNotify);
+			RemovePin(OrphanedPin, false);
 		}
 		else
 		{
@@ -16137,7 +16197,7 @@ void URigVMController::ResolveTemplateNodeMetaData(URigVMTemplateNode* InNode, b
 				const FString NewDefaultValue = InNode->GetInitialDefaultValueForPin(Pin->GetFName(), FilteredPermutationIndices);
 				if(!NewDefaultValue.IsEmpty())
 				{
-					SetPinDefaultValue(Pin, NewDefaultValue, true, bSetupUndoRedo, false, true);
+					SetPinDefaultValue(Pin, NewDefaultValue, true, bSetupUndoRedo, false);
 				}
 			}
 		}
@@ -16380,7 +16440,7 @@ bool URigVMController::FullyResolveTemplateNode(URigVMTemplateNode* InNode, int3
 	// remove obsolete pins
 	for(URigVMPin* PinToRemove : PinsToRemove)
 	{
-		RemovePin(PinToRemove, false, true);
+		RemovePin(PinToRemove, false);
 	}
 
 	// add missing pins
@@ -18295,7 +18355,10 @@ bool URigVMController::RecomputeAllTemplateFilteredPermutations(bool bSetupUndoR
 			}
 		});
 			
-		DetachLinksFromPinObjects(&SortedLinks);
+		{
+			TGuardValue<bool> SuspendNotifications(bSuspendNotifications, true);
+			DetachLinksFromPinObjects(&SortedLinks);
+		}
 		TArray<URigVMLink*> LinksToRemove;
 		for (int32 i=0; i<SortedLinks.Num(); ++i)
 		{
@@ -18842,7 +18905,7 @@ bool URigVMController::ChangePinType(URigVMPin* InPin, TRigVMTypeIndex InTypeInd
 	{
 		Links.Append(InPin->GetSourceLinks(true));
 		Links.Append(InPin->GetTargetLinks(true));
-		DetachLinksFromPinObjects(&Links, true);
+		DetachLinksFromPinObjects(&Links);
 
 		const FString OrphanedName = FString::Printf(TEXT("%s%s"), *URigVMPin::OrphanPinPrefix, *InPin->GetName());
 		if(InPin->GetNode()->FindPin(OrphanedName) == nullptr)
@@ -18853,7 +18916,7 @@ bool URigVMController::ChangePinType(URigVMPin* InPin, TRigVMTypeIndex InTypeInd
 
 			if(OrphanedPin->IsStruct())
 			{
-				AddPinsForStruct(OrphanedPin->GetScriptStruct(), OrphanedPin->GetNode(), OrphanedPin, OrphanedPin->Direction, OrphanedPin->GetDefaultValue(), false, true);
+				AddPinsForStruct(OrphanedPin->GetScriptStruct(), OrphanedPin->GetNode(), OrphanedPin, OrphanedPin->Direction, OrphanedPin->GetDefaultValue(), false);
 			}
 				
 			InPin->GetNode()->OrphanedPins.Add(OrphanedPin);
@@ -18865,7 +18928,7 @@ bool URigVMController::ChangePinType(URigVMPin* InPin, TRigVMTypeIndex InTypeInd
 		TArray<URigVMPin*> Pins = InPin->SubPins;
 		for (URigVMPin* Pin : Pins)
 		{
-			RemovePin(Pin, bSetupUndoRedo, true);
+			RemovePin(Pin, bSetupUndoRedo);
 		}
 		
 		InPin->SubPins.Reset();
@@ -18913,7 +18976,7 @@ bool URigVMController::ChangePinType(URigVMPin* InPin, TRigVMTypeIndex InTypeInd
 	{
 		FString DefaultValue = InPin->DefaultValue;
 		CreateDefaultValueForStructIfRequired(InPin->GetScriptStruct(), DefaultValue);
-		AddPinsForStruct(InPin->GetScriptStruct(), InPin->GetNode(), InPin, InPin->Direction, DefaultValue, false, true);
+		AddPinsForStruct(InPin->GetScriptStruct(), InPin->GetNode(), InPin, InPin->Direction, DefaultValue, false);
 	}
 
 	if (InPin->IsArray())
@@ -19104,8 +19167,8 @@ bool URigVMController::ChangePinType(URigVMPin* InPin, TRigVMTypeIndex InTypeInd
 
 	if(Links.Num() > 0)
 	{
-		ReattachLinksToPinObjects(false, &Links, true, true);
-		RemoveUnusedOrphanedPins(InPin->GetNode(), true);
+		ReattachLinksToPinObjects(false, &Links, true);
+		RemoveUnusedOrphanedPins(InPin->GetNode());
 	}
 
 	return true;
@@ -19395,10 +19458,10 @@ void URigVMController::RefreshFunctionReferences(URigVMLibraryNode* InFunctionDe
 			FRigVMControllerGraphGuard GraphGuard(this, ReferenceNode->GetGraph(), bSetupUndoRedo);
 
 			TArray<URigVMLink*> Links = ReferenceNode->GetLinks();
-			DetachLinksFromPinObjects(&Links, true);
-			RepopulatePinsOnNode(ReferenceNode, false, true);
+			DetachLinksFromPinObjects(&Links);
+			RepopulatePinsOnNode(ReferenceNode, false);
 			TGuardValue<bool> ReportGuard(bReportWarningsAndErrors, false);
-			ReattachLinksToPinObjects(false, &Links, true);
+			ReattachLinksToPinObjects(false, &Links);
 		});
 	}
 }
diff --git a/Engine/Source/Developer/RigVMDeveloper/Public/RigVMModel/RigVMController.h b/Engine/Source/Developer/RigVMDeveloper/Public/RigVMModel/RigVMController.h
--- a/Engine/Source/Developer/RigVMDeveloper/Public/RigVMModel/RigVMController.h
+++ b/Engine/Source/Developer/RigVMDeveloper/Public/RigVMModel/RigVMController.h
@@ -818,7 +818,7 @@ public:
 	bool SetLocalVariableTypeFromObjectPath(const FName& InVariableName, const FString& InCPPType, const FString& InCPPTypeObjectPath, bool bSetupUndoRedo = true, bool bPrintPythonCommand = false);
 
 	UFUNCTION(BlueprintCallable, Category = RigVMController)
-	bool SetLocalVariableDefaultValue(const FName& InVariableName, const FString& InDefaultValue, bool bSetupUndoRedo = true, bool bPrintPythonCommand = false, bool bNotify = true);
+	bool SetLocalVariableDefaultValue(const FName& InVariableName, const FString& InDefaultValue, bool bSetupUndoRedo = true, bool bPrintPythonCommand = false);
 
 	// creates the options struct for a given workflow
 	UFUNCTION(BlueprintCallable, Category = RigVMController)
@@ -829,10 +829,10 @@ public:
 	bool PerformUserWorkflow(const FRigVMUserWorkflow& InWorkflow, const URigVMUserWorkflowOptions* InOptions, bool bSetupUndoRedo = true);
 
 	// Determine affected function references for a potential bulk edit on a library node
-	TArray<TSoftObjectPtr<URigVMFunctionReferenceNode>> GetAffectedReferences(ERigVMControllerBulkEditType InEditType, bool bForceLoad = false, bool bNotify = true);
+	TArray<TSoftObjectPtr<URigVMFunctionReferenceNode>> GetAffectedReferences(ERigVMControllerBulkEditType InEditType, bool bForceLoad = false);
 
 	// Determine affected assets for a potential bulk edit on a library node
-	TArray<FAssetData> GetAffectedAssets(ERigVMControllerBulkEditType InEditType, bool bForceLoad = false, bool bNotify = true);
+	TArray<FAssetData> GetAffectedAssets(ERigVMControllerBulkEditType InEditType, bool bForceLoad = false);
 
 	// A delegate that can be set to change the struct unfolding behaviour
 	FRigVMController_ShouldStructUnfoldDelegate UnfoldStructDelegate;
@@ -864,8 +864,8 @@ public:
 	// A delegate to request to configure an options instance for a node workflow
 	FRigVMController_ConfigureWorkflowOptionsDelegate ConfigureWorkflowOptionsDelegate; 
 
-	int32 DetachLinksFromPinObjects(const TArray<URigVMLink*>* InLinks = nullptr, bool bNotify = false);
-	int32 ReattachLinksToPinObjects(bool bFollowCoreRedirectors = false, const TArray<URigVMLink*>* InLinks = nullptr, bool bNotify = false, bool bSetupOrphanedPins = false, bool bAllowNonArgumentLinks = false);
+	int32 DetachLinksFromPinObjects(const TArray<URigVMLink*>* InLinks = nullptr);
+	int32 ReattachLinksToPinObjects(bool bFollowCoreRedirectors = false, const TArray<URigVMLink*>* InLinks = nullptr, bool bSetupOrphanedPins = false, bool bAllowNonArgumentLinks = false);
 	void AddPinRedirector(bool bInput, bool bOutput, const FString& OldPinPath, const FString& NewPinPath);
 
 	// Removes nodes which went stale.
@@ -875,11 +875,11 @@ public:
 	bool ShouldRedirectPin(UScriptStruct* InOwningStruct, const FString& InOldRelativePinPath, FString& InOutNewRelativePinPath) const;
 	bool ShouldRedirectPin(const FString& InOldPinPath, FString& InOutNewPinPath) const;
 
-	void RepopulatePinsOnNode(URigVMNode* InNode, bool bFollowCoreRedirectors = true, bool bNotify = false, bool bSetupOrphanedPins = false);
-	void RemovePinsDuringRepopulate(URigVMNode* InNode, TArray<URigVMPin*>& InPins, bool bNotify, bool bSetupOrphanedPins);
+	void RepopulatePinsOnNode(URigVMNode* InNode, bool bFollowCoreRedirectors = true, bool bSetupOrphanedPins = false);
+	void RemovePinsDuringRepopulate(URigVMNode* InNode, TArray<URigVMPin*>& InPins, bool bSetupOrphanedPins);
 
 	// removes any orphan pins that no longer holds a link
-	bool RemoveUnusedOrphanedPins(URigVMNode* InNode, bool bNotify);
+	bool RemoveUnusedOrphanedPins(URigVMNode* InNode);
 
 	// Initializes and recomputes the filtered permutations of all template nodes in the graph
 	// Returns true if any pin has change it's type or link was broken
@@ -960,13 +960,6 @@ public:
 		ReportAndNotifyError(FString::Printf(Fmt, Args...));
 	}
 
-	/**
-	 * Function to override the notification behavior and temporarily
-	 * disable all notifications. Client code is responsible for calling
-	 * SuspendNotifications(true) once all changes have been done.
-	 */
-	void SuspendNotifications(bool bSuspend) { bSuspendNotifications = bSuspend; }
-
 	/**
 	 * Helper function to disable a series of checks that can be ignored during a unit test
 	 */
@@ -991,7 +984,7 @@ private:
 	TObjectPtr<URigVMNode> FindEventNode(const UScriptStruct* InScriptStruct) const;
 	bool CanAddEventNode(UScriptStruct* InScriptStruct, const bool bReportErrors) const;
 	bool CanAddFunctionRefForDefinition(const FRigVMGraphFunctionHeader& InFunctionDefinition, bool bReportErrors, bool bAllowPrivateFunctions=false);
-	void AddPinsForStruct(UStruct* InStruct, URigVMNode* InNode, URigVMPin* InParentPin, ERigVMPinDirection InPinDirection, const FString& InDefaultValue, bool bAutoExpandArrays, bool bNotify = false);
+	void AddPinsForStruct(UStruct* InStruct, URigVMNode* InNode, URigVMPin* InParentPin, ERigVMPinDirection InPinDirection, const FString& InDefaultValue, bool bAutoExpandArrays);
 	void AddPinsForArray(FArrayProperty* InArrayProperty, URigVMNode* InNode, URigVMPin* InParentPin, ERigVMPinDirection InPinDirection, const TArray<FString>& InDefaultValues, bool bAutoExpandArrays);
 	void AddPinsForTemplate(const FRigVMTemplate* InTemplate, const FRigVMTemplateTypeMap& InPinTypeMap, URigVMNode* InNode);
 	void ConfigurePinFromProperty(FProperty* InProperty, URigVMPin* InOutPin, ERigVMPinDirection InPinDirection = ERigVMPinDirection::Invalid);
@@ -999,12 +992,12 @@ private:
 	void ConfigurePinFromArgument(URigVMPin* InOutPin, const FRigVMGraphFunctionArgument& InArgument, bool bCopyDisplayName = false);
 	virtual bool ShouldStructBeUnfolded(const UStruct* InStruct);
 	virtual bool ShouldPinBeUnfolded(URigVMPin* InPin);
-	bool SetPinDefaultValue(URigVMPin* InPin, const FString& InDefaultValue, bool bResizeArrays, bool bSetupUndoRedo, bool bMergeUndoAction, bool bNotify = true);
+	bool SetPinDefaultValue(URigVMPin* InPin, const FString& InDefaultValue, bool bResizeArrays, bool bSetupUndoRedo, bool bMergeUndoAction);
 	bool ResetPinDefaultValue(URigVMPin* InPin, bool bSetupUndoRedo);
 	static FString GetPinInitialDefaultValue(const URigVMPin* InPin);
 	static FString GetPinInitialDefaultValueFromStruct(UScriptStruct* ScriptStruct, const URigVMPin* InPin, uint32 InOffset);
 	URigVMPin* InsertArrayPin(URigVMPin* ArrayPin, int32 InIndex, const FString& InDefaultValue, bool bSetupUndoRedo);
-	bool RemovePin(URigVMPin* InPinToRemove, bool bSetupUndoRedo, bool bNotify);
+	bool RemovePin(URigVMPin* InPinToRemove, bool bSetupUndoRedo);
 	FProperty* FindPropertyForPin(const FString& InPinPath);
 	bool BindPinToVariable(URigVMPin* InPin, const FString& InNewBoundVariablePath, bool bSetupUndoRedo, const FString& InVariableNodeName = FString());
 	bool UnbindPinFromVariable(URigVMPin* InPin, bool bSetupUndoRedo);
@@ -1039,7 +1032,7 @@ private:
 	URigVMCollapseNode* PromoteFunctionReferenceNodeToCollapseNode(URigVMFunctionReferenceNode* InFunctionRefNode, bool bSetupUndoRedo, bool bRemoveFunctionDefinition);
 	void SetReferencedFunction(URigVMFunctionReferenceNode* InFunctionRefNode, URigVMLibraryNode* InNewReferencedNode, bool bSetupUndoRedo);
 
-	void RefreshFunctionPins(URigVMNode* InNode, bool bNotify = true);
+	void RefreshFunctionPins(URigVMNode* InNode);
 
 	void ReportRemovedLink(const FString& InSourcePinPath, const FString& InTargetPinPath);
 
@@ -1256,6 +1249,7 @@ private:
 	friend class FRigVMParserAST;
 	friend class FRigVMControllerCompileBracketScope;
 	friend class FRigVMControllerGraphGuard;
+	friend class FRigVMControllerNotifGuard;
 	friend struct FRigVMClient;
 };
 
@@ -1290,6 +1284,28 @@ private:
 	int32 NumGraphs;
 };
 
+class FRigVMControllerNotifGuard
+{
+public:
+
+	FRigVMControllerNotifGuard(URigVMController* InController, bool bInSuspendNotifications = true)
+		: Controller(InController)
+	{
+		bPreviousSuspendNotifications = Controller->bSuspendNotifications;
+		Controller->bSuspendNotifications = bInSuspendNotifications;
+	}
+
+	~FRigVMControllerNotifGuard()
+	{
+		Controller->bSuspendNotifications = bPreviousSuspendNotifications;
+	}
+
+private:
+
+	URigVMController* Controller;
+	bool bPreviousSuspendNotifications;
+};
+
 USTRUCT()
 struct FRigVMController_CommonTypePerTemplate
 {
