RigVM: refactor: prepare for new select, if and control flow nodes

diff --git a/Engine/Plugins/Animation/ControlRig/Source/ControlRig/Private/Units/Animation/RigUnit_AnimAttribute.cpp b/Engine/Plugins/Animation/ControlRig/Source/ControlRig/Private/Units/Animation/RigUnit_AnimAttribute.cpp
--- a/Engine/Plugins/Animation/ControlRig/Source/ControlRig/Private/Units/Animation/RigUnit_AnimAttribute.cpp
+++ b/Engine/Plugins/Animation/ControlRig/Source/ControlRig/Private/Units/Animation/RigUnit_AnimAttribute.cpp
@@ -471,7 +471,7 @@ TArray<FRigVMTemplateArgument> FRigDispatch_SetAnimAttribute::GetArguments() con
 	return Arguments;
 }
 
-TArray<FRigVMExecuteArgument> FRigDispatch_SetAnimAttribute::GetExecuteArguments_Impl() const
+TArray<FRigVMExecuteArgument> FRigDispatch_SetAnimAttribute::GetExecuteArguments_Impl(const FRigVMDispatchContext& InContext) const
 {
 	return {{TEXT("ExecuteContext"), ERigVMPinDirection::IO}};
 }
diff --git a/Engine/Plugins/Animation/ControlRig/Source/ControlRig/Private/Units/Animation/RigUnit_AnimAttribute.h b/Engine/Plugins/Animation/ControlRig/Source/ControlRig/Private/Units/Animation/RigUnit_AnimAttribute.h
--- a/Engine/Plugins/Animation/ControlRig/Source/ControlRig/Private/Units/Animation/RigUnit_AnimAttribute.h
+++ b/Engine/Plugins/Animation/ControlRig/Source/ControlRig/Private/Units/Animation/RigUnit_AnimAttribute.h
@@ -257,7 +257,7 @@ struct CONTROLRIG_API FRigDispatch_SetAnimAttribute: public FRigDispatch_AnimAtt
 	GENERATED_BODY()
 	virtual bool IsSet() const override { return true; }
 	virtual TArray<FRigVMTemplateArgument> GetArguments() const override;
-	virtual TArray<FRigVMExecuteArgument> GetExecuteArguments_Impl() const override;
+	virtual TArray<FRigVMExecuteArgument> GetExecuteArguments_Impl(const FRigVMDispatchContext& InContext) const override;
 	virtual FRigVMTemplateTypeMap OnNewArgumentType(const FName& InArgumentName, TRigVMTypeIndex InTypeIndex) const override;
 	
 protected:
diff --git a/Engine/Plugins/Animation/ControlRig/Source/ControlRig/Private/Units/Hierarchy/RigUnit_Metadata.cpp b/Engine/Plugins/Animation/ControlRig/Source/ControlRig/Private/Units/Hierarchy/RigUnit_Metadata.cpp
--- a/Engine/Plugins/Animation/ControlRig/Source/ControlRig/Private/Units/Hierarchy/RigUnit_Metadata.cpp
+++ b/Engine/Plugins/Animation/ControlRig/Source/ControlRig/Private/Units/Hierarchy/RigUnit_Metadata.cpp
@@ -283,7 +283,7 @@ TArray<FRigVMTemplateArgument> FRigDispatch_SetMetadata::GetArguments() const
 	return Arguments;
 }
 
-TArray<FRigVMExecuteArgument> FRigDispatch_SetMetadata::GetExecuteArguments_Impl() const
+TArray<FRigVMExecuteArgument> FRigDispatch_SetMetadata::GetExecuteArguments_Impl(const FRigVMDispatchContext& InContext) const
 {
 	return {{TEXT("ExecuteContext"), ERigVMPinDirection::IO}};
 }
diff --git a/Engine/Plugins/Animation/ControlRig/Source/ControlRig/Private/Units/Hierarchy/RigUnit_Metadata.h b/Engine/Plugins/Animation/ControlRig/Source/ControlRig/Private/Units/Hierarchy/RigUnit_Metadata.h
--- a/Engine/Plugins/Animation/ControlRig/Source/ControlRig/Private/Units/Hierarchy/RigUnit_Metadata.h
+++ b/Engine/Plugins/Animation/ControlRig/Source/ControlRig/Private/Units/Hierarchy/RigUnit_Metadata.h
@@ -114,7 +114,7 @@ struct CONTROLRIG_API FRigDispatch_SetMetadata : public FRigDispatch_MetadataBas
 	GENERATED_BODY()
 
 	virtual TArray<FRigVMTemplateArgument> GetArguments() const override;
-	virtual TArray<FRigVMExecuteArgument> GetExecuteArguments_Impl() const override;
+	virtual TArray<FRigVMExecuteArgument> GetExecuteArguments_Impl(const FRigVMDispatchContext& InContext) const override;
 	virtual bool IsSetMetadata() const override { return true; }
 
 protected:
diff --git a/Engine/Plugins/Animation/ControlRig/Source/ControlRigEditor/Private/Graph/NodeSpawners/ControlRigTemplateNodeSpawner.cpp b/Engine/Plugins/Animation/ControlRig/Source/ControlRigEditor/Private/Graph/NodeSpawners/ControlRigTemplateNodeSpawner.cpp
--- a/Engine/Plugins/Animation/ControlRig/Source/ControlRigEditor/Private/Graph/NodeSpawners/ControlRigTemplateNodeSpawner.cpp
+++ b/Engine/Plugins/Animation/ControlRig/Source/ControlRigEditor/Private/Graph/NodeSpawners/ControlRigTemplateNodeSpawner.cpp
@@ -16,6 +16,7 @@
 #include "ControlRigBlueprintUtils.h"
 #include "ScopedTransaction.h"
 #include "ControlRig.h"
+#include "RigVMModel/Nodes/RigVMDispatchNode.h"
 
 #include UE_INLINE_GENERATED_CPP_BY_NAME(ControlRigTemplateNodeSpawner)
 
@@ -100,7 +101,13 @@ bool UControlRigTemplateNodeSpawner::IsTemplateNodeFilteredOut(FBlueprintActionF
 								{
 									if(ModelPin->IsExecuteContext())
 									{
-										if(Template->NumExecuteArguments() > 0)
+										FRigVMDispatchContext DispatchContext;
+										if(const URigVMDispatchNode* DispatchNode = Cast<URigVMDispatchNode>(ModelPin->GetNode()))
+										{
+											DispatchContext = DispatchNode->GetDispatchContext();
+										}
+
+										if(Template->NumExecuteArguments(DispatchContext) > 0)
 										{
 											return false;
 										}
diff --git a/Engine/Plugins/Animation/ControlRig/Source/ControlRigEditor/Private/Graph/SControlRigGraphChangePinType.cpp b/Engine/Plugins/Animation/ControlRig/Source/ControlRigEditor/Private/Graph/SControlRigGraphChangePinType.cpp
--- a/Engine/Plugins/Animation/ControlRig/Source/ControlRigEditor/Private/Graph/SControlRigGraphChangePinType.cpp
+++ b/Engine/Plugins/Animation/ControlRig/Source/ControlRigEditor/Private/Graph/SControlRigGraphChangePinType.cpp
@@ -107,7 +107,18 @@ FText SControlRigChangePinType::GetBindingText(URigVMPin* ModelPin) const
 {
 	if (ModelPin)
 	{
-		return GetBindingText(ModelPin->GetTemplateArgumentType());
+		FRigVMTemplateArgumentType ArgumentType = ModelPin->GetRootPin()->GetTemplateArgumentType();
+		if(ModelPin->IsArrayElement())
+		{
+			const FRigVMRegistry& Registry = FRigVMRegistry::Get();
+			const TRigVMTypeIndex ArrayTypeIndex = Registry.GetTypeIndex(ArgumentType);
+			const TRigVMTypeIndex ElementTypeIndex = Registry.GetBaseTypeFromArrayTypeIndex(ArrayTypeIndex);
+			ArgumentType = Registry.GetType(ElementTypeIndex);
+		}
+		if(ArgumentType.IsValid())
+		{
+			return GetBindingText(ArgumentType);
+		}
 	}
 	return FText();
 }
@@ -197,11 +208,17 @@ void SControlRigChangePinType::FillPinTypeMenu(FMenuBuilder& MenuBuilder)
 	TMap<TRigVMTypeIndex, FArgumentInfo> Types;
 	for(URigVMPin* ModelPin : ModelPins)
 	{
+		const bool bIsArrayElement = ModelPin->IsArrayElement();
+		if(bIsArrayElement)
+		{
+			ModelPin = ModelPin->GetParentPin();
+		}
+
 		if(!ModelPin->IsRootPin())
 		{
 			continue;
 		}
-		
+
 		if(const URigVMTemplateNode* TemplateNode = Cast<URigVMTemplateNode>(ModelPin->GetNode()))
 		{
 			if(const FRigVMTemplate* Template = TemplateNode->GetTemplate())
@@ -253,6 +270,11 @@ void SControlRigChangePinType::FillPinTypeMenu(FMenuBuilder& MenuBuilder)
 								}
 							}
 						}
+
+						if(bIsArrayElement)
+						{
+							ArgumentTypeIndex = FRigVMRegistry::Get().GetBaseTypeFromArrayTypeIndex(ArgumentTypeIndex);
+						}
 						
 						const FRigVMTemplateArgumentType ArgumentType = FRigVMRegistry::Get().GetType(ArgumentTypeIndex);
 						Types.Add(ArgumentTypeIndex, FArgumentInfo::Make(ArgumentType, bIsFilteredOut));
@@ -406,8 +428,19 @@ void SControlRigChangePinType::HandlePinTypeChanged(FRigVMTemplateArgumentType I
 			continue;
 		}
 
-		const TRigVMTypeIndex TypeIndex = FRigVMRegistry::Get().GetTypeIndex(InType);
-		Controller->ResolveWildCardPin(ModelPin->GetPinPath(), TypeIndex, true, true);
+		TRigVMTypeIndex TypeIndex = FRigVMRegistry::Get().GetTypeIndex(InType);
+		if(ModelPin->IsArrayElement())
+		{
+			TypeIndex = FRigVMRegistry::Get().GetArrayTypeFromBaseTypeIndex(TypeIndex);
+			ModelPin = ModelPin->GetParentPin();
+		}
+
+		if(!ModelPin->IsRootPin())
+		{
+			continue;
+		}
+		
+		Controller->ResolveWildCardPin(ModelPin->GetRootPin()->GetPinPath(), TypeIndex, true, true);
 	}
 }
 
diff --git a/Engine/Plugins/Animation/ControlRig/Source/ControlRigEditor/Private/Graph/SControlRigGraphNode.cpp b/Engine/Plugins/Animation/ControlRig/Source/ControlRigEditor/Private/Graph/SControlRigGraphNode.cpp
--- a/Engine/Plugins/Animation/ControlRig/Source/ControlRigEditor/Private/Graph/SControlRigGraphNode.cpp
+++ b/Engine/Plugins/Animation/ControlRig/Source/ControlRigEditor/Private/Graph/SControlRigGraphNode.cpp
@@ -34,7 +34,10 @@
 #include "Styling/AppStyle.h"
 #include "Graph/SControlRigGraphPinVariableBinding.h"
 #include "RigVMModel/Nodes/RigVMAggregateNode.h"
+#include "RigVMModel/Nodes/RigVMDispatchNode.h"
 #include "Slate/SlateTextures.h"
+//#include "RigVMFunctions/RigVMDispatch_If.h"
+//#include "RigVMFunctions/RigVMDispatch_Select.h"
 
 #if WITH_EDITOR
 #include "Editor.h"
@@ -991,11 +994,28 @@ void SControlRigGraphNode::RefreshErrorInfo()
 			{
 				if(const URigVMNode* RigModelNode = RigGraphNode->GetModelNode())
 				{
-					if(RigModelNode->IsA<URigVMIfNode>() || RigModelNode->IsA<URigVMSelectNode>())
+					bool bShowCopyWarning = RigModelNode->IsA<URigVMIfNode>() || RigModelNode->IsA<URigVMSelectNode>();
+					if(!bShowCopyWarning)
+					{
+						/*
+						if(const URigVMDispatchNode* DispatchNode = Cast<URigVMDispatchNode>(RigModelNode))
+						{
+							if(const UScriptStruct* FactoryStruct = DispatchNode->GetScriptStruct())
+							{
+								if(FactoryStruct == FRigVMDispatch_If::StaticStruct() ||
+									FactoryStruct == FRigVMDispatch_SelectInt32::StaticStruct())
+								{
+									bShowCopyWarning = true;
+								}
+							}
+						}
+						*/
+					}
+					if(bShowCopyWarning)
 					{
 						for(const URigVMPin* Pin : RigModelNode->GetPins())
 						{
-							if(Pin->IsArray())
+							if(Pin->IsArray() && !Pin->IsFixedSizeArray())
 							{
 								GraphNode->bHasCompilerMessage = true;
 								GraphNode->ErrorType = int32(EMessageSeverity::Info);
@@ -1257,9 +1277,12 @@ EVisibility SControlRigGraphNode::GetPinVisibility(int32 InPinInfoIndex, bool bA
 				return ParentPinVisibility;
 			}
 
-			if(!PinInfos[ParentPinIndex].bExpanded)
+			if(!PinInfos[InPinInfoIndex].bFixedArray)
 			{
-				return EVisibility::Collapsed;
+				if(!PinInfos[ParentPinIndex].bExpanded)
+				{
+					return EVisibility::Collapsed;
+				}
 			}
 		}
 	}
@@ -1556,6 +1579,10 @@ void SControlRigGraphNode::UpdatePinTreeView()
 			}
 			PinInfo.ParentIndex = *ParentIndexPtr;
 			PinInfo.Depth = PinInfos[PinInfo.ParentIndex].Depth + 1;
+			if(PinInfos[PinInfo.ParentIndex].bFixedArray)
+			{
+				PinInfo.Depth--;
+			}
 		}
 
 		TAttribute<EVisibility> PinVisibilityAttribute = TAttribute<EVisibility>::CreateSP(this, &SControlRigGraphNode::GetPinVisibility, PinInfo.Index, false);
diff --git a/Engine/Source/Developer/RigVMDeveloper/Private/RigVMCompiler/RigVMAST.cpp b/Engine/Source/Developer/RigVMDeveloper/Private/RigVMCompiler/RigVMAST.cpp
--- a/Engine/Source/Developer/RigVMDeveloper/Private/RigVMCompiler/RigVMAST.cpp
+++ b/Engine/Source/Developer/RigVMDeveloper/Private/RigVMCompiler/RigVMAST.cpp
@@ -20,6 +20,7 @@
 #include "Stats/StatsHierarchical.h"
 #include "RigVMDeveloperModule.h"
 #include "VisualGraphUtils.h"
+#include "RigVMModel/Nodes/RigVMDispatchNode.h"
 #include "UObject/FieldIterator.h"
 
 #include UE_INLINE_GENERATED_CPP_BY_NAME(RigVMAST)
@@ -1100,32 +1101,37 @@ TArray<FRigVMExprAST*> FRigVMParserAST::TraversePins(const FRigVMASTProxy& InNod
 		Pins = Node->GetPins();
 	}
 
-	for (URigVMPin* Pin : Pins)
+	// dispatch nodes may contain a fixed size array
+	if(Node->IsA<URigVMDispatchNode>() || Node->IsA<URigVMSelectNode>())
 	{
-		FRigVMASTProxy PinProxy = InNodeProxy.GetSibling(Pin);
-
-		if (InParentExpr &&
-		    Pin->GetDirection() == ERigVMPinDirection::Input &&
-		    InParentExpr->IsA(FRigVMExprAST::EType::Select))
+		TArray<URigVMPin*> FlattenedPins;
+		FlattenedPins.Reserve(Pins.Num());
+		for(URigVMPin* Pin : Pins)
 		{
-			if (Pin->GetName() == URigVMSelectNode::ValueName)
+			if(Pin->IsFixedSizeArray())
 			{
-				const TArray<URigVMPin*>& CasePins = Pin->GetSubPins();
-				for (URigVMPin* CasePin : CasePins)
-				{
-					FRigVMASTProxy CasePinProxy = PinProxy.GetSibling(CasePin);
-
-					PinExpressions.Add(TraversePin(CasePinProxy, InParentExpr));
-
-					const FString PinName = FString::Printf(TEXT("%s_%s"), *Pin->GetName(), *CasePin->GetName());
-					const int32 ChildIndex = InParentExpr->Children.Find(PinExpressions.Last());
-					InParentExpr->PinNameToChildIndex.FindOrAdd(*PinName) = ChildIndex;
-				}
-				continue;
+				FlattenedPins.Append(Pin->GetSubPins());
+			}
+			else
+			{
+				FlattenedPins.Add(Pin);
 			}
 		}
+		Swap(Pins, FlattenedPins);
+	}
 
+	for (URigVMPin* Pin : Pins)
+	{
+		FRigVMASTProxy PinProxy = InNodeProxy.GetSibling(Pin);
 		PinExpressions.Add(TraversePin(PinProxy, InParentExpr));
+
+		if (InParentExpr && !Pin->IsRootPin())
+		{
+			const FString PinName = FString::Printf(TEXT("%s_%s"), *Pin->GetParentPin()->GetName(), *Pin->GetName());
+			const int32 ChildIndex = InParentExpr->Children.Find(PinExpressions.Last());
+			InParentExpr->PinNameToChildIndex.FindOrAdd(*PinName) = ChildIndex;
+		}
+
 		if(InParentExpr)
 		{
 			const int32 ChildIndex = InParentExpr->Children.Find(PinExpressions.Last());
diff --git a/Engine/Source/Developer/RigVMDeveloper/Private/RigVMCompiler/RigVMCodeGenerator.cpp b/Engine/Source/Developer/RigVMDeveloper/Private/RigVMCompiler/RigVMCodeGenerator.cpp
--- a/Engine/Source/Developer/RigVMDeveloper/Private/RigVMCompiler/RigVMCodeGenerator.cpp
+++ b/Engine/Source/Developer/RigVMDeveloper/Private/RigVMCompiler/RigVMCodeGenerator.cpp
@@ -6,6 +6,7 @@
 #include "RigVMDeveloperModule.h"
 #include "Algo/Count.h"
 #include "Animation/Rig.h"
+#include "RigVMModel/Nodes/RigVMDispatchNode.h"
 
 static constexpr TCHAR RigVM_CommaSeparator[] = TEXT(", ");
 static constexpr TCHAR RigVM_NewLineFormat[] = TEXT("\r\n");
@@ -501,16 +502,15 @@ FString FRigVMCodeGenerator::DumpDispatches(bool bLog)
 
 	if(!Dispatches.IsEmpty())
 	{
-		for(auto Pair : Dispatches)
+		for(const FRigVMDispatchInfo& Info : Dispatches)
 		{
-			const FRigVMFunction* Function = Pair.Value;
-			const FString FunctionName = Function->GetName();
-			const FRigVMDispatchFactory* Factory = Function->Factory;
+			const FString FunctionName = Info.Function->GetName();
+			const FRigVMDispatchFactory* Factory = Info.Function->Factory;
 
 			TArray<FString> InputArguments, OutputArguments, MemoryHandles;
 			OutputArguments.Add(RigVM_ContextFormat);
 
-			for(const FRigVMFunctionArgument& Argument : Function->Arguments)
+			for(const FRigVMFunctionArgument& Argument : Info.Function->Arguments)
 			{
 				const TRigVMTypeIndex TypeIndex = FRigVMRegistry::Get().GetTypeIndexFromCPPType(Argument.Type);
 				const FString DispatchArgument = RequiredUProperties.FindChecked(TypeIndex).Get<1>();
@@ -532,7 +532,7 @@ FString FRigVMCodeGenerator::DumpDispatches(bool bLog)
 			OutputArguments.Add(TEXT("MemoryHandles"));
 			
 			Lines.Emplace();
-			Lines.Add(Format(RigVM_DispatchDeclarationFormat, Pair.Key, FString::Join(InputArguments, RigVM_CommaSeparator), *FunctionName));
+			Lines.Add(Format(RigVM_DispatchDeclarationFormat, Info.Name, FString::Join(InputArguments, RigVM_CommaSeparator), *FunctionName));
 			Lines.Add(Format(TEXT("\t\tTArray<FRigVMMemoryHandle> MemoryHandles = {\r\n\t\t\t{0}\r\n\t\t};"), FString::Join(MemoryHandles, TEXT(",\r\n\t\t\t"))));
 			Lines.Add(Format(RigVM_InvokeDispatchPtrFormat, FString::Join(OutputArguments, RigVM_CommaSeparator)));
 			Lines.Add(TEXT("\t\treturn true;\r\n\t}"));
@@ -827,11 +827,11 @@ FString FRigVMCodeGenerator::DumpInstructions(const FString& InPrefix, const TAr
 				else if(Function->Factory)
 				{
 					FString DispatchName;
-					for(auto Pair : Dispatches)
+					for(const FRigVMDispatchInfo& Info : Dispatches)
 					{
-						if(Pair.Value == Function)
+						if(Info.Function == Function)
 						{
-							DispatchName = Pair.Key;
+							DispatchName = Info.Name;
 							break;
 						}
 					}
@@ -1457,8 +1457,10 @@ void FRigVMCodeGenerator::ParseRequiredUProperties()
 		}
 	};
 
-	for(const FRigVMInstruction& Instruction : Instructions)
+	for(int32 InstructionIndex = 0; InstructionIndex < Instructions.Num(); InstructionIndex++)
 	{
+		const FRigVMInstruction& Instruction = Instructions[InstructionIndex];
+		
 		if(Instruction.OpCode >= ERigVMOpCode::Execute_0_Operands &&
 			Instruction.OpCode <= ERigVMOpCode::Execute_64_Operands)
 		{
@@ -1470,7 +1472,14 @@ void FRigVMCodeGenerator::ParseRequiredUProperties()
 				const int32 PermutationIndex = Factory->GetTemplate()->FindPermutation(Function); 
 				check(PermutationIndex != INDEX_NONE);
 				const FString DispatchKey = Format(RigVM_DispatchKeyFormat, Factory->GetFactoryName().ToString(), PermutationIndex);
-				Dispatches.Add(DispatchKey, Function);
+
+				FRigVMDispatchContext Context;
+				if(URigVMDispatchNode* DispatchNode = Cast<URigVMDispatchNode>(ByteCode.GetSubjectForInstruction(InstructionIndex)))
+				{
+					Context = DispatchNode->GetDispatchContext();
+				}
+
+				Dispatches.Add({DispatchKey, Function, Context});
 
 				for(const FRigVMFunctionArgument& Argument : Function->Arguments)
 				{
diff --git a/Engine/Source/Developer/RigVMDeveloper/Private/RigVMCompiler/RigVMCompiler.cpp b/Engine/Source/Developer/RigVMDeveloper/Private/RigVMCompiler/RigVMCompiler.cpp
--- a/Engine/Source/Developer/RigVMDeveloper/Private/RigVMCompiler/RigVMCompiler.cpp
+++ b/Engine/Source/Developer/RigVMDeveloper/Private/RigVMCompiler/RigVMCompiler.cpp
@@ -710,9 +710,10 @@ bool URigVMCompiler::Compile(TArray<URigVMGraph*> InGraphs, URigVMController* In
 				
 				if(const FRigVMTemplate* Template = TemplateNode->GetTemplate())
 				{
-					for(int32 Index = 0; Index < Template->NumExecuteArguments(); Index++)
+					const FRigVMDispatchContext DispatchContext;
+					for(int32 Index = 0; Index < Template->NumExecuteArguments(DispatchContext); Index++)
 					{
-						if(const FRigVMExecuteArgument* Argument = Template->GetExecuteArgument(Index))
+						if(const FRigVMExecuteArgument* Argument = Template->GetExecuteArgument(Index, DispatchContext))
 						{
 							if(Registry.IsExecuteType(Argument->TypeIndex))
 							{
@@ -1524,38 +1525,62 @@ int32 URigVMCompiler::TraverseCallExtern(const FRigVMCallExternExprAST* InExpr,
 		FRigVMOperand BlockToRunOperand;
 		for(const FRigVMFunctionArgument& Argument : Function->GetArguments())
 		{
-			const FRigVMExprAST* ChildExpr = InExpr->FindExprWithPinName(Argument.Name);
-			if(ChildExpr == nullptr)
-			{
-				// opaque arguments don't have a matching child expression
-				continue;
-			}
-			
-			if (ChildExpr->GetType() == FRigVMExprAST::EType::CachedValue)
-			{
-				Operands.Add(WorkData.ExprToOperand.FindChecked(GetSourceVarExpr(ChildExpr->To<FRigVMCachedValueExprAST>()->GetVarExpr())));
-			}
-			else if (ChildExpr->IsA(FRigVMExprAST::EType::Var))
+			auto ProcessArgument = [
+				&WorkData,
+				Argument,
+				&Operands,
+				&BlockToRunOperand,
+				&CountOperand,
+				&IndexOperand
+			](const FRigVMExprAST* InExpr)
 			{
-				Operands.Add(WorkData.ExprToOperand.FindChecked(GetSourceVarExpr(ChildExpr->To<FRigVMVarExprAST>())));
-			}
-			else
-			{
-				break;
-			}
+				if (InExpr->GetType() == FRigVMExprAST::EType::CachedValue)
+				{
+					Operands.Add(WorkData.ExprToOperand.FindChecked(GetSourceVarExpr(InExpr->To<FRigVMCachedValueExprAST>()->GetVarExpr())));
+				}
+				else if (InExpr->IsA(FRigVMExprAST::EType::Var))
+				{
+					Operands.Add(WorkData.ExprToOperand.FindChecked(GetSourceVarExpr(InExpr->To<FRigVMVarExprAST>())));
+				}
+				else
+				{
+					return;
+				}
 
-			if(Argument.Name == FRigVMStruct::ControlFlowBlockToRunName)
-			{
-				BlockToRunOperand = Operands.Last();
-			}
-			else if(Argument.Name == FRigVMStruct::ControlFlowCountName)
-			{
-				CountOperand = Operands.Last();
-			}
-			else if(Argument.Name == FRigVMStruct::ControlFlowIndexName)
+				if(Argument.Name == FRigVMStruct::ControlFlowBlockToRunName)
+				{
+					BlockToRunOperand = Operands.Last();
+				}
+				else if(Argument.Name == FRigVMStruct::ControlFlowCountName)
+				{
+					CountOperand = Operands.Last();
+				}
+				else if(Argument.Name == FRigVMStruct::ControlFlowIndexName)
+				{
+					IndexOperand = Operands.Last();
+				}
+				
+			};
+
+			if(URigVMPin* Pin = InExpr->GetNode()->FindPin(Argument.Name))
 			{
-				IndexOperand = Operands.Last();
+				if(Pin->IsFixedSizeArray())
+				{
+					for(URigVMPin* SubPin : Pin->GetSubPins())
+					{
+						static constexpr TCHAR Format[] = TEXT("%s_%s");
+						const FString PinName = FString::Printf(Format, *Pin->GetName(), *SubPin->GetName());
+						const FRigVMExprAST* SubPinExpr = InExpr->FindExprWithPinName(*PinName);
+						check(SubPinExpr);
+						ProcessArgument(SubPinExpr);
+					}
+					continue;
+				}
 			}
+			
+			const FRigVMExprAST* ChildExpr = InExpr->FindExprWithPinName(Argument.Name);
+			check(ChildExpr);
+			ProcessArgument(ChildExpr);
 		}
 
 		// make sure to skip the output blocks while we are traversing this call extern
@@ -2164,32 +2189,29 @@ void URigVMCompiler::TraverseAssign(const FRigVMAssignExprAST* InExpr, FRigVMCom
 						}
 					}
 
-					// for select nodes we create a register for each case (since the cases are fixed in size)
+					// for fixed array pins we create a register for each array element
 					// thus we do not need to setup a registeroffset for the array element.
-					if (URigVMSelectNode* SelectNode = Cast<URigVMSelectNode>(RootPin->GetNode()))
+					if (RootPin->IsFixedSizeArray())
 					{
-						if(RootPin->GetName() == URigVMSelectNode::ValueName)
+						if (Pin->GetParentPin() == RootPin)
 						{
-							if (Pin->GetParentPin() == RootPin)
-							{
-								return;
-							}
+							return;
+						}
 
-							// if the pin is a sub pin of a case of the select (for example: Values.0.Translation)
-							// we'll need to re-adjust the root pin to the case pin (for example: Values.0)
-							TArray<FString> SegmentPathPaths;
-							if(ensure(URigVMPin::SplitPinPath(SegmentPath, SegmentPathPaths)))
-							{
-								RootPin = RootPin->FindSubPin(SegmentPathPaths[0]);
+						// if the pin is a sub pin of a case of a fixed array
+						// we'll need to re-adjust the root pin to the case pin (for example: Values.0)
+						TArray<FString> SegmentPathPaths;
+						if(ensure(URigVMPin::SplitPinPath(SegmentPath, SegmentPathPaths)))
+						{
+							RootPin = RootPin->FindSubPin(SegmentPathPaths[0]);
 
-								SegmentPathPaths.RemoveAt(0);
-								ensure(SegmentPathPaths.Num() > 0);
-								SegmentPath = URigVMPin::JoinPinPath(SegmentPathPaths);
-							}
-							else
-							{
-								return;
-							}
+							SegmentPathPaths.RemoveAt(0);
+							ensure(SegmentPathPaths.Num() > 0);
+							SegmentPath = URigVMPin::JoinPinPath(SegmentPathPaths);
+						}
+						else
+						{
+							return;
 						}
 					}
 
diff --git a/Engine/Source/Developer/RigVMDeveloper/Private/RigVMModel/Nodes/RigVMDispatchNode.cpp b/Engine/Source/Developer/RigVMDeveloper/Private/RigVMModel/Nodes/RigVMDispatchNode.cpp
--- a/Engine/Source/Developer/RigVMDeveloper/Private/RigVMModel/Nodes/RigVMDispatchNode.cpp
+++ b/Engine/Source/Developer/RigVMDeveloper/Private/RigVMModel/Nodes/RigVMDispatchNode.cpp
@@ -54,7 +54,7 @@ const TArray<FName>& URigVMDispatchNode::GetControlFlowBlocks() const
 {
 	if (const FRigVMDispatchFactory* Factory = GetFactory())
 	{
-		return Factory->GetControlFlowBlocks();
+		return Factory->GetControlFlowBlocks(GetDispatchContext());
 	}
 	return Super::GetControlFlowBlocks();
 }
@@ -96,6 +96,34 @@ FString URigVMDispatchNode::GetDeprecatedMetadata() const
 	return FString();
 }
 
+FRigVMDispatchContext URigVMDispatchNode::GetDispatchContext() const
+{
+	FRigVMDispatchContext Context;
+	Context.Instance = ConstructFactoryInstance(false, &Context.StringRepresentation);
+	Context.Subject = this;
+	return Context;
+}
+
+TSharedPtr<FStructOnScope> URigVMDispatchNode::ConstructFactoryInstance(bool bUseDefault, FString* OutFactoryDefault) const
+{
+	if (UScriptStruct* Struct = (UScriptStruct*)GetFactoryStruct())
+	{
+		TSharedPtr<FStructOnScope> StructOnScope = MakeShareable(new FStructOnScope(Struct));
+		FRigVMStruct* StructMemory = (FRigVMStruct*)StructOnScope->GetStructMemory();
+		if (!bUseDefault)
+		{
+			const FString FactoryDefaultValue = GetFactoryDefaultValue();
+			if(OutFactoryDefault)
+			{
+				*OutFactoryDefault = FactoryDefaultValue;
+			}
+			Struct->ImportText(*FactoryDefaultValue, StructMemory, nullptr, PPF_IncludeTransient, GLog, Struct->GetName());
+		}
+		return StructOnScope;
+	}
+	return nullptr;
+}
+
 TArray<URigVMPin*> URigVMDispatchNode::GetAggregateInputs() const
 {
 	TArray<URigVMPin*> AggregateInputs;
@@ -204,11 +232,38 @@ bool URigVMDispatchNode::ShouldInputPinComputeLazily(const URigVMPin* InPin) con
 	return Super::ShouldInputPinComputeLazily(InPin);
 }
 
+FString URigVMDispatchNode::GetFactoryDefaultValue() const
+{
+	TArray<FString> PinDefaultValues;
+	for (URigVMPin* Pin : GetPins())
+	{
+		if (Pin->GetDirection() == ERigVMPinDirection::Hidden)
+		{
+			continue;
+		}
+		FString PinDefaultValue = Pin->GetDefaultValue();
+		if (Pin->IsStringType())
+		{
+			PinDefaultValue = TEXT("\"") + PinDefaultValue + TEXT("\"");
+		}
+		else if (PinDefaultValue.IsEmpty() || PinDefaultValue == TEXT("()"))
+		{
+			continue;
+		}
+		PinDefaultValues.Add(FString::Printf(TEXT("%s=%s"), *Pin->GetName(), *PinDefaultValue));
+	}
+	if (PinDefaultValues.Num() == 0)
+	{
+		return TEXT("()");
+	}
+	return FString::Printf(TEXT("(%s)"), *FString::Join(PinDefaultValues, TEXT(",")));
+}
+
 FRigVMStructUpgradeInfo URigVMDispatchNode::GetUpgradeInfo() const
 {
 	if(const FRigVMDispatchFactory* Factory = GetFactory())
 	{
-		return Factory->GetUpgradeInfo(GetFilteredTypes());
+		return Factory->GetUpgradeInfo(GetFilteredTypes(), GetDispatchContext());
 	}
 	return FRigVMStructUpgradeInfo();
 }
diff --git a/Engine/Source/Developer/RigVMDeveloper/Private/RigVMModel/RigVMController.cpp b/Engine/Source/Developer/RigVMDeveloper/Private/RigVMModel/RigVMController.cpp
--- a/Engine/Source/Developer/RigVMDeveloper/Private/RigVMModel/RigVMController.cpp
+++ b/Engine/Source/Developer/RigVMDeveloper/Private/RigVMModel/RigVMController.cpp
@@ -14391,10 +14391,16 @@ void URigVMController::AddPinsForArray(FArrayProperty* InArrayProperty, URigVMNo
 void URigVMController::AddPinsForTemplate(const FRigVMTemplate* InTemplate, const FRigVMTemplateTypeMap& InPinTypeMap, URigVMNode* InNode)
 {
 	const FRigVMRegistry& Registry = FRigVMRegistry::Get();
+
+	FRigVMDispatchContext DispatchContext;
+	if(const URigVMDispatchNode* DispatchNode = Cast<URigVMDispatchNode>(InNode))
+	{
+		DispatchContext = DispatchNode->GetDispatchContext();
+	}
 	
-	for (int32 ArgIndex = 0; ArgIndex < InTemplate->NumExecuteArguments(); ArgIndex++)
+	for (int32 ArgIndex = 0; ArgIndex < InTemplate->NumExecuteArguments(DispatchContext); ArgIndex++)
     {
-        const FRigVMExecuteArgument* Arg = InTemplate->GetExecuteArgument(ArgIndex);
+        const FRigVMExecuteArgument* Arg = InTemplate->GetExecuteArgument(ArgIndex, DispatchContext);
         URigVMPin* Pin = NewObject<URigVMPin>(InNode, Arg->Name);
         const FRigVMTemplateArgumentType Type = Registry.GetType(Arg->TypeIndex);
         
@@ -14446,6 +14452,20 @@ void URigVMController::AddPinsForTemplate(const FRigVMTemplate* InTemplate, cons
 				SetPinDefaultValue(Pin, DefaultValue, true, false, false);
 			}
 		}
+		else if(Pin->IsFixedSizeArray())
+		{
+			if(const URigVMDispatchNode* DispatchNode = Cast<URigVMDispatchNode>(Pin->GetNode()))
+			{
+				if(const FRigVMDispatchFactory* Factory = DispatchNode->GetFactory())
+				{
+					const FString DefaultValue =  Factory->GetArgumentDefaultValue(Pin->GetFName(), RigVMTypeUtils::TypeIndex::WildCardArray);
+					if(!DefaultValue.IsEmpty())
+					{
+						SetPinDefaultValue(Pin, DefaultValue, true, false, false);
+					}
+				}
+			}
+		}
 	}
 }
 
@@ -16332,6 +16352,12 @@ bool URigVMController::FullyResolveTemplateNode(URigVMTemplateNode* InNode, int3
 	{
 		if(ResolvedFunction->Struct == nullptr)
 		{
+			FRigVMDispatchContext DispatchContext;
+			if(const URigVMDispatchNode* DispatchNode = Cast<URigVMDispatchNode>(InNode))
+			{
+				DispatchContext = DispatchNode->GetDispatchContext();
+			}
+
 			const TArray<FRigVMTemplateArgument>& Arguments = Template->Arguments;
 			for(URigVMPin* Pin : InNode->GetPins())
 			{
@@ -16341,7 +16367,7 @@ bool URigVMController::FullyResolveTemplateNode(URigVMTemplateNode* InNode, int3
 				});
 				if(!bFound)
 				{
-					bFound = Template->GetExecuteArguments().ContainsByPredicate([Pin](const FRigVMExecuteArgument& Argument)
+					bFound = Template->GetExecuteArguments(DispatchContext).ContainsByPredicate([Pin](const FRigVMExecuteArgument& Argument)
 					{
 						return Pin->GetFName() == Argument.Name;
 					});
diff --git a/Engine/Source/Developer/RigVMDeveloper/Private/RigVMModel/RigVMPin.cpp b/Engine/Source/Developer/RigVMDeveloper/Private/RigVMModel/RigVMPin.cpp
--- a/Engine/Source/Developer/RigVMDeveloper/Private/RigVMModel/RigVMPin.cpp
+++ b/Engine/Source/Developer/RigVMDeveloper/Private/RigVMModel/RigVMPin.cpp
@@ -478,6 +478,13 @@ ERigVMPinDirection URigVMPin::GetDirection() const
 
 bool URigVMPin::IsExpanded() const
 {
+	if(!bIsExpanded)
+	{
+		if(IsFixedSizeArray())
+		{
+			return true;
+		}
+	}
 	return bIsExpanded;
 }
 
diff --git a/Engine/Source/Developer/RigVMDeveloper/Public/RigVMCompiler/RigVMCodeGenerator.h b/Engine/Source/Developer/RigVMDeveloper/Public/RigVMCompiler/RigVMCodeGenerator.h
--- a/Engine/Source/Developer/RigVMDeveloper/Public/RigVMCompiler/RigVMCodeGenerator.h
+++ b/Engine/Source/Developer/RigVMDeveloper/Public/RigVMCompiler/RigVMCodeGenerator.h
@@ -238,7 +238,15 @@ private:
 	
 	FStringArray Libraries;
 	FStringArray Includes;
-	TMap<FString,const FRigVMFunction*> Dispatches;
+
+	struct FRigVMDispatchInfo
+	{
+		FString Name;
+		const FRigVMFunction* Function;
+		FRigVMDispatchContext Context;
+	};
+	
+	TArray<FRigVMDispatchInfo> Dispatches;
 	TMap<TRigVMTypeIndex,TTuple<FString,FString>> RequiredUProperties;
 	TArray<FInstructionGroup> InstructionGroups;
 	FStringArray WrappedArrayCPPTypes;
diff --git a/Engine/Source/Developer/RigVMDeveloper/Public/RigVMModel/Nodes/RigVMDispatchNode.h b/Engine/Source/Developer/RigVMDeveloper/Public/RigVMModel/Nodes/RigVMDispatchNode.h
--- a/Engine/Source/Developer/RigVMDeveloper/Public/RigVMModel/Nodes/RigVMDispatchNode.h
+++ b/Engine/Source/Developer/RigVMDeveloper/Public/RigVMModel/Nodes/RigVMDispatchNode.h
@@ -50,6 +50,29 @@ public:
 	bool IsDeprecated() const;
 	FString GetDeprecatedMetadata() const;
 
+	FRigVMDispatchContext GetDispatchContext() const;
+
+	// Returns an instance of the factory with the current values.
+	// @param bUseDefault If set to true the default struct will be created - otherwise the struct will contains the values from the node
+	TSharedPtr<FStructOnScope> ConstructFactoryInstance(bool bUseDefault = false, FString* OutFactoryDefault = nullptr) const;
+
+	// Returns a copy of the struct with the current values
+	template <
+		typename T,
+		typename TEnableIf<TModels<CRigVMUStruct, T>::Value>::Type* = nullptr
+	>
+	T ConstructFactoryInstance() const
+	{
+		if(!ensure(T::StaticStruct() == GetFactoryStruct()))
+		{
+			return T();
+		}
+
+		TSharedPtr<FStructOnScope> Instance = ConstructFactoryInstance(false);
+		const T& InstanceRef = *(const T*)Instance->GetStructMemory();
+		return InstanceRef;
+	}
+	
 protected:
 
 	virtual FText GetToolTipTextForPin(const URigVMPin* InPin) const override;
@@ -57,6 +80,7 @@ protected:
 	const FRigVMTemplateTypeMap& GetFilteredTypes() const;
 	virtual void InvalidateCache() override;
 	virtual bool ShouldInputPinComputeLazily(const URigVMPin* InPin) const override;
+	FString GetFactoryDefaultValue() const;
 
 private:
 
diff --git a/Engine/Source/Runtime/RigVM/Private/RigVMCore/RigVM.cpp b/Engine/Source/Runtime/RigVM/Private/RigVMCore/RigVM.cpp
--- a/Engine/Source/Runtime/RigVM/Private/RigVMCore/RigVM.cpp
+++ b/Engine/Source/Runtime/RigVM/Private/RigVMCore/RigVM.cpp
@@ -1162,7 +1162,10 @@ void URigVM::CacheMemoryHandlesIfRequired(TArrayView<URigVMMemoryStorage*> InMem
 
 			for (int32 ArgIndex = 0; ArgIndex < Operands.Num(); ArgIndex++)
 			{
-				InOpFunc(InHandleBaseIndex++, {InstructionIndex, ArgIndex, Function->Arguments[ArgIndex].Name}, Operands[ArgIndex]);
+				InOpFunc(
+					InHandleBaseIndex++,
+					{InstructionIndex, ArgIndex, Function->GetArgumentNameForOperandIndex(ArgIndex, Operands.Num())},
+					Operands[ArgIndex]);
 			}
 		}
 		else
diff --git a/Engine/Source/Runtime/RigVM/Private/RigVMCore/RigVMDispatchFactory.cpp b/Engine/Source/Runtime/RigVM/Private/RigVMCore/RigVMDispatchFactory.cpp
--- a/Engine/Source/Runtime/RigVM/Private/RigVMCore/RigVMDispatchFactory.cpp
+++ b/Engine/Source/Runtime/RigVM/Private/RigVMCore/RigVMDispatchFactory.cpp
@@ -131,9 +131,15 @@ bool FRigVMDispatchFactory::IsLazyInputArgument(const FName& InArgumentName) con
 
 #endif
 
-const TArray<FName>& FRigVMDispatchFactory::GetControlFlowBlocks() const
+FName FRigVMDispatchFactory::GetArgumentNameForOperandIndex(int32 InOperandIndex, int32 InTotalOperands) const
 {
-	const TArray<FName>& Blocks = GetControlFlowBlocks_Impl();
+	check(GetArguments().Num() == InTotalOperands);
+	return GetArguments()[InOperandIndex].GetName();
+}
+
+const TArray<FName>& FRigVMDispatchFactory::GetControlFlowBlocks(const FRigVMDispatchContext& InContext) const
+{
+	const TArray<FName>& Blocks = GetControlFlowBlocks_Impl(InContext);
 #if WITH_EDITOR
 	FRigVMStruct::ValidateControlFlowBlocks(Blocks);
 #endif
@@ -145,15 +151,15 @@ bool FRigVMDispatchFactory::SupportsExecuteContextStruct(const UScriptStruct* In
 	return InExecuteContextStruct->IsChildOf(GetExecuteContextStruct());
 }
 
-const TArray<FName>& FRigVMDispatchFactory::GetControlFlowBlocks_Impl() const
+const TArray<FName>& FRigVMDispatchFactory::GetControlFlowBlocks_Impl(const FRigVMDispatchContext& InContext) const
 {
 	static const TArray<FName> EmptyArray;
 	return EmptyArray;
 }
 
-TArray<FRigVMExecuteArgument> FRigVMDispatchFactory::GetExecuteArguments() const
+TArray<FRigVMExecuteArgument> FRigVMDispatchFactory::GetExecuteArguments(const FRigVMDispatchContext& InContext) const
 {
-	TArray<FRigVMExecuteArgument> Arguments = GetExecuteArguments_Impl();
+	TArray<FRigVMExecuteArgument> Arguments = GetExecuteArguments_Impl(InContext);
 	const TRigVMTypeIndex ExecuteTypeIndex = FRigVMRegistry::Get().GetTypeIndex<FRigVMExecuteContext>();
 	for(FRigVMExecuteArgument& Argument : Arguments)
 	{
diff --git a/Engine/Source/Runtime/RigVM/Private/RigVMCore/RigVMFunction.cpp b/Engine/Source/Runtime/RigVM/Private/RigVMCore/RigVMFunction.cpp
--- a/Engine/Source/Runtime/RigVM/Private/RigVMCore/RigVMFunction.cpp
+++ b/Engine/Source/Runtime/RigVM/Private/RigVMCore/RigVMFunction.cpp
@@ -159,3 +159,15 @@ bool FRigVMFunction::SupportsExecuteContextStruct(const UScriptStruct* InExecute
 {
 	return InExecuteContextStruct->IsChildOf(GetExecuteContextStruct());
 }
+
+FName FRigVMFunction::GetArgumentNameForOperandIndex(int32 InOperandIndex, int32 InTotalOperands) const
+{
+	if(Factory)
+	{
+		return Factory->GetArgumentNameForOperandIndex(InOperandIndex, InTotalOperands);
+	}
+
+	check(Arguments.IsValidIndex(InOperandIndex));
+	check(Arguments.Num() == InTotalOperands);
+	return Arguments[InOperandIndex].Name;
+}
diff --git a/Engine/Source/Runtime/RigVM/Private/RigVMCore/RigVMTemplate.cpp b/Engine/Source/Runtime/RigVM/Private/RigVMCore/RigVMTemplate.cpp
--- a/Engine/Source/Runtime/RigVM/Private/RigVMCore/RigVMTemplate.cpp
+++ b/Engine/Source/Runtime/RigVM/Private/RigVMCore/RigVMTemplate.cpp
@@ -3,6 +3,7 @@
 #include "RigVMCore/RigVMTemplate.h"
 #include "RigVMCore/RigVMRegistry.h"
 #include "RigVMCore/RigVMStruct.h"
+#include "RigVMCore/RigVMDispatchFactory.h"
 #include "RigVMModule.h"
 #include "Algo/Sort.h"
 #include "UObject/UObjectIterator.h"
@@ -1056,14 +1057,14 @@ const FRigVMTemplateArgument* FRigVMTemplate::FindArgument(const FName& InArgume
 	});
 }
 
-int32 FRigVMTemplate::NumExecuteArguments() const
+int32 FRigVMTemplate::NumExecuteArguments(const FRigVMDispatchContext& InContext) const
 {
-	return GetExecuteArguments().Num();
+	return GetExecuteArguments(InContext).Num();
 }
 
-const FRigVMExecuteArgument* FRigVMTemplate::GetExecuteArgument(int32 InIndex) const
+const FRigVMExecuteArgument* FRigVMTemplate::GetExecuteArgument(int32 InIndex, const FRigVMDispatchContext& InContext) const
 {
-	const TArray<FRigVMExecuteArgument>& Args = GetExecuteArguments();
+	const TArray<FRigVMExecuteArgument>& Args = GetExecuteArguments(InContext);
 	if(Args.IsValidIndex(InIndex))
 	{
 		return &Args[InIndex];
@@ -1071,16 +1072,16 @@ const FRigVMExecuteArgument* FRigVMTemplate::GetExecuteArgument(int32 InIndex) c
 	return nullptr;
 }
 
-const FRigVMExecuteArgument* FRigVMTemplate::FindExecuteArgument(const FName& InArgumentName) const
+const FRigVMExecuteArgument* FRigVMTemplate::FindExecuteArgument(const FName& InArgumentName, const FRigVMDispatchContext& InContext) const
 {
-	const TArray<FRigVMExecuteArgument>& Args = GetExecuteArguments();
+	const TArray<FRigVMExecuteArgument>& Args = GetExecuteArguments(InContext);
 	return Args.FindByPredicate([InArgumentName](const FRigVMExecuteArgument& Arg) -> bool
 	{
 		return Arg.Name == InArgumentName;
 	});
 }
 
-const TArray<FRigVMExecuteArgument>& FRigVMTemplate::GetExecuteArguments() const
+const TArray<FRigVMExecuteArgument>& FRigVMTemplate::GetExecuteArguments(const FRigVMDispatchContext& InContext) const
 {
 	if(ExecuteArguments.IsEmpty())
 	{
@@ -1089,7 +1090,7 @@ const TArray<FRigVMExecuteArgument>& FRigVMTemplate::GetExecuteArguments() const
 			const FRigVMDispatchFactory* Factory = Delegates.GetDispatchFactoryDelegate.Execute();
 			check(Factory);
 
-			ExecuteArguments = Factory->GetExecuteArguments();
+			ExecuteArguments = Factory->GetExecuteArguments(InContext);
 		}
 		else if(const FRigVMFunction* PrimaryPermutation = GetPrimaryPermutation())
 		{
diff --git a/Engine/Source/Runtime/RigVM/Private/RigVMFunctions/RigVMDispatch_Print.cpp b/Engine/Source/Runtime/RigVM/Private/RigVMFunctions/RigVMDispatch_Print.cpp
--- a/Engine/Source/Runtime/RigVM/Private/RigVMFunctions/RigVMDispatch_Print.cpp
+++ b/Engine/Source/Runtime/RigVM/Private/RigVMFunctions/RigVMDispatch_Print.cpp
@@ -22,7 +22,7 @@ TArray<FRigVMTemplateArgument> FRigVMDispatch_Print::GetArguments() const
 	};
 }
 
-TArray<FRigVMExecuteArgument> FRigVMDispatch_Print::GetExecuteArguments_Impl() const
+TArray<FRigVMExecuteArgument> FRigVMDispatch_Print::GetExecuteArguments_Impl(const FRigVMDispatchContext& InContext) const
 {
 	return {{TEXT("ExecuteContext"), ERigVMPinDirection::IO}};
 }
diff --git a/Engine/Source/Runtime/RigVM/Public/RigVMCore/RigVMDispatchFactory.h b/Engine/Source/Runtime/RigVM/Public/RigVMCore/RigVMDispatchFactory.h
--- a/Engine/Source/Runtime/RigVM/Public/RigVMCore/RigVMDispatchFactory.h
+++ b/Engine/Source/Runtime/RigVM/Public/RigVMCore/RigVMDispatchFactory.h
@@ -16,9 +16,28 @@
 #include "UObject/Class.h"
 #include "UObject/NameTypes.h"
 #include "UObject/ObjectMacros.h"
+#include "UObject/StructOnScope.h"
 
 #include "RigVMDispatchFactory.generated.h"
 
+struct FRigVMDispatchFactory;
+
+/**
+ * A context used for inquiring from dispatch factories
+ */
+struct RIGVM_API FRigVMDispatchContext
+{
+	FRigVMDispatchContext()
+		: Instance()
+		, Subject(nullptr)
+		, StringRepresentation()
+	{}
+	
+	TSharedPtr<FStructOnScope> Instance;
+	const UObject* Subject;
+	FString StringRepresentation;
+};
+
 /**
  * A factory to generate a template and its dispatch functions
  */
@@ -91,11 +110,16 @@ public:
 
 #endif
 
+	// Returns the name to use for the branch info / argument based on the operand index.
+	// operand index count may be different than the number of arguments, since fixed size array
+	// arguments get unrolled.
+	virtual FName GetArgumentNameForOperandIndex(int32 InOperandIndex, int32 InTotalOperands) const;
+
 	// returns true if the dispatch is a control flow dispatch
-	bool IsControlFlowDispatch() const { return !GetControlFlowBlocks().IsEmpty(); }
+	bool IsControlFlowDispatch(const FRigVMDispatchContext& InContext) const { return !GetControlFlowBlocks(InContext).IsEmpty(); }
 
 	// returns the control flow blocks of this dispatch
-	const TArray<FName>& GetControlFlowBlocks() const;
+	const TArray<FName>& GetControlFlowBlocks(const FRigVMDispatchContext& InContext) const;
 
 	// returns true if a given control flow block needs to be sliced
 	virtual const bool IsControlFlowBlockSliced(const FName& InBlockName) const { return false; }
@@ -113,14 +137,14 @@ public:
 	virtual TArray<FRigVMTemplateArgument> GetArguments() const { return TArray<FRigVMTemplateArgument>(); }
 
 	// returns the execute arguments of the template
-	TArray<FRigVMExecuteArgument> GetExecuteArguments() const;
+	TArray<FRigVMExecuteArgument> GetExecuteArguments(const FRigVMDispatchContext& InContext) const;
 
 	// returns the delegate to react to new types being added to an argument.
 	// this happens if types are being loaded later after this factory has already been deployed
 	virtual FRigVMTemplateTypeMap OnNewArgumentType(const FName& InArgumentName, TRigVMTypeIndex InTypeIndex) const { return FRigVMTemplateTypeMap(); }
 
 	// returns the upgrade info to use for this factory
-	virtual FRigVMStructUpgradeInfo GetUpgradeInfo(const FRigVMTemplateTypeMap& InTypes) const { return FRigVMStructUpgradeInfo(); }
+	virtual FRigVMStructUpgradeInfo GetUpgradeInfo(const FRigVMTemplateTypeMap& InTypes, const FRigVMDispatchContext& InContext) const { return FRigVMStructUpgradeInfo(); }
 
 	// returns the dispatch function for a given type set
 	FRigVMFunctionPtr GetDispatchFunction(const FRigVMTemplateTypeMap& InTypes) const;
@@ -138,7 +162,7 @@ protected:
 
 	virtual FRigVMFunctionPtr GetDispatchFunctionImpl(const FRigVMTemplateTypeMap& InTypes) const { return nullptr; }
 
-	virtual TArray<FRigVMExecuteArgument> GetExecuteArguments_Impl() const { return TArray<FRigVMExecuteArgument>(); }
+	virtual TArray<FRigVMExecuteArgument> GetExecuteArguments_Impl(const FRigVMDispatchContext& InContext) const { return TArray<FRigVMExecuteArgument>(); }
 
 	template <
 	typename T,
@@ -179,7 +203,7 @@ protected:
 	}
 #endif
 
-	virtual const TArray<FName>& GetControlFlowBlocks_Impl() const;
+	virtual const TArray<FName>& GetControlFlowBlocks_Impl(const FRigVMDispatchContext& InContext) const;
 
 	static const FString DispatchPrefix;
 
diff --git a/Engine/Source/Runtime/RigVM/Public/RigVMCore/RigVMFunction.h b/Engine/Source/Runtime/RigVM/Public/RigVMCore/RigVMFunction.h
--- a/Engine/Source/Runtime/RigVM/Public/RigVMCore/RigVMFunction.h
+++ b/Engine/Source/Runtime/RigVM/Public/RigVMCore/RigVMFunction.h
@@ -157,4 +157,5 @@ struct RIGVM_API FRigVMFunction
 	const FRigVMTemplate* GetTemplate() const;
 	const UScriptStruct* GetExecuteContextStruct() const;
 	bool SupportsExecuteContextStruct(const UScriptStruct* InExecuteContextStruct) const;
+	FName GetArgumentNameForOperandIndex(int32 InOperandIndex, int32 InTotalOperands) const;
 };
diff --git a/Engine/Source/Runtime/RigVM/Public/RigVMCore/RigVMTemplate.h b/Engine/Source/Runtime/RigVM/Public/RigVMCore/RigVMTemplate.h
--- a/Engine/Source/Runtime/RigVM/Public/RigVMCore/RigVMTemplate.h
+++ b/Engine/Source/Runtime/RigVM/Public/RigVMCore/RigVMTemplate.h
@@ -28,6 +28,7 @@
 class FProperty;
 struct FRigVMDispatchFactory;
 struct FRigVMTemplate;
+struct FRigVMDispatchContext;
 
 typedef TMap<FName, TRigVMTypeIndex> FRigVMTemplateTypeMap;
 
@@ -298,13 +299,13 @@ public:
 	const FRigVMTemplateArgument* FindArgument(const FName& InArgumentName) const;
 
 	// returns the number of args of this template
-	int32 NumExecuteArguments() const;
+	int32 NumExecuteArguments(const FRigVMDispatchContext& InContext) const;
 
 	// returns an argument for a given index
-	const FRigVMExecuteArgument* GetExecuteArgument(int32 InIndex) const;
+	const FRigVMExecuteArgument* GetExecuteArgument(int32 InIndex, const FRigVMDispatchContext& InContext) const;
 
 	// returns an argument given a name (or nullptr)
-	const FRigVMExecuteArgument* FindExecuteArgument(const FName& InArgumentName) const;
+	const FRigVMExecuteArgument* FindExecuteArgument(const FName& InArgumentName, const FRigVMDispatchContext& InContext) const;
 
 	// returns the top level execute context struct this template uses
 	const UScriptStruct* GetExecuteContextStruct() const;
@@ -425,7 +426,7 @@ private:
 
 	static FLinearColor GetColorFromMetadata(FString InMetadata);
 
-	const TArray<FRigVMExecuteArgument>& GetExecuteArguments() const;
+	const TArray<FRigVMExecuteArgument>& GetExecuteArguments(const FRigVMDispatchContext& InContext) const;
 
 	int32 Index;
 	FName Notation;
diff --git a/Engine/Source/Runtime/RigVM/Public/RigVMFunctions/RigVMDispatch_Print.h b/Engine/Source/Runtime/RigVM/Public/RigVMFunctions/RigVMDispatch_Print.h
--- a/Engine/Source/Runtime/RigVM/Public/RigVMFunctions/RigVMDispatch_Print.h
+++ b/Engine/Source/Runtime/RigVM/Public/RigVMFunctions/RigVMDispatch_Print.h
@@ -16,7 +16,7 @@ struct FRigVMDispatch_Print : public FRigVMDispatchFactory
 public:
 
 	virtual TArray<FRigVMTemplateArgument> GetArguments() const override;
-	virtual TArray<FRigVMExecuteArgument> GetExecuteArguments_Impl() const override;
+	virtual TArray<FRigVMExecuteArgument> GetExecuteArguments_Impl(const FRigVMDispatchContext& InContext) const override;
 	virtual FRigVMTemplateTypeMap OnNewArgumentType(const FName& InArgumentName, TRigVMTypeIndex InTypeIndex) const override;
 
 #if WITH_EDITOR
