RigVM: New if and select dispatches

diff --git a/Engine/Plugins/Animation/ControlRig/Source/ControlRigEditor/Private/ControlRigEditorModule.cpp b/Engine/Plugins/Animation/ControlRig/Source/ControlRigEditor/Private/ControlRigEditorModule.cpp
--- a/Engine/Plugins/Animation/ControlRig/Source/ControlRigEditor/Private/ControlRigEditorModule.cpp
+++ b/Engine/Plugins/Animation/ControlRig/Source/ControlRigEditor/Private/ControlRigEditorModule.cpp
@@ -49,8 +49,6 @@
 #include "Graph/NodeSpawners/ControlRigUnitNodeSpawner.h"
 #include "Graph/NodeSpawners/ControlRigVariableNodeSpawner.h"
 #include "Graph/NodeSpawners/ControlRigRerouteNodeSpawner.h"
-#include "Graph/NodeSpawners/ControlRigIfNodeSpawner.h"
-#include "Graph/NodeSpawners/ControlRigSelectNodeSpawner.h"
 #include "Graph/NodeSpawners/ControlRigTemplateNodeSpawner.h"
 #include "Graph/NodeSpawners/ControlRigEnumNodeSpawner.h"
 #include "Graph/NodeSpawners/ControlRigFunctionRefNodeSpawner.h"
@@ -1116,18 +1114,6 @@ void FControlRigEditorModule::GetTypeActions(UControlRigBlueprint* CRB, FBluepri
 		LOCTEXT("RerouteSpawnerTooltip", "Adds a new reroute node to the graph"));
 	ActionRegistrar.AddBlueprintAction(ActionKey, RerouteNodeSpawner);
 
-	UBlueprintNodeSpawner* IfNodeSpawner = UControlRigIfNodeSpawner::CreateGeneric(
-		LOCTEXT("IfSpawnerDesc", "If"),
-		LOCTEXT("IfSpawnerCategory", "Execution"),
-		LOCTEXT("IfSpawnerTooltip", "Adds a new 'if' node to the graph"));
-	ActionRegistrar.AddBlueprintAction(ActionKey, IfNodeSpawner);
-
-	UBlueprintNodeSpawner* SelectNodeSpawner = UControlRigSelectNodeSpawner::CreateGeneric(
-		LOCTEXT("SelectSpawnerDesc", "Select"),
-		LOCTEXT("SelectSpawnerCategory", "Execution"),
-		LOCTEXT("SelectSpawnerTooltip", "Adds a new 'select' node to the graph"));
-	ActionRegistrar.AddBlueprintAction(ActionKey, SelectNodeSpawner);
-
 	const int32 FirstArrayOpCode = (int32)ERigVMOpCode::FirstArrayOpCode; 
 	const int32 LastArrayOpCode = (int32)ERigVMOpCode::LastArrayOpCode;
 	for(int32 OpCodeIndex = FirstArrayOpCode; OpCodeIndex <= LastArrayOpCode; OpCodeIndex++)
diff --git a/Engine/Plugins/Animation/ControlRig/Source/ControlRigEditor/Private/Graph/NodeSpawners/ControlRigIfNodeSpawner.cpp b/Engine/Plugins/Animation/ControlRig/Source/ControlRigEditor/Private/Graph/NodeSpawners/ControlRigIfNodeSpawner.cpp
deleted file mode 100644
diff --git a/Engine/Plugins/Animation/ControlRig/Source/ControlRigEditor/Private/Graph/NodeSpawners/ControlRigIfNodeSpawner.h b/Engine/Plugins/Animation/ControlRig/Source/ControlRigEditor/Private/Graph/NodeSpawners/ControlRigIfNodeSpawner.h
deleted file mode 100644
diff --git a/Engine/Plugins/Animation/ControlRig/Source/ControlRigEditor/Private/Graph/NodeSpawners/ControlRigSelectNodeSpawner.cpp b/Engine/Plugins/Animation/ControlRig/Source/ControlRigEditor/Private/Graph/NodeSpawners/ControlRigSelectNodeSpawner.cpp
deleted file mode 100644
diff --git a/Engine/Plugins/Animation/ControlRig/Source/ControlRigEditor/Private/Graph/NodeSpawners/ControlRigSelectNodeSpawner.h b/Engine/Plugins/Animation/ControlRig/Source/ControlRigEditor/Private/Graph/NodeSpawners/ControlRigSelectNodeSpawner.h
deleted file mode 100644
diff --git a/Engine/Plugins/Animation/ControlRig/Source/ControlRigEditor/Private/Graph/SControlRigGraphNode.cpp b/Engine/Plugins/Animation/ControlRig/Source/ControlRigEditor/Private/Graph/SControlRigGraphNode.cpp
--- a/Engine/Plugins/Animation/ControlRig/Source/ControlRigEditor/Private/Graph/SControlRigGraphNode.cpp
+++ b/Engine/Plugins/Animation/ControlRig/Source/ControlRigEditor/Private/Graph/SControlRigGraphNode.cpp
@@ -36,8 +36,8 @@
 #include "RigVMModel/Nodes/RigVMAggregateNode.h"
 #include "RigVMModel/Nodes/RigVMDispatchNode.h"
 #include "Slate/SlateTextures.h"
-//#include "RigVMFunctions/RigVMDispatch_If.h"
-//#include "RigVMFunctions/RigVMDispatch_Select.h"
+#include "RigVMFunctions/RigVMDispatch_If.h"
+#include "RigVMFunctions/RigVMDispatch_Select.h"
 
 #if WITH_EDITOR
 #include "Editor.h"
@@ -994,10 +994,9 @@ void SControlRigGraphNode::RefreshErrorInfo()
 			{
 				if(const URigVMNode* RigModelNode = RigGraphNode->GetModelNode())
 				{
-					bool bShowCopyWarning = RigModelNode->IsA<URigVMIfNode>() || RigModelNode->IsA<URigVMSelectNode>();
+					bool bShowCopyWarning = RigModelNode->IsA<UDEPRECATED_RigVMIfNode>() || RigModelNode->IsA<UDEPRECATED_RigVMSelectNode>();
 					if(!bShowCopyWarning)
 					{
-						/*
 						if(const URigVMDispatchNode* DispatchNode = Cast<URigVMDispatchNode>(RigModelNode))
 						{
 							if(const UScriptStruct* FactoryStruct = DispatchNode->GetScriptStruct())
@@ -1009,7 +1008,6 @@ void SControlRigGraphNode::RefreshErrorInfo()
 								}
 							}
 						}
-						*/
 					}
 					if(bShowCopyWarning)
 					{
diff --git a/Engine/Source/Developer/RigVMDeveloper/Private/RigVMCompiler/RigVMAST.cpp b/Engine/Source/Developer/RigVMDeveloper/Private/RigVMCompiler/RigVMAST.cpp
--- a/Engine/Source/Developer/RigVMDeveloper/Private/RigVMCompiler/RigVMAST.cpp
+++ b/Engine/Source/Developer/RigVMDeveloper/Private/RigVMCompiler/RigVMAST.cpp
@@ -7,8 +7,6 @@
 #include "RigVMModel/Nodes/RigVMVariableNode.h"
 #include "RigVMModel/Nodes/RigVMCommentNode.h"
 #include "RigVMModel/Nodes/RigVMRerouteNode.h"
-#include "RigVMModel/Nodes/RigVMIfNode.h"
-#include "RigVMModel/Nodes/RigVMSelectNode.h"
 #include "RigVMModel/Nodes/RigVMEnumNode.h"
 #include "RigVMModel/Nodes/RigVMFunctionReturnNode.h"
 #include "RigVMModel/Nodes/RigVMFunctionEntryNode.h"
@@ -84,14 +82,6 @@ FName FRigVMExprAST::GetTypeName() const
 		{
 			return TEXT("[.Exit..]");
 		}
-		case EType::If:
-		{
-			return TEXT("[..If...]");
-		}
-		case EType::Select:
-		{
-			return TEXT("[Select.]");
-		}
 		case EType::Array:
 		{
 			return TEXT("[Array..]");
@@ -616,79 +606,6 @@ bool FRigVMVarExprAST::SupportsSoftLinks() const
 	return false;
 }
 
-bool FRigVMIfExprAST::IsConstant() const
-{
-	if (IsAlwaysTrue())
-	{
-		return GetTrueExpr()->IsConstant();
-	}
-	else if (IsAlwaysFalse())
-	{
-		return GetFalseExpr()->IsConstant();
-	}
-	return FRigVMNodeExprAST::IsConstant();
-}
-
-bool FRigVMIfExprAST::IsAlwaysTrue() const
-{
-	const FRigVMVarExprAST* ConditionExpr = GetConditionExpr();
-	if (ConditionExpr->IsA(EType::Literal))
-	{
-		const FString& PinDefaultValue = ConditionExpr->GetDefaultValue();
-		return PinDefaultValue == TEXT("True");
-	}
-	return false;
-}
-
-bool FRigVMIfExprAST::IsAlwaysFalse() const
-{
-	const FRigVMVarExprAST* ConditionExpr = GetConditionExpr();
-	if (ConditionExpr->IsA(EType::Literal))
-	{
-		const FString& PinDefaultValue = ConditionExpr->GetDefaultValue();
-		return PinDefaultValue == TEXT("False") || PinDefaultValue.IsEmpty();
-	}
-	return false;
-}
-
-bool FRigVMSelectExprAST::IsConstant() const
-{
-	int32 ConstantCaseIndex = GetConstantValueIndex();
-	if (ConstantCaseIndex != INDEX_NONE)
-	{
-		return GetValueExpr(ConstantCaseIndex)->IsConstant();
-	}
-	return FRigVMNodeExprAST::IsConstant();
-}
-
-int32 FRigVMSelectExprAST::GetConstantValueIndex() const
-{
-	const FRigVMVarExprAST* IndexExpr = GetIndexExpr();
-	if (IndexExpr->IsA(EType::Literal))
-	{
-		int32 NumCases = NumValues();
-		if (NumCases == 0)
-		{
-			return INDEX_NONE;
-		}
-
-		const FString& PinDefaultValue = IndexExpr->GetDefaultValue();
-		int32 CaseIndex = 0;
-		if (!PinDefaultValue.IsEmpty())
-		{
-			CaseIndex = FCString::Atoi(*PinDefaultValue);
-		}
-
-		return FMath::Clamp<int32>(CaseIndex, 0, NumCases - 1);
-	}
-	return INDEX_NONE;
-}
-
-int32 FRigVMSelectExprAST::NumValues() const
-{
-	return GetNode()->FindPin(URigVMSelectNode::ValueName)->GetArraySize();
-}
-
 void FRigVMParserASTSettings::Report(EMessageSeverity::Type InSeverity, UObject* InSubject, const FString& InMessage) const
 {
 	if (ReportDelegate.IsBound())
@@ -924,7 +841,7 @@ FRigVMExprAST* FRigVMParserAST::TraverseMutableNode(const FRigVMASTProxy& InNode
 					FRigVMASTProxy SourcePinProxy = InNodeProxy.GetSibling(SourcePin);
 
 					FRigVMExprAST* ParentExpr = InParentExpr;
-					if (NodeExpr->IsA(FRigVMExprAST::Branch) || Node->IsControlFlowNode())
+					if (Node->IsControlFlowNode())
 					{
 						if (FRigVMExprAST** PinExpr = SubjectToExpression.Find(SourcePinProxy))
 						{
@@ -1033,14 +950,6 @@ FRigVMExprAST* FRigVMParserAST::CreateExpressionForNode(const FRigVMASTProxy& In
 		{
 			NodeExpr = MakeExpr<FRigVMNoOpExprAST>(InNodeProxy);
 		}
-		else if (InNodeProxy.IsA<URigVMIfNode>())
-		{
-			NodeExpr = MakeExpr<FRigVMIfExprAST>(InNodeProxy);
-		}
-		else if (InNodeProxy.IsA<URigVMSelectNode>())
-		{
-			NodeExpr = MakeExpr<FRigVMSelectExprAST>(InNodeProxy);
-		}
 		else if (InNodeProxy.IsA<URigVMArrayNode>())
 		{
 			NodeExpr = MakeExpr<FRigVMArrayExprAST>(InNodeProxy);
@@ -1102,7 +1011,7 @@ TArray<FRigVMExprAST*> FRigVMParserAST::TraversePins(const FRigVMASTProxy& InNod
 	}
 
 	// dispatch nodes may contain a fixed size array
-	if(Node->IsA<URigVMDispatchNode>() || Node->IsA<URigVMSelectNode>())
+	if(Node->IsA<URigVMDispatchNode>() || Node->IsA<UDEPRECATED_RigVMSelectNode>())
 	{
 		TArray<URigVMPin*> FlattenedPins;
 		FlattenedPins.Reserve(Pins.Num());
@@ -1127,7 +1036,7 @@ TArray<FRigVMExprAST*> FRigVMParserAST::TraversePins(const FRigVMASTProxy& InNod
 
 		if (InParentExpr && !Pin->IsRootPin())
 		{
-			const FString PinName = FString::Printf(TEXT("%s_%s"), *Pin->GetParentPin()->GetName(), *Pin->GetName());
+			const FString PinName = FRigVMBranchInfo::GetFixedArrayLabel(Pin->GetParentPin()->GetName(), Pin->GetName());
 			const int32 ChildIndex = InParentExpr->Children.Find(PinExpressions.Last());
 			InParentExpr->PinNameToChildIndex.FindOrAdd(*PinName) = ChildIndex;
 		}
@@ -1311,15 +1220,6 @@ FRigVMExprAST* FRigVMParserAST::TraversePin(const FRigVMASTProxy& InPinProxy, FR
 				ParentExprForLinks = BlockExpr;
 			}
 		}
-		else if ((Pin->GetDirection() == ERigVMPinDirection::IO || Pin->GetDirection() == ERigVMPinDirection::Input) &&
-			(InParentExpr->IsA(FRigVMExprAST::If) || InParentExpr->IsA(FRigVMExprAST::Select)) &&
-			LinkIndices.Num() > 0)
-		{
-			FRigVMBlockExprAST* BlockExpr = MakeExpr<FRigVMBlockExprAST>(FRigVMExprAST::EType::Block, FRigVMASTProxy());
-			BlockExpr->Name = Pin->GetFName();
-			BlockExpr->AddParent(PinExpr);
-			ParentExprForLinks = BlockExpr;
-		}
 	}
 
 	for (const int32 LinkIndex : LinkIndices)
@@ -1369,8 +1269,6 @@ FRigVMExprAST* FRigVMParserAST::TraverseLink(int32 InLinkIndex, FRigVMExprAST* I
 		// if this is a copy expression - we should require the copy to use a ref instead
 		if (NodeExpr->IsA(FRigVMExprAST::EType::CallExtern) ||
 			NodeExpr->IsA(FRigVMExprAST::EType::InlineFunction) ||
-			NodeExpr->IsA(FRigVMExprAST::EType::If) ||
-			NodeExpr->IsA(FRigVMExprAST::EType::Select) ||
 			NodeExpr->IsA(FRigVMExprAST::EType::Array))
 		{
 			for (FRigVMExprAST* ChildExpr : *NodeExpr)
@@ -1579,9 +1477,7 @@ void FRigVMParserAST::BubbleUpExpressions()
 				// parent == pin, grandparent == node
 				if(GrandParents.Num() == 1 && !GrandParents.Contains(nullptr))
 				{
-					if(GrandParents[0]->IsA(FRigVMExprAST::EType::CallExtern) ||
-						GrandParents[0]->IsA(FRigVMExprAST::EType::If) ||
-						GrandParents[0]->IsA(FRigVMExprAST::EType::Select))
+					if(GrandParents[0]->IsA(FRigVMExprAST::EType::CallExtern))
 					{
 						const FRigVMNodeExprAST* NodeExpr = GrandParents[0]->To<FRigVMNodeExprAST>();
 
@@ -1598,21 +1494,6 @@ void FRigVMParserAST::BubbleUpExpressions()
 							if(Pin->GetDirection() == ERigVMPinDirection::Input ||
 								Pin->GetDirection() == ERigVMPinDirection::IO)
 							{
-								if(Pin->GetNode()->IsA<URigVMIfNode>())
-								{
-									if(Pin->GetName() != URigVMIfNode::ConditionName)
-									{
-										continue;
-									}
-								}
-								else if(Pin->GetNode()->IsA<URigVMSelectNode>())
-								{
-									if(Pin->GetName() != URigVMSelectNode::IndexName)
-									{
-										continue;
-									}
-								}
-
 								if(const FRigVMExprAST* PinExpr = NodeExpr->FindExprWithPinName(Pin->GetFName()))
 								{
 									TopLevelBlockExpression = PinExpr->GetFirstChildOfType(FRigVMExprAST::Block);
@@ -2754,7 +2635,6 @@ FString FRigVMParserAST::DumpDot() const
 			{
 				case FRigVMExprAST::EType::Entry:
 				case FRigVMExprAST::EType::Exit:
-				case FRigVMExprAST::EType::Branch:
 				case FRigVMExprAST::EType::Block:
 				{
 					Shape = EVisualGraphShape::Diamond;
@@ -2764,8 +2644,6 @@ FString FRigVMParserAST::DumpDot() const
 				case FRigVMExprAST::EType::Copy:
 				case FRigVMExprAST::EType::CallExtern:
 				case FRigVMExprAST::EType::InlineFunction:
-				case FRigVMExprAST::EType::If:
-				case FRigVMExprAST::EType::Select:
 				case FRigVMExprAST::EType::NoOp:
 				{
 					Shape = EVisualGraphShape::Box;
diff --git a/Engine/Source/Developer/RigVMDeveloper/Private/RigVMCompiler/RigVMCompiler.cpp b/Engine/Source/Developer/RigVMDeveloper/Private/RigVMCompiler/RigVMCompiler.cpp
--- a/Engine/Source/Developer/RigVMDeveloper/Private/RigVMCompiler/RigVMCompiler.cpp
+++ b/Engine/Source/Developer/RigVMDeveloper/Private/RigVMCompiler/RigVMCompiler.cpp
@@ -477,6 +477,20 @@ bool URigVMCompiler::Compile(TArray<URigVMGraph*> InGraphs, URigVMController* In
 				bEncounteredGraphError = true;
 			}
 
+			if(ModelNode->IsA<UDEPRECATED_RigVMIfNode>())
+			{
+				static const FString LinkedMessage = TEXT("Node @@ is a deprecated if node. Cannot compile.");
+				Settings.ASTSettings.Report(EMessageSeverity::Error, ModelNode, LinkedMessage);
+				bEncounteredGraphError = true;
+			}
+
+			if(ModelNode->IsA<UDEPRECATED_RigVMSelectNode>())
+			{
+				static const FString LinkedMessage = TEXT("Node @@ is a deprecated select node. Cannot compile.");
+				Settings.ASTSettings.Report(EMessageSeverity::Error, ModelNode, LinkedMessage);
+				bEncounteredGraphError = true;
+			}
+
 			if(!InController->RemoveUnusedOrphanedPins(ModelNode))
 			{
 				static const FString LinkedMessage = TEXT("Node @@ uses pins that no longer exist. Please rewire the links and re-compile.");
@@ -1205,16 +1219,6 @@ void URigVMCompiler::TraverseExpression(const FRigVMExprAST* InExpr, FRigVMCompi
 			TraverseExit(InExpr->To<FRigVMExitExprAST>(), WorkData);
 			break;
 		}
-		case FRigVMExprAST::EType::If:
-		{
-			TraverseIf(InExpr->To<FRigVMIfExprAST>(), WorkData);
-			break;
-		}
-		case FRigVMExprAST::EType::Select:
-		{
-			TraverseSelect(InExpr->To<FRigVMSelectExprAST>(), WorkData);
-			break;
-		}
 		case FRigVMExprAST::EType::Array:
 		{
 			TraverseArray(InExpr->To<FRigVMArrayExprAST>(), WorkData);
@@ -1276,12 +1280,19 @@ void URigVMCompiler::TraverseBlock(const FRigVMBlockExprAST* InExpr, FRigVMCompi
 
 					if(Pin)
 					{
-						Pin = Pin->GetRootPin();
-						if(Pin->IsLazy())
+						URigVMPin* RootPin = Pin->GetRootPin();
+						if(RootPin->IsLazy())
 						{
-							CallExternNode = Pin->GetNode();
+							CallExternNode = RootPin->GetNode();
 							
-							BranchInfo.Label = Pin->GetFName();
+							if(RootPin->IsFixedSizeArray() && Pin->GetParentPin() == RootPin)
+							{
+								BranchInfo.Label = FRigVMBranchInfo::GetFixedArrayLabel(RootPin->GetFName(), Pin->GetFName());
+							}
+							else
+							{
+								BranchInfo.Label = RootPin->GetFName();
+							}
 							BranchInfo.InstructionIndex = INDEX_NONE; // we'll fill in the instruction info later
 							BranchInfo.FirstInstruction = WorkData.VM->GetByteCode().GetNumInstructions();
 
@@ -1290,14 +1301,31 @@ void URigVMCompiler::TraverseBlock(const FRigVMBlockExprAST* InExpr, FRigVMCompi
 							{
 								if(const FRigVMTemplate* Template = TemplateNode->GetTemplate())
 								{
+									int32 FlatArgumentIndex = 0;
 									for(int32 ArgumentIndex = 0; ArgumentIndex != Template->NumArguments(); ArgumentIndex++)
 									{
 										const FRigVMTemplateArgument* Argument = Template->GetArgument(ArgumentIndex);
-										if(Template->GetArgument(ArgumentIndex)->GetName() == Pin->GetFName())
+										if(Template->GetArgument(ArgumentIndex)->GetName() == RootPin->GetFName())
 										{
-											BranchInfo.ArgumentIndex = ArgumentIndex;
+											BranchInfo.ArgumentIndex = FlatArgumentIndex;
+
+											if(RootPin->IsFixedSizeArray() && Pin->GetParentPin() == RootPin)
+											{
+												BranchInfo.ArgumentIndex += Pin->GetPinIndex();
+											}
 											break;
 										}
+
+										if(const URigVMPin* PinForArgument = RootPin->GetNode()->FindPin(Argument->Name.ToString()))
+										{
+											if(PinForArgument->IsFixedSizeArray())
+											{
+												FlatArgumentIndex += RootPin->GetSubPins().Num();
+												continue;
+											}
+										}
+										
+										FlatArgumentIndex++;
 									}
 								}
 								// we also need to deal with unit nodes separately here. if a unit node does
@@ -1311,7 +1339,7 @@ void URigVMCompiler::TraverseBlock(const FRigVMBlockExprAST* InExpr, FRigVMCompi
 										for(int32 ArgumentIndex = 0; ArgumentIndex != Function->Arguments.Num(); ArgumentIndex++)
 										{
 											const FRigVMFunctionArgument& Argument = Function->Arguments[ArgumentIndex];
-											if(Argument.Name == Pin->GetFName())
+											if(Argument.Name == RootPin->GetFName())
 											{
 												BranchInfo.ArgumentIndex = ArgumentIndex;
 												break;
@@ -1568,8 +1596,7 @@ int32 URigVMCompiler::TraverseCallExtern(const FRigVMCallExternExprAST* InExpr,
 				{
 					for(URigVMPin* SubPin : Pin->GetSubPins())
 					{
-						static constexpr TCHAR Format[] = TEXT("%s_%s");
-						const FString PinName = FString::Printf(Format, *Pin->GetName(), *SubPin->GetName());
+						const FString PinName = FRigVMBranchInfo::GetFixedArrayLabel(Pin->GetName(), SubPin->GetName());
 						const FRigVMExprAST* SubPinExpr = InExpr->FindExprWithPinName(*PinName);
 						check(SubPinExpr);
 						ProcessArgument(SubPinExpr);
@@ -2252,237 +2279,6 @@ void URigVMCompiler::TraverseExit(const FRigVMExitExprAST* InExpr, FRigVMCompile
 	}
 }
 
-void URigVMCompiler::TraverseIf(const FRigVMIfExprAST* InExpr, FRigVMCompilerWorkData& WorkData)
-{
-	ensure(InExpr->NumChildren() == 4);
-
-	if (WorkData.bSetupMemory)
-	{
-		TraverseChildren(InExpr, WorkData);
-		return;
-	}
-
-	URigVMIfNode* IfNode = Cast<URigVMIfNode>(InExpr->GetNode());
-	if(!ValidateNode(IfNode))
-	{
-		return;
-	}
-
-	const FRigVMCallstack Callstack = InExpr->GetProxy().GetCallstack();
-
-	const FRigVMVarExprAST* ConditionExpr = InExpr->ChildAt<FRigVMVarExprAST>(0);
-	const FRigVMVarExprAST* TrueExpr = InExpr->ChildAt<FRigVMVarExprAST>(1);
-	const FRigVMVarExprAST* FalseExpr = InExpr->ChildAt<FRigVMVarExprAST>(2);
-	const FRigVMVarExprAST* ResultExpr = InExpr->ChildAt<FRigVMVarExprAST>(3);
-
-	// traverse the condition first
-	TraverseExpression(ConditionExpr, WorkData);
-
-	if (ConditionExpr->IsA(FRigVMExprAST::CachedValue))
-	{
-		ConditionExpr = ConditionExpr->To<FRigVMCachedValueExprAST>()->GetVarExpr();
-	}
-
-	FRigVMOperand& ConditionOperand = WorkData.ExprToOperand.FindChecked(GetSourceVarExpr(ConditionExpr));
-	FRigVMOperand& ResultOperand = WorkData.ExprToOperand.FindChecked(GetSourceVarExpr(ResultExpr));
-
-	// setup the first jump
-	uint64 JumpToFalseByte = WorkData.VM->GetByteCode().AddJumpIfOp(ERigVMOpCode::JumpForwardIf, 1, ConditionOperand, false);
-	int32 JumpToFalseInstruction = WorkData.VM->GetByteCode().GetNumInstructions() - 1;
-	if (Settings.SetupNodeInstructionIndex)
-	{
-		WorkData.VM->GetByteCode().SetSubject(WorkData.VM->GetByteCode().GetNumInstructions() - 1, Callstack.GetCallPath(), Callstack.GetStack());
-	}
-
-	// traverse the true case
-	TraverseExpression(TrueExpr, WorkData);
-
-	if (TrueExpr->IsA(FRigVMExprAST::CachedValue))
-	{
-		TrueExpr = TrueExpr->To<FRigVMCachedValueExprAST>()->GetVarExpr();
-	}
-
-	FRigVMOperand& TrueOperand = WorkData.ExprToOperand.FindChecked(GetSourceVarExpr(TrueExpr));
-
-	WorkData.VM->GetByteCode().AddCopyOp(WorkData.VM->GetCopyOpForOperands(TrueOperand, ResultOperand));
-	if (Settings.SetupNodeInstructionIndex)
-	{
-		WorkData.VM->GetByteCode().SetSubject(WorkData.VM->GetByteCode().GetNumInstructions() - 1, Callstack.GetCallPath(), Callstack.GetStack());
-	}
-
-	uint64 JumpToEndByte = WorkData.VM->GetByteCode().AddJumpOp(ERigVMOpCode::JumpForward, 1);
-	int32 JumpToEndInstruction = WorkData.VM->GetByteCode().GetNumInstructions() - 1;
-	if (Settings.SetupNodeInstructionIndex)
-	{
-		WorkData.VM->GetByteCode().SetSubject(WorkData.VM->GetByteCode().GetNumInstructions() - 1, Callstack.GetCallPath(), Callstack.GetStack());
-	}
-
-	// correct the jump to false instruction index
-	int32 NumInstructionsInTrueCase = WorkData.VM->GetByteCode().GetNumInstructions() - JumpToFalseInstruction;
-	WorkData.VM->GetByteCode().GetOpAt<FRigVMJumpIfOp>(JumpToFalseByte).InstructionIndex = NumInstructionsInTrueCase;
-
-	// traverse the false case
-	TraverseExpression(FalseExpr, WorkData);
-
-	if (FalseExpr->IsA(FRigVMExprAST::CachedValue))
-	{
-		FalseExpr = FalseExpr->To<FRigVMCachedValueExprAST>()->GetVarExpr();
-	}
-
-	FRigVMOperand& FalseOperand = WorkData.ExprToOperand.FindChecked(GetSourceVarExpr(FalseExpr));
-
-	WorkData.VM->GetByteCode().AddCopyOp(WorkData.VM->GetCopyOpForOperands(FalseOperand, ResultOperand));
-	if (Settings.SetupNodeInstructionIndex)
-	{
-		WorkData.VM->GetByteCode().SetSubject(WorkData.VM->GetByteCode().GetNumInstructions() - 1, Callstack.GetCallPath(), Callstack.GetStack());
-	}
-
-	// correct the jump to end instruction index
-	int32 NumInstructionsInFalseCase = WorkData.VM->GetByteCode().GetNumInstructions() - JumpToEndInstruction;
-	WorkData.VM->GetByteCode().GetOpAt<FRigVMJumpOp>(JumpToEndByte).InstructionIndex = NumInstructionsInFalseCase;
-}
-
-void URigVMCompiler::TraverseSelect(const FRigVMSelectExprAST* InExpr, FRigVMCompilerWorkData& WorkData)
-{
-	URigVMSelectNode* SelectNode = Cast<URigVMSelectNode>(InExpr->GetNode());
-	if(!ValidateNode(SelectNode))
-	{
-		return;
-	}
-
-	const FRigVMCallstack Callstack = InExpr->GetProxy().GetCallstack();
-
-	int32 NumCases = SelectNode->FindPin(URigVMSelectNode::ValueName)->GetArraySize();
-
-	if (WorkData.bSetupMemory)
-	{
-		TraverseChildren(InExpr, WorkData);
-
-		// setup literals for each index (we don't need zero)
-		for (int32 CaseIndex = 1; CaseIndex < NumCases; CaseIndex++)
-		{
-			if (!WorkData.IntegerLiterals.Contains(CaseIndex))
-			{
-				FName LiteralName = *FString::FromInt(CaseIndex);
-
-				const FString DefaultValue = FString::FromInt(CaseIndex);
-				FRigVMOperand Operand = WorkData.AddProperty(
-					ERigVMMemoryType::Literal,
-					LiteralName,
-					TEXT("int32"),
-					nullptr,
-					DefaultValue);
-				
-				WorkData.IntegerLiterals.Add(CaseIndex, Operand);
-			}
-		}
-
-		if (!WorkData.ComparisonOperand.IsValid())
-		{
-			WorkData.ComparisonOperand = WorkData.AddProperty(
-				ERigVMMemoryType::Work,
-				FName(TEXT("IntEquals")),
-				TEXT("bool"),
-				nullptr,
-				TEXT("false"));
-		}
-		return;
-	}
-
-	const FRigVMVarExprAST* IndexExpr = InExpr->ChildAt<FRigVMVarExprAST>(0);
-	TArray<const FRigVMVarExprAST*> CaseExpressions;
-	for (int32 CaseIndex = 0; CaseIndex < NumCases; CaseIndex++)
-	{
-		CaseExpressions.Add(InExpr->ChildAt<FRigVMVarExprAST>(CaseIndex + 1));
-	}
-
-	const FRigVMVarExprAST* ResultExpr = InExpr->ChildAt<FRigVMVarExprAST>(InExpr->NumChildren() - 1);
-
-	// traverse the condition first
-	TraverseExpression(IndexExpr, WorkData);
-
-	// this can happen if the optimizer doesn't remove it
-	if (CaseExpressions.Num() == 0)
-	{
-		return;
-	}
-
-	if (IndexExpr->IsA(FRigVMExprAST::CachedValue))
-	{
-		IndexExpr = IndexExpr->To<FRigVMCachedValueExprAST>()->GetVarExpr();
-	}
-
-	FRigVMOperand& IndexOperand = WorkData.ExprToOperand.FindChecked(GetSourceVarExpr(IndexExpr));
-	FRigVMOperand& ResultOperand = WorkData.ExprToOperand.FindChecked(GetSourceVarExpr(ResultExpr));
-
-	// setup the jumps for each case
-	TArray<uint64> JumpToCaseBytes;
-	TArray<int32> JumpToCaseInstructions;
-	JumpToCaseBytes.Add(0);
-	JumpToCaseInstructions.Add(0);
-
-	for (int32 CaseIndex = 1; CaseIndex < NumCases; CaseIndex++)
-	{
-		// compare and jump eventually
-		WorkData.VM->GetByteCode().AddEqualsOp(IndexOperand, WorkData.IntegerLiterals.FindChecked(CaseIndex), WorkData.ComparisonOperand);
-		if (Settings.SetupNodeInstructionIndex)
-		{
-			WorkData.VM->GetByteCode().SetSubject(WorkData.VM->GetByteCode().GetNumInstructions() - 1, Callstack.GetCallPath(), Callstack.GetStack());
-		}
-
-		uint64 JumpByte = WorkData.VM->GetByteCode().AddJumpIfOp(ERigVMOpCode::JumpForwardIf, 1, WorkData.ComparisonOperand, true);
-		int32 JumpInstruction = WorkData.VM->GetByteCode().GetNumInstructions() - 1;
-		if (Settings.SetupNodeInstructionIndex)
-		{
-			WorkData.VM->GetByteCode().SetSubject(WorkData.VM->GetByteCode().GetNumInstructions() - 1, Callstack.GetCallPath(), Callstack.GetStack());
-		}
-
-		JumpToCaseBytes.Add(JumpByte);
-		JumpToCaseInstructions.Add(JumpInstruction);
-	}
-
-	TArray<uint64> JumpToEndBytes;
-	TArray<int32> JumpToEndInstructions;
-
-	for (int32 CaseIndex = 0; CaseIndex < NumCases; CaseIndex++)
-	{
-		if (CaseIndex > 0)
-		{
-			int32 NumInstructionsInCase = WorkData.VM->GetByteCode().GetNumInstructions() - JumpToCaseInstructions[CaseIndex];
-			WorkData.VM->GetByteCode().GetOpAt<FRigVMJumpIfOp>(JumpToCaseBytes[CaseIndex]).InstructionIndex = NumInstructionsInCase;
-		}
-
-		TraverseExpression(CaseExpressions[CaseIndex], WorkData);
-
-		// add copy op to copy the result
-		FRigVMOperand& CaseOperand = WorkData.ExprToOperand.FindChecked(GetSourceVarExpr(CaseExpressions[CaseIndex]));
-		WorkData.VM->GetByteCode().AddCopyOp(WorkData.VM->GetCopyOpForOperands(CaseOperand, ResultOperand));
-		if (Settings.SetupNodeInstructionIndex)
-		{
-			WorkData.VM->GetByteCode().SetSubject(WorkData.VM->GetByteCode().GetNumInstructions() - 1, Callstack.GetCallPath(), Callstack.GetStack());
-		}
-		
-		if (CaseIndex < NumCases - 1)
-		{
-			uint64 JumpByte = WorkData.VM->GetByteCode().AddJumpOp(ERigVMOpCode::JumpForward, 1);
-			int32 JumpInstruction = WorkData.VM->GetByteCode().GetNumInstructions() - 1;
-			if (Settings.SetupNodeInstructionIndex)
-			{
-				WorkData.VM->GetByteCode().SetSubject(WorkData.VM->GetByteCode().GetNumInstructions() - 1, Callstack.GetCallPath(), Callstack.GetStack());
-			}
-
-			JumpToEndBytes.Add(JumpByte);
-			JumpToEndInstructions.Add(JumpInstruction);
-		}
-	}
-
-	for (int32 CaseIndex = 0; CaseIndex < NumCases - 1; CaseIndex++)
-	{
-		int32 NumInstructionsToEnd = WorkData.VM->GetByteCode().GetNumInstructions() - JumpToEndInstructions[CaseIndex];
-		WorkData.VM->GetByteCode().GetOpAt<FRigVMJumpOp>(JumpToEndBytes[CaseIndex]).InstructionIndex = NumInstructionsToEnd;
-	}
-}
-
 void URigVMCompiler::TraverseArray(const FRigVMArrayExprAST* InExpr, FRigVMCompilerWorkData& WorkData)
 {
 	URigVMArrayNode* ArrayNode = Cast<URigVMArrayNode>(InExpr->GetNode());
diff --git a/Engine/Source/Developer/RigVMDeveloper/Private/RigVMModel/Nodes/RigVMIfNode.cpp b/Engine/Source/Developer/RigVMDeveloper/Private/RigVMModel/Nodes/RigVMIfNode.cpp
--- a/Engine/Source/Developer/RigVMDeveloper/Private/RigVMModel/Nodes/RigVMIfNode.cpp
+++ b/Engine/Source/Developer/RigVMDeveloper/Private/RigVMModel/Nodes/RigVMIfNode.cpp
@@ -4,21 +4,21 @@
 
 #include UE_INLINE_GENERATED_CPP_BY_NAME(RigVMIfNode)
 
-const FString URigVMIfNode::IfName = TEXT("If");
-const FString URigVMIfNode::ConditionName = TEXT("Condition");
-const FString URigVMIfNode::TrueName = TEXT("True");
-const FString URigVMIfNode::FalseName = TEXT("False");
-const FString URigVMIfNode::ResultName = TEXT("Result");
+const FString UDEPRECATED_RigVMIfNode::IfName = TEXT("If");
+const FString UDEPRECATED_RigVMIfNode::ConditionName = TEXT("Condition");
+const FString UDEPRECATED_RigVMIfNode::TrueName = TEXT("True");
+const FString UDEPRECATED_RigVMIfNode::FalseName = TEXT("False");
+const FString UDEPRECATED_RigVMIfNode::ResultName = TEXT("Result");
 
 
-FName URigVMIfNode::GetNotation() const
+FName UDEPRECATED_RigVMIfNode::GetNotation() const
 {
 	static constexpr TCHAR Format[] = TEXT("%s(in %s,in %s,in %s,out %s)");
 	static const FName Notation = *FString::Printf(Format, *IfName, *ConditionName, *TrueName, *FalseName, *ResultName);
 	return Notation;
 }
 
-const FRigVMTemplate* URigVMIfNode::GetTemplate() const
+const FRigVMTemplate* UDEPRECATED_RigVMIfNode::GetTemplate() const
 {
 	if(const FRigVMTemplate* SuperTemplate = Super::GetTemplate())
 	{
diff --git a/Engine/Source/Developer/RigVMDeveloper/Private/RigVMModel/Nodes/RigVMSelectNode.cpp b/Engine/Source/Developer/RigVMDeveloper/Private/RigVMModel/Nodes/RigVMSelectNode.cpp
--- a/Engine/Source/Developer/RigVMDeveloper/Private/RigVMModel/Nodes/RigVMSelectNode.cpp
+++ b/Engine/Source/Developer/RigVMDeveloper/Private/RigVMModel/Nodes/RigVMSelectNode.cpp
@@ -4,12 +4,12 @@
 
 #include UE_INLINE_GENERATED_CPP_BY_NAME(RigVMSelectNode)
 
-const FString URigVMSelectNode::SelectName = TEXT("Select");
-const FString URigVMSelectNode::IndexName = TEXT("Index");
-const FString URigVMSelectNode::ValueName = TEXT("Values");
-const FString URigVMSelectNode::ResultName = TEXT("Result");
+const FString UDEPRECATED_RigVMSelectNode::SelectName = TEXT("Select");
+const FString UDEPRECATED_RigVMSelectNode::IndexName = TEXT("Index");
+const FString UDEPRECATED_RigVMSelectNode::ValueName = TEXT("Values");
+const FString UDEPRECATED_RigVMSelectNode::ResultName = TEXT("Result");
 
-bool URigVMSelectNode::AllowsLinksOn(const URigVMPin* InPin) const
+bool UDEPRECATED_RigVMSelectNode::AllowsLinksOn(const URigVMPin* InPin) const
 {
 	if(InPin->GetRootPin() == InPin)
 	{
@@ -22,14 +22,14 @@ bool URigVMSelectNode::AllowsLinksOn(const URigVMPin* InPin) const
 	return true;
 }
 
-FName URigVMSelectNode::GetNotation() const
+FName UDEPRECATED_RigVMSelectNode::GetNotation() const
 {
 	static constexpr TCHAR Format[] = TEXT("%s(in %s,in %s,out %s)");
 	static const FName Notation = *FString::Printf(Format, *SelectName, *IndexName, *ValueName, *ResultName);
 	return Notation;
 }
 
-const FRigVMTemplate* URigVMSelectNode::GetTemplate() const
+const FRigVMTemplate* UDEPRECATED_RigVMSelectNode::GetTemplate() const
 {
 	if(const FRigVMTemplate* SuperTemplate = Super::GetTemplate())
 	{
diff --git a/Engine/Source/Developer/RigVMDeveloper/Private/RigVMModel/RigVMClient.cpp b/Engine/Source/Developer/RigVMDeveloper/Private/RigVMModel/RigVMClient.cpp
--- a/Engine/Source/Developer/RigVMDeveloper/Private/RigVMModel/RigVMClient.cpp
+++ b/Engine/Source/Developer/RigVMDeveloper/Private/RigVMModel/RigVMClient.cpp
@@ -667,6 +667,7 @@ void FRigVMClient::PatchModelsOnLoad()
 		URigVMController* Controller = GetOrCreateController(Model);
 		Controller->PatchDispatchNodesOnLoad();
 		Controller->PatchBranchNodesOnLoad();
+		Controller->PatchIfSelectNodesOnLoad();
 	}
 }
 
diff --git a/Engine/Source/Developer/RigVMDeveloper/Private/RigVMModel/RigVMController.cpp b/Engine/Source/Developer/RigVMDeveloper/Private/RigVMModel/RigVMController.cpp
--- a/Engine/Source/Developer/RigVMDeveloper/Private/RigVMModel/RigVMController.cpp
+++ b/Engine/Source/Developer/RigVMDeveloper/Private/RigVMModel/RigVMController.cpp
@@ -24,6 +24,8 @@
 #include "RigVMPythonUtils.h"
 #include "RigVMTypeUtils.h"
 #include "Engine/UserDefinedStruct.h"
+#include "RigVMFunctions/RigVMDispatch_If.h"
+#include "RigVMFunctions/RigVMDispatch_Select.h"
 #include "RigVMModel/RigVMClient.h"
 #include "RigVMModel/Nodes/RigVMBranchNode.h"
 
@@ -600,28 +602,6 @@ TArray<FString> URigVMController::GetAddNodePythonCommands(URigVMNode* Node) con
 					*RigVMPythonUtils::LinearColorToPythonString(CommentNode->GetNodeColor()),
 					*NodeName));	
 	}
-	else if (const URigVMIfNode* IfNode = Cast<URigVMIfNode>(Node))
-	{
-		// add_if_node(cpp_type, cpp_type_object_path, position=[0.0, 0.0], node_name='', undo=True)
-		URigVMPin* ResultPin = IfNode->FindPin(URigVMIfNode::ResultName);
-		Commands.Add(FString::Printf(TEXT("blueprint.get_controller_by_name('%s').add_if_node('%s', '%s', %s, '%s')"),
-						*GraphName,
-						*ResultPin->GetCPPType(),
-						*ResultPin->CPPTypeObject->GetPathName(),
-						*RigVMPythonUtils::Vector2DToPythonString(IfNode->GetPosition()),
-						*NodeName));
-	}
-	else if (const URigVMSelectNode* SelectNode = Cast<URigVMSelectNode>(Node))
-	{
-		// add_select_node(cpp_type, cpp_type_object_path, position=[0.0, 0.0], node_name='', undo=True)
-		URigVMPin* ResultPin = SelectNode->FindPin(URigVMSelectNode::ResultName);
-		Commands.Add(FString::Printf(TEXT("blueprint.get_controller_by_name('%s').add_select_node('%s', '%s', %s, '%s')"),
-						*GraphName,
-						*ResultPin->GetCPPType(),
-						*ResultPin->CPPTypeObject->GetPathName(),
-						*RigVMPythonUtils::Vector2DToPythonString(SelectNode->GetPosition()),
-						*NodeName));
-	}
 	else if (const URigVMRerouteNode* RerouteNode = Cast<URigVMRerouteNode>(Node))
 	{
 		// add_free_reroute_node(bool bShowAsFullNode, const FString& InCPPType, const FName& InCPPTypeObjectPath, bool bIsConstant, const FName& InCustomWidgetName, const FString& InDefaultValue, const FVector2D& InPosition = FVector2D::ZeroVector, const FString& InNodeName = TEXT(""), bool bSetupUndoRedo = true);
@@ -8483,6 +8463,15 @@ bool URigVMController::RemoveArrayPin(const FString& InArrayElementPinPath, bool
 		ActionStack->BeginAction(Action);
 	}
 
+	// we need to keep at least one element for fixed size arrays
+	if(ArrayPin->IsExecuteContext() || ArrayPin->IsFixedSizeArray())
+	{
+		if(ArrayPin->GetArraySize() == 1)
+		{
+			return false;
+		}
+	}
+
 	int32 IndexToRemove = ArrayElementPin->GetPinIndex();
 	if (!RemovePin(ArrayElementPin, bSetupUndoRedo))
 	{
@@ -12829,7 +12818,7 @@ URigVMNode* URigVMController::AddBranchNode(const FVector2D& InPosition, const F
 	return AddUnitNode(FRigVMFunction_ControlFlowBranch::StaticStruct(), FRigVMStruct::ExecuteName, InPosition, InNodeName, bSetupUndoRedo, bPrintPythonCommand);
 }
 
-URigVMIfNode* URigVMController::AddIfNode(const FString& InCPPType, const FName& InCPPTypeObjectPath, const FVector2D& InPosition, const FString& InNodeName, bool  bSetupUndoRedo, bool bPrintPythonCommand)
+URigVMNode* URigVMController::AddIfNode(const FString& InCPPType, const FName& InCPPTypeObjectPath, const FVector2D& InPosition, const FString& InNodeName, bool  bSetupUndoRedo, bool bPrintPythonCommand)
 {
 	if(!IsValidGraph())
 	{
@@ -12857,109 +12846,31 @@ URigVMIfNode* URigVMController::AddIfNode(const FString& InCPPType, const FName&
 		}
 	}
 
-	FString CPPType = RigVMTypeUtils::PostProcessCPPType(InCPPType, CPPTypeObject);
-
-	FString DefaultValue;
-	if(UScriptStruct* ScriptStruct = Cast<UScriptStruct>(CPPTypeObject))
-	{
-		if (ScriptStruct->IsChildOf(FRigVMExecuteContext::StaticStruct()))
-		{
-			ReportErrorf(TEXT("Cannot create an if node for this type '%s'."), *InCPPTypeObjectPath.ToString());
-			return nullptr;
-		}
-		CreateDefaultValueForStructIfRequired(ScriptStruct, DefaultValue);
-	}
-	
-	FString Name = GetValidNodeName(InNodeName.IsEmpty() ? FString(TEXT("IfNode")) : InNodeName);
-	URigVMIfNode* Node = NewObject<URigVMIfNode>(Graph, *Name);
-	Node->Position = InPosition;
-
-	URigVMPin* ConditionPin = NewObject<URigVMPin>(Node, *URigVMIfNode::ConditionName);
-	ConditionPin->CPPType = RigVMTypeUtils::BoolType;
-	ConditionPin->Direction = ERigVMPinDirection::Input;
-	AddNodePin(Node, ConditionPin);
-
-	URigVMPin* TruePin = NewObject<URigVMPin>(Node, *URigVMIfNode::TrueName);
-	TruePin->CPPType = CPPType;
-	TruePin->CPPTypeObject = CPPTypeObject;
-	TruePin->CPPTypeObjectPath = InCPPTypeObjectPath;
-	TruePin->Direction = ERigVMPinDirection::Input;
-	TruePin->DefaultValue = DefaultValue;
-	AddNodePin(Node, TruePin);
-
-	if (TruePin->IsStruct())
-	{
-		TGuardValue<bool> SuspendNotifications(bSuspendNotifications, true);
-		AddPinsForStruct(TruePin->GetScriptStruct(), Node, TruePin, TruePin->Direction, FString(), false);
-	}
-
-	URigVMPin* FalsePin = NewObject<URigVMPin>(Node, *URigVMIfNode::FalseName);
-	FalsePin->CPPType = CPPType;
-	FalsePin->CPPTypeObject = CPPTypeObject;
-	FalsePin->CPPTypeObjectPath = InCPPTypeObjectPath;
-	FalsePin->Direction = ERigVMPinDirection::Input;
-	FalsePin->DefaultValue = DefaultValue;
-	AddNodePin(Node, FalsePin);
-
-	if (FalsePin->IsStruct())
-	{
-		TGuardValue<bool> SuspendNotifications(bSuspendNotifications, true);
-		AddPinsForStruct(FalsePin->GetScriptStruct(), Node, FalsePin, FalsePin->Direction, FString(), false);
-	}
-
-	URigVMPin* ResultPin = NewObject<URigVMPin>(Node, *URigVMIfNode::ResultName);
-	ResultPin->CPPType = CPPType;
-	ResultPin->CPPTypeObject = CPPTypeObject;
-	ResultPin->CPPTypeObjectPath = InCPPTypeObjectPath;
-	ResultPin->Direction = ERigVMPinDirection::Output;
-	AddNodePin(Node, ResultPin);
-
-	if (ResultPin->IsStruct())
-	{
-		TGuardValue<bool> SuspendNotifications(bSuspendNotifications, true);
-		AddPinsForStruct(ResultPin->GetScriptStruct(), Node, ResultPin, ResultPin->Direction, FString(), false);
-	}
-
-	Graph->Nodes.Add(Node);
-
-	Notify(ERigVMGraphNotifType::NodeAdded, Node);
-
-	FRigVMControllerCompileBracketScope CompileScope(this);
-	if (bSetupUndoRedo)
+	if(bSetupUndoRedo)
 	{
-		ActionStack->AddAction(FRigVMAddIfNodeAction(Node));
+		OpenUndoBracket(TEXT("Add If Node"));
 	}
 
-	Node->InitializeFilteredPermutations();
-	if (InCPPType != RigVMTypeUtils::GetWildCardCPPType() && InCPPType != RigVMTypeUtils::GetWildCardArrayCPPType())
+	const FString CPPType = RigVMTypeUtils::PostProcessCPPType(InCPPType, CPPTypeObject);
+	const FString Name = GetValidNodeName(InNodeName.IsEmpty() ? FString(TEXT("IfNode")) : InNodeName);
+	const TRigVMTypeIndex& TypeIndex = FRigVMRegistry::Get().FindOrAddType({*CPPType, CPPTypeObject});
+	
+	const FRigVMDispatchFactory* Factory = FRigVMRegistry::Get().FindOrAddDispatchFactory(FRigVMDispatch_If::StaticStruct());
+	URigVMNode* Node = AddTemplateNode(Factory->GetTemplate()->GetNotation(), InPosition, Name, bSetupUndoRedo, bPrintPythonCommand);
+	if(Node)
 	{
-		PrepareTemplatePinForType(TruePin, {TruePin->GetTypeIndex()}, bSetupUndoRedo);
-		TArray<int32> FilterPermutations = Node->GetFilteredPermutationsIndices();
-		if (FilterPermutations.Num() == 1)
-		{
-			const TArray<FRigVMTemplatePreferredType> NewPreferredPermutationTypes = Node->GetPreferredTypesForPermutation(FilterPermutations[0]);
-			if (bSetupUndoRedo)
-			{
-				ActionStack->AddAction(FRigVMSetPreferredTemplatePermutationsAction(Node, NewPreferredPermutationTypes));
-			}
-			Node->PreferredPermutationPairs = NewPreferredPermutationTypes;
-		}
+		ResolveWildCardPin(Node->GetPins().Last(), TypeIndex, bSetupUndoRedo, bPrintPythonCommand);
 	}
 	
-	if (bPrintPythonCommand)
+	if(bSetupUndoRedo)
 	{
-		TArray<FString> Commands = GetAddNodePythonCommands(Node);
-		for (const FString& Command : Commands)
-		{
-			RigVMPythonUtils::Print(GetGraphOuterName(), 
-				FString::Printf(TEXT("%s"), *Command));
-		}
+		CloseUndoBracket();
 	}
 
 	return Node;
 }
 
-URigVMIfNode* URigVMController::AddIfNodeFromStruct(UScriptStruct* InScriptStruct, const FVector2D& InPosition, const FString& InNodeName, bool bSetupUndoRedo)
+URigVMNode* URigVMController::AddIfNodeFromStruct(UScriptStruct* InScriptStruct, const FVector2D& InPosition, const FString& InNodeName, bool bSetupUndoRedo)
 {
 	if (!InScriptStruct)
 	{
@@ -12969,7 +12880,7 @@ URigVMIfNode* URigVMController::AddIfNodeFromStruct(UScriptStruct* InScriptStruc
 	return AddIfNode(RigVMTypeUtils::GetUniqueStructTypeName(InScriptStruct), FName(InScriptStruct->GetPathName()), InPosition, InNodeName, bSetupUndoRedo);
 }
 
-URigVMSelectNode* URigVMController::AddSelectNode(const FString& InCPPType, const FName& InCPPTypeObjectPath, const FVector2D& InPosition, const FString& InNodeName, bool bSetupUndoRedo, bool bPrintPythonCommand)
+URigVMNode* URigVMController::AddSelectNode(const FString& InCPPType, const FName& InCPPTypeObjectPath, const FVector2D& InPosition, const FString& InNodeName, bool bSetupUndoRedo, bool bPrintPythonCommand)
 {
 	if (!IsValidGraph())
 	{
@@ -12997,91 +12908,30 @@ URigVMSelectNode* URigVMController::AddSelectNode(const FString& InCPPType, cons
 		}
 	}
 
-	FString CPPType = RigVMTypeUtils::PostProcessCPPType(InCPPType, CPPTypeObject);
-
-	FString DefaultValue;
-	if (UScriptStruct* ScriptStruct = Cast<UScriptStruct>(CPPTypeObject))
-	{
-		if (ScriptStruct->IsChildOf(FRigVMExecuteContext::StaticStruct()))
-		{
-			ReportErrorf(TEXT("Cannot create a select node for this type '%s'."), *InCPPTypeObjectPath.ToString());
-			return nullptr;
-		}
-		CreateDefaultValueForStructIfRequired(ScriptStruct, DefaultValue);
-	}
-
-	FString Name = GetValidNodeName(InNodeName.IsEmpty() ? FString(TEXT("IfNode")) : InNodeName);
-	URigVMSelectNode* Node = NewObject<URigVMSelectNode>(Graph, *Name);
-	Node->Position = InPosition;
-
-	URigVMPin* IndexPin = NewObject<URigVMPin>(Node, *URigVMSelectNode::IndexName);
-	IndexPin->CPPType = RigVMTypeUtils::Int32Type;
-	IndexPin->Direction = ERigVMPinDirection::Input;
-	AddNodePin(Node, IndexPin);
-
-	URigVMPin* ValuePin = NewObject<URigVMPin>(Node, *URigVMSelectNode::ValueName);
-	ValuePin->CPPType = RigVMTypeUtils::ArrayTypeFromBaseType(CPPType);
-	ValuePin->CPPTypeObject = CPPTypeObject;
-	ValuePin->CPPTypeObjectPath = InCPPTypeObjectPath;
-	ValuePin->Direction = ERigVMPinDirection::Input;
-	ValuePin->bIsExpanded = true;
-	AddNodePin(Node, ValuePin);
-
-	URigVMPin* ResultPin = NewObject<URigVMPin>(Node, *URigVMSelectNode::ResultName);
-	ResultPin->CPPType = CPPType;
-	ResultPin->CPPTypeObject = CPPTypeObject;
-	ResultPin->CPPTypeObjectPath = InCPPTypeObjectPath;
-	ResultPin->Direction = ERigVMPinDirection::Output;
-	AddNodePin(Node, ResultPin);
-
-	if (ResultPin->IsStruct())
+	if(bSetupUndoRedo)
 	{
-		TGuardValue<bool> SuspendNotifications(bSuspendNotifications, true);
-		AddPinsForStruct(ResultPin->GetScriptStruct(), Node, ResultPin, ResultPin->Direction, FString(), false);
+		OpenUndoBracket(TEXT("Add Select Node"));
 	}
 
-	Graph->Nodes.Add(Node);
-
-	Notify(ERigVMGraphNotifType::NodeAdded, Node);
-
-	SetArrayPinSize(ValuePin->GetPinPath(), 2, DefaultValue, false);
-
-	FRigVMControllerCompileBracketScope CompileScope(this);
-	if (bSetupUndoRedo)
+	const FString CPPType = RigVMTypeUtils::PostProcessCPPType(InCPPType, CPPTypeObject);
+	const FString Name = GetValidNodeName(InNodeName.IsEmpty() ? FString(TEXT("SelectNode")) : InNodeName);
+	const TRigVMTypeIndex& TypeIndex = FRigVMRegistry::Get().FindOrAddType({*CPPType, CPPTypeObject});
+	
+	const FRigVMDispatchFactory* Factory = FRigVMRegistry::Get().FindOrAddDispatchFactory(FRigVMDispatch_SelectInt32::StaticStruct());
+	URigVMNode* Node = AddTemplateNode(Factory->GetTemplate()->GetNotation(), InPosition, Name, bSetupUndoRedo, bPrintPythonCommand);
+	if(Node)
 	{
-		ActionStack->AddAction(FRigVMAddSelectNodeAction(Node));
+		ResolveWildCardPin(Node->GetPins().Last(), TypeIndex, bSetupUndoRedo, bPrintPythonCommand);
 	}
 
-	Node->InitializeFilteredPermutations();
-	if (InCPPType != RigVMTypeUtils::GetWildCardCPPType() && InCPPType != RigVMTypeUtils::GetWildCardArrayCPPType())
-	{
-		PrepareTemplatePinForType(ResultPin, {ResultPin->GetTypeIndex()}, bSetupUndoRedo);
-		TArray<int32> FilterPermutations = Node->GetFilteredPermutationsIndices();
-		if (FilterPermutations.Num() == 1)
-		{
-			const TArray<FRigVMTemplatePreferredType> NewPreferredPermutationTypes = Node->GetPreferredTypesForPermutation(FilterPermutations[0]);
-			if (bSetupUndoRedo)
-			{
-				ActionStack->AddAction(FRigVMSetPreferredTemplatePermutationsAction(Node, NewPreferredPermutationTypes));
-			}
-			Node->PreferredPermutationPairs = NewPreferredPermutationTypes;
-		}
-	}
-	
-	if (bPrintPythonCommand)
+	if(bSetupUndoRedo)
 	{
-		TArray<FString> Commands = GetAddNodePythonCommands(Node);
-		for (const FString& Command : Commands)
-		{
-			RigVMPythonUtils::Print(GetGraphOuterName(), 
-				FString::Printf(TEXT("%s"), *Command));
-		}
+		CloseUndoBracket();
 	}
-
 	return Node;
 }
 
-URigVMSelectNode* URigVMController::AddSelectNodeFromStruct(UScriptStruct* InScriptStruct, const FVector2D& InPosition,
+URigVMNode* URigVMController::AddSelectNodeFromStruct(UScriptStruct* InScriptStruct, const FVector2D& InPosition,
 	const FString& InNodeName, bool bSetupUndoRedo)
 {
 	if (!InScriptStruct)
@@ -14415,6 +14265,21 @@ void URigVMController::AddPinsForTemplate(const FRigVMTemplate* InTemplate, cons
         Pin->LastKnownCPPType = Pin->CPPType;
 
         AddNodePin(InNode, Pin);
+
+		if(Registry.IsArrayType(Arg->TypeIndex))
+		{
+			if(const URigVMDispatchNode* DispatchNode = Cast<URigVMDispatchNode>(Pin->GetNode()))
+			{
+				if(const FRigVMDispatchFactory* Factory = DispatchNode->GetFactory())
+				{
+					const FString DefaultValue =  Factory->GetArgumentDefaultValue(Pin->GetFName(), Arg->TypeIndex);
+					if(!DefaultValue.IsEmpty())
+					{
+						SetPinDefaultValue(Pin, DefaultValue, true, false, false);
+					}
+				}
+			}
+		}
     }
 
 	for (int32 ArgIndex = 0; ArgIndex < InTemplate->NumArguments(); ArgIndex++)
@@ -14667,10 +14532,11 @@ bool URigVMController::ShouldPinBeUnfolded(URigVMPin* InPin)
 	{
 		return ShouldStructBeUnfolded(InPin->GetScriptStruct());
 	}
-	else if (InPin->IsArray())
+	if (InPin->IsArray())
 	{
 		return InPin->GetDirection() == ERigVMPinDirection::Input ||
-			InPin->GetDirection() == ERigVMPinDirection::IO;
+			InPin->GetDirection() == ERigVMPinDirection::IO ||
+			InPin->IsFixedSizeArray();
 	}
 	return false;
 }
@@ -15194,8 +15060,6 @@ void URigVMController::RepopulatePinsOnNode(URigVMNode* InNode, bool bFollowCore
 	URigVMCollapseNode* CollapseNode = Cast<URigVMCollapseNode>(InNode);
 	URigVMFunctionReferenceNode* FunctionRefNode = Cast<URigVMFunctionReferenceNode>(InNode);
 	URigVMVariableNode* VariableNode = Cast<URigVMVariableNode>(InNode);
-	URigVMIfNode* IfNode = Cast<URigVMIfNode>(InNode);
-	URigVMSelectNode* SelectNode = Cast<URigVMSelectNode>(InNode);
 	URigVMArrayNode* ArrayNode = Cast<URigVMArrayNode>(InNode);
 	URigVMDispatchNode* DispatchNode = Cast<URigVMDispatchNode>(InNode);
 
@@ -15552,7 +15416,7 @@ void URigVMController::RepopulatePinsOnNode(URigVMNode* InNode, bool bFollowCore
 			ApplyPinStates(InNode, ReferencedPinStates);
 		}
 	}
-	else if (IfNode || SelectNode || ArrayNode)
+	else if (ArrayNode)
 	{
 		// PinStates were already saved earlier, and will be applied later, no need to do it here
 	 
@@ -17023,6 +16887,11 @@ bool URigVMController::ResolveWildCardPin(URigVMPin* InPin, TRigVMTypeIndex InTy
 		return false;
 	}
 
+	if(FRigVMRegistry::Get().IsWildCardType(InTypeIndex))
+	{
+		return false;
+	}
+
 	URigVMPin* RootPin = InPin;
 	TRigVMTypeIndex Type = InTypeIndex;	
 	while (RootPin->IsArrayElement())
@@ -19823,6 +19692,53 @@ void URigVMController::PatchBranchNodesOnLoad()
 	}
 }
 
+void URigVMController::PatchIfSelectNodesOnLoad()
+{
+	if (const URigVMGraph* Graph = GetGraph())
+	{
+		TArray<URigVMNode*> IfOrSelectNodes = Graph->GetNodes().FilterByPredicate([](URigVMNode* Node)
+		{
+			return Node->IsA<UDEPRECATED_RigVMIfNode>() ||
+				Node->IsA<UDEPRECATED_RigVMSelectNode>();
+		});
+
+		for(URigVMNode* IfOrSelectNode : IfOrSelectNodes)
+		{
+			const bool bIsIfNode = IfOrSelectNode->IsA<UDEPRECATED_RigVMIfNode>();
+			TArray<TPair<FString, FString>> LinkedPaths = GetLinkedPinPaths(IfOrSelectNode);
+			const FVector2D NodePosition = IfOrSelectNode->GetPosition();
+			const FString NodeName = IfOrSelectNode->GetName();
+			const TRigVMTypeIndex TypeIndex = IfOrSelectNode->GetPins().Last()->GetTypeIndex();
+			TMap<FString, FPinState> PinStates = GetPinStates(IfOrSelectNode, true);
+			
+			RemoveNode(IfOrSelectNode, false, true, false, false);
+
+			const FRigVMDispatchFactory* Factory = FRigVMRegistry::Get().FindOrAddDispatchFactory(
+				bIsIfNode ? FRigVMDispatch_If::StaticStruct() : FRigVMDispatch_SelectInt32::StaticStruct());
+
+			const FRigVMTemplate* Template = Factory->GetTemplate();
+			URigVMTemplateNode* NewNode = AddTemplateNode(
+				Template->GetNotation(),
+				NodePosition, 
+				NodeName, 
+				false, 
+				false);
+
+			for(URigVMPin* Pin : NewNode->GetPins())
+			{
+				if(Pin->IsWildCard() && !FRigVMRegistry::Get().IsWildCardType(TypeIndex))
+				{
+					ChangePinType(Pin, TypeIndex, false, false);
+					NewNode->UpdateFilteredPermutations(Pin, {TypeIndex});
+				}
+			}
+
+			ApplyPinStates(NewNode, PinStates, {}, false);
+			RestoreLinkedPaths(LinkedPaths, {}, {}, false);
+		}
+	}
+}
+
 #endif
 
 URigVMControllerSettings::URigVMControllerSettings(const FObjectInitializer& Initializer)
diff --git a/Engine/Source/Developer/RigVMDeveloper/Private/RigVMModel/RigVMControllerActions.cpp b/Engine/Source/Developer/RigVMDeveloper/Private/RigVMModel/RigVMControllerActions.cpp
--- a/Engine/Source/Developer/RigVMDeveloper/Private/RigVMModel/RigVMControllerActions.cpp
+++ b/Engine/Source/Developer/RigVMDeveloper/Private/RigVMModel/RigVMControllerActions.cpp
@@ -487,94 +487,6 @@ bool FRigVMAddRerouteNodeAction::Redo(URigVMController* InController)
 	return false;
 }
 
-FRigVMAddIfNodeAction::FRigVMAddIfNodeAction()
-	: CPPType()
-	, CPPTypeObjectPath(NAME_None)
-	, Position(FVector2D::ZeroVector)
-	, NodePath()
-{
-}
-
-FRigVMAddIfNodeAction::FRigVMAddIfNodeAction(URigVMIfNode* InNode)
-	: CPPType()
-	, CPPTypeObjectPath(NAME_None)
-	, Position(InNode->GetPosition())
-	, NodePath(InNode->GetNodePath())
-{
-	if(URigVMPin* ResultPin = InNode->FindPin(URigVMIfNode::ResultName))
-	{
-		CPPType = ResultPin->GetCPPType();
-		if (ResultPin->GetCPPTypeObject())
-		{
-			CPPTypeObjectPath = *ResultPin->GetCPPTypeObject()->GetPathName();
-		}
-	}
-}
-
-bool FRigVMAddIfNodeAction::Undo(URigVMController* InController)
-{
-	if (!FRigVMBaseAction::Undo(InController))
-	{
-		return false;
-	}
-	return InController->RemoveNodeByName(*NodePath, false);
-}
-
-bool FRigVMAddIfNodeAction::Redo(URigVMController* InController)
-{
-#if WITH_EDITOR
-	if (URigVMIfNode* Node = InController->AddIfNode(CPPType, CPPTypeObjectPath, Position, NodePath, false))
-	{
-		return FRigVMBaseAction::Redo(InController);
-	}
-#endif
-	return false;
-}
-
-FRigVMAddSelectNodeAction::FRigVMAddSelectNodeAction()
-	: CPPType()
-	, CPPTypeObjectPath(NAME_None)
-	, Position(FVector2D::ZeroVector)
-	, NodePath()
-{
-}
-
-FRigVMAddSelectNodeAction::FRigVMAddSelectNodeAction(URigVMSelectNode* InNode)
-	: CPPType()
-	, CPPTypeObjectPath(NAME_None)
-	, Position(InNode->GetPosition())
-	, NodePath(InNode->GetNodePath())
-{
-	if(URigVMPin* ResultPin = InNode->FindPin(URigVMSelectNode::ResultName))
-	{
-		CPPType = ResultPin->GetCPPType();
-		if (ResultPin->GetCPPTypeObject())
-		{
-			CPPTypeObjectPath = *ResultPin->GetCPPTypeObject()->GetPathName();
-		}
-	}
-}
-
-bool FRigVMAddSelectNodeAction::Undo(URigVMController* InController)
-{
-	if (!FRigVMBaseAction::Undo(InController))
-	{
-		return false;
-	}
-	return InController->RemoveNodeByName(*NodePath, false);
-}
-
-bool FRigVMAddSelectNodeAction::Redo(URigVMController* InController)
-{
-#if WITH_EDITOR
-	if (URigVMSelectNode* Node = InController->AddSelectNode(CPPType, CPPTypeObjectPath, Position, NodePath, false))
-	{
-		return FRigVMBaseAction::Redo(InController);
-	}
-#endif
-	return false;
-}
-
 FRigVMAddEnumNodeAction::FRigVMAddEnumNodeAction()
 	: CPPTypeObjectPath(NAME_None)
 	, Position(FVector2D::ZeroVector)
@@ -706,20 +618,6 @@ FRigVMRemoveNodeAction::FRigVMRemoveNodeAction(URigVMNode* InNode, URigVMControl
 	{
 		InverseAction.AddAction(FRigVMAddRerouteNodeAction(RerouteNode), InController);
 	}
-	else if (URigVMIfNode* IfNode = Cast<URigVMIfNode>(InNode))
-	{
-		InverseAction.AddAction(FRigVMAddIfNodeAction(IfNode), InController);
-		URigVMPin* TrueNamePin = IfNode->FindPin(URigVMIfNode::TrueName);
-		URigVMPin* FalseNamePin = IfNode->FindPin(URigVMIfNode::FalseName);
-		InverseAction.AddAction(FRigVMSetPinDefaultValueAction(TrueNamePin, TrueNamePin->GetDefaultValue()), InController);
-		InverseAction.AddAction(FRigVMSetPinDefaultValueAction(FalseNamePin, FalseNamePin->GetDefaultValue()), InController);
-	}
-	else if (URigVMSelectNode* SelectNode = Cast<URigVMSelectNode>(InNode))
-	{
-		InverseAction.AddAction(FRigVMAddSelectNodeAction(SelectNode), InController);
-		URigVMPin* ValuesNamePin = SelectNode->FindPin(URigVMSelectNode::ValueName);
-		InverseAction.AddAction(FRigVMSetPinDefaultValueAction(ValuesNamePin, ValuesNamePin->GetDefaultValue()), InController);		
-	}
 	else if (URigVMEnumNode* EnumNode = Cast<URigVMEnumNode>(InNode))
 	{
 		InverseAction.AddAction(FRigVMAddEnumNodeAction(EnumNode), InController);
diff --git a/Engine/Source/Developer/RigVMDeveloper/Private/RigVMModel/RigVMNode.cpp b/Engine/Source/Developer/RigVMDeveloper/Private/RigVMModel/RigVMNode.cpp
--- a/Engine/Source/Developer/RigVMDeveloper/Private/RigVMModel/RigVMNode.cpp
+++ b/Engine/Source/Developer/RigVMDeveloper/Private/RigVMModel/RigVMNode.cpp
@@ -1,6 +1,7 @@
 // Copyright Epic Games, Inc. All Rights Reserved.
 
 #include "RigVMModel/RigVMNode.h"
+#include "RigVMStringUtils.h"
 #include "RigVMModel/Nodes/RigVMUnitNode.h"
 #include "RigVMModel/RigVMGraph.h"
 #include "RigVMModel/Nodes/RigVMLibraryNode.h"
@@ -53,56 +54,27 @@ FString URigVMNode::GetNodePath(bool bRecursive) const
 
 bool URigVMNode::SplitNodePathAtStart(const FString& InNodePath, FString& LeftMost, FString& Right)
 {
-	return InNodePath.Split(TEXT("|"), &LeftMost, &Right, ESearchCase::IgnoreCase, ESearchDir::FromStart);
+	return RigVMStringUtils::SplitNodePathAtStart(InNodePath, LeftMost, Right);
 }
 
 bool URigVMNode::SplitNodePathAtEnd(const FString& InNodePath, FString& Left, FString& RightMost)
 {
-	return InNodePath.Split(TEXT("|"), &Left, &RightMost, ESearchCase::IgnoreCase, ESearchDir::FromEnd);
+	return RigVMStringUtils::SplitNodePathAtEnd(InNodePath, Left, RightMost);
 }
 
 bool URigVMNode::SplitNodePath(const FString& InNodePath, TArray<FString>& Parts)
 {
-	int32 OriginalPartsCount = Parts.Num();
-	FString NodePathRemaining = InNodePath;
-	FString Left, Right;
-	Right = NodePathRemaining;
-
-	while (SplitNodePathAtStart(NodePathRemaining, Left, Right))
-	{
-		Parts.Add(Left);
-		Left.Empty();
-		NodePathRemaining = Right;
-	}
-
-	if (!Right.IsEmpty())
-	{
-		Parts.Add(Right);
-	}
-
-	return Parts.Num() > OriginalPartsCount;
+	return RigVMStringUtils::SplitNodePath(InNodePath, Parts);
 }
 
 FString URigVMNode::JoinNodePath(const FString& Left, const FString& Right)
 {
-	ensure(!Left.IsEmpty() && !Right.IsEmpty());
-	return Left + TEXT("|") + Right;
+	return RigVMStringUtils::JoinNodePath(Left, Right);
 }
 
 FString URigVMNode::JoinNodePath(const TArray<FString>& InParts)
 {
-	if (InParts.Num() == 0)
-	{
-		return FString();
-	}
-
-	FString Result = InParts[0];
-	for (int32 PartIndex = 1; PartIndex < InParts.Num(); PartIndex++)
-	{
-		Result += TEXT("|") + InParts[PartIndex];
-	}
-
-	return Result;
+	return RigVMStringUtils::JoinNodePath(InParts);
 }
 
 int32 URigVMNode::GetNodeIndex() const
diff --git a/Engine/Source/Developer/RigVMDeveloper/Private/RigVMModel/RigVMPin.cpp b/Engine/Source/Developer/RigVMDeveloper/Private/RigVMModel/RigVMPin.cpp
--- a/Engine/Source/Developer/RigVMDeveloper/Private/RigVMModel/RigVMPin.cpp
+++ b/Engine/Source/Developer/RigVMDeveloper/Private/RigVMModel/RigVMPin.cpp
@@ -23,6 +23,7 @@
 #include "RigVMModel/Nodes/RigVMFunctionReturnNode.h"
 #include "RigVMModel/Nodes/RigVMInvokeEntryNode.h"
 #include "RigVMModel/Nodes/RigVMSelectNode.h"
+#include "RigVMStringUtils.h"
 
 #include UE_INLINE_GENERATED_CPP_BY_NAME(RigVMPin)
 
@@ -56,122 +57,32 @@ const FString URigVMPin::OrphanPinPrefix = TEXT("Orphan::");
 
 bool URigVMPin::SplitPinPathAtStart(const FString& InPinPath, FString& LeftMost, FString& Right)
 {
-	return InPinPath.Split(TEXT("."), &LeftMost, &Right, ESearchCase::IgnoreCase, ESearchDir::FromStart);
+	return RigVMStringUtils::SplitPinPathAtStart(InPinPath, LeftMost, Right);
 }
 
 bool URigVMPin::SplitPinPathAtEnd(const FString& InPinPath, FString& Left, FString& RightMost)
 {
-	return InPinPath.Split(TEXT("."), &Left, &RightMost, ESearchCase::IgnoreCase, ESearchDir::FromEnd);
+	return RigVMStringUtils::SplitPinPathAtEnd(InPinPath, Left, RightMost);
 }
 
 bool URigVMPin::SplitPinPath(const FString& InPinPath, TArray<FString>& Parts)
 {
-	int32 OriginalPartsCount = Parts.Num();
-	FString PinPathRemaining = InPinPath;
-	FString Left, Right;
-	while(SplitPinPathAtStart(PinPathRemaining, Left, Right))
-	{
-		Parts.Add(Left);
-		Left.Empty();
-		PinPathRemaining = Right;
-	}
-
-	if (!Right.IsEmpty())
-	{
-		Parts.Add(Right);
-	}
-
-	return Parts.Num() > OriginalPartsCount;
+	return RigVMStringUtils::SplitPinPath(InPinPath, Parts);
 }
 
 FString URigVMPin::JoinPinPath(const FString& Left, const FString& Right)
 {
-	ensure(!Left.IsEmpty() && !Right.IsEmpty());
-	return Left + TEXT(".") + Right;
+	return RigVMStringUtils::JoinPinPath(Left, Right);
 }
 
 FString URigVMPin::JoinPinPath(const TArray<FString>& InParts)
 {
-	if (InParts.Num() == 0)
-	{
-		return FString();
-	}
-
-	FString Result = InParts[0];
-	for (int32 PartIndex = 1; PartIndex < InParts.Num(); PartIndex++)
-	{
-		Result += TEXT(".") + InParts[PartIndex];
-	}
-
-	return Result;
+	return RigVMStringUtils::JoinPinPath(InParts);
 }
 
 TArray<FString> URigVMPin::SplitDefaultValue(const FString& InDefaultValue)
 {
-	TArray<FString> Parts;
-	if (InDefaultValue.IsEmpty())
-	{
-		return Parts;
-	}
-
-	if(InDefaultValue[0] != TCHAR('('))
-	{
-		return Parts;
-	}
-	if(InDefaultValue[InDefaultValue.Len() - 1] != TCHAR(')'))
-	{
-		return Parts;
-	}
-
-	FString Content = InDefaultValue.Mid(1, InDefaultValue.Len() - 2);
-	int32 BraceCount = 0;
-	int32 QuoteCount = 0;
-
-	int32 LastPartStartIndex = 0;
-	for (int32 CharIndex = 0; CharIndex < Content.Len(); CharIndex++)
-	{
-		TCHAR Char = Content[CharIndex];
-		if (QuoteCount > 0)
-		{
-			if (Char == TCHAR('"'))
-			{
-				QuoteCount = 0;
-			}
-		}
-		else if (Char == TCHAR('"'))
-		{
-			QuoteCount = 1;
-		}
-
-		if (Char == TCHAR('('))
-		{
-			if (QuoteCount == 0)
-			{
-				BraceCount++;
-			}
-		}
-		else if (Char == TCHAR(')'))
-		{
-			if (QuoteCount == 0)
-			{
-				BraceCount--;
-				BraceCount = FMath::Max<int32>(BraceCount, 0);
-			}
-		}
-		else if (Char == TCHAR(',') && BraceCount == 0 && QuoteCount == 0)
-		{
-			// ignore whitespaces
-			Parts.Add(Content.Mid(LastPartStartIndex, CharIndex - LastPartStartIndex).Replace(TEXT(" "), TEXT("")));
-			LastPartStartIndex = CharIndex + 1;
-		}
-	}
-
-	if (!Content.IsEmpty())
-	{
-		// ignore whitespaces from the start and end of the string
-		Parts.Add(Content.Mid(LastPartStartIndex).TrimStartAndEnd());
-	}
-	return Parts;
+	return RigVMStringUtils::SplitDefaultValue(InDefaultValue);
 }
 
 FString URigVMPin::GetDefaultValueForArray(TConstArrayView<FString> DefaultValues)
@@ -753,9 +664,9 @@ bool URigVMPin::IsFixedSizeArray() const
 					return Factory->HasArgumentMetaData(GetFName(), FRigVMStruct::FixedSizeArrayMetaName);
 				}
 			}
-			else if(Node->IsA<URigVMSelectNode>())
+			else if(Node->IsA<UDEPRECATED_RigVMSelectNode>())
 			{
-				return GetFName().ToString() == URigVMSelectNode::ValueName;
+				return GetFName().ToString() == UDEPRECATED_RigVMSelectNode::ValueName;
 			}
 		}
 	}
diff --git a/Engine/Source/Developer/RigVMDeveloper/Public/RigVMCompiler/RigVMAST.h b/Engine/Source/Developer/RigVMDeveloper/Public/RigVMCompiler/RigVMAST.h
--- a/Engine/Source/Developer/RigVMDeveloper/Public/RigVMCompiler/RigVMAST.h
+++ b/Engine/Source/Developer/RigVMDeveloper/Public/RigVMCompiler/RigVMAST.h
@@ -26,8 +26,6 @@ class FRigVMAssignExprAST;
 class FRigVMCopyExprAST;
 class FRigVMCachedValueExprAST;
 class FRigVMExitExprAST;
-class FRigVMIfExprAST;
-class FRigVMSelectExprAST;
 class FRigVMArrayExprAST;
 
 class URigVMPin;
@@ -93,9 +91,6 @@ public:
 		Copy,
 		CachedValue,
 		Exit,
-		Branch,
-		If,
-		Select,
 		Array,
 		Invalid
 	};
@@ -321,9 +316,6 @@ inline const FRigVMNodeExprAST* FRigVMExprAST::To() const
 		IsA(EType::CallExtern) ||
 		IsA(EType::InlineFunction) ||
 		IsA(EType::NoOp) ||
-		IsA(EType::Branch) ||
-		IsA(EType::If) ||
-		IsA(EType::Select) ||
 		IsA(EType::Array)
 	);
 	return (const FRigVMNodeExprAST*)this;
@@ -461,28 +453,6 @@ inline const FRigVMExitExprAST* FRigVMExprAST::To() const
 	return (const FRigVMExitExprAST*)this;
 }
 
-// specialized cast for type checking
-// for a If / FRigVMIfExprAST expression
-// will raise if types are not compatible
-// @return this expression cast to FRigVMIfExprAST 
-template<>
-inline const FRigVMIfExprAST* FRigVMExprAST::To() const
-{
-	ensure(IsA(EType::If));
-	return (const FRigVMIfExprAST*)this;
-}
-
-// specialized cast for type checking
-// for a Select / FRigVMSelectExprAST expression
-// will raise if types are not compatible
-// @return this expression cast to FRigVMSelectExprAST 
-template<>
-inline const FRigVMSelectExprAST* FRigVMExprAST::To() const
-{
-	ensure(IsA(EType::Select));
-	return (const FRigVMSelectExprAST*)this;
-}
-
 // specialized cast for type checking
 // for a Array / FRigVMArrayExprAST expression
 // will raise if types are not compatible
@@ -1106,90 +1076,6 @@ private:
 	friend class FRigVMParserAST;
 };
 
-/*
- * An abstract syntax tree if expression represents a branch point for values.
- * In C++ the branch is is the definition: (Condition ? True : False)
- */
-class RIGVMDEVELOPER_API FRigVMIfExprAST : public FRigVMNodeExprAST
-{
-public:
-
-	// virtual destructor
-	virtual ~FRigVMIfExprAST() {}
-
-	// disable copy constructor
-	FRigVMIfExprAST(const FRigVMEntryExprAST&) = delete;
-
-	// overload of the type checking mechanism
-	virtual bool IsA(EType InType) const override
-	{
-		return InType == EType::If;
-	};
-
-	virtual bool IsConstant() const override;
-
-	const FRigVMVarExprAST* GetConditionExpr() const { return ChildAt(0)->To<FRigVMVarExprAST>(); }
-	const FRigVMVarExprAST* GetTrueExpr() const { return ChildAt(1)->To<FRigVMVarExprAST>(); }
-	const FRigVMVarExprAST* GetFalseExpr() const { return ChildAt(2)->To<FRigVMVarExprAST>(); }
-	const FRigVMVarExprAST* GetResultExpr() const { return ChildAt(3)->To<FRigVMVarExprAST>(); }
-
-	bool IsAlwaysTrue() const;
-	bool IsAlwaysFalse() const;
-
-protected:
-
-	// default constructor (protected so that only parser can access it)
-	FRigVMIfExprAST(const FRigVMASTProxy& InProxy)
-		: FRigVMNodeExprAST(EType::If, InProxy)
-	{}
-
-
-private:
-
-	friend class FRigVMParserAST;
-};
-
-/*
- * An abstract syntax tree select expression represents a branch point for
- * selecting between multiple values.
- * In C++ the select is the definition: switch case
- */
-class RIGVMDEVELOPER_API FRigVMSelectExprAST : public FRigVMNodeExprAST
-{
-public:
-
-	// virtual destructor
-	virtual ~FRigVMSelectExprAST() {}
-
-	// disable copy constructor
-	FRigVMSelectExprAST(const FRigVMEntryExprAST&) = delete;
-
-	// overload of the type checking mechanism
-	virtual bool IsA(EType InType) const override
-	{
-		return InType == EType::Select;
-	};
-
-	virtual bool IsConstant() const override;
-
-	int32 GetConstantValueIndex() const;
-	int32 NumValues() const;
-	const FRigVMVarExprAST* GetIndexExpr() const { return ChildAt(0)->To<FRigVMVarExprAST>(); }
-	const FRigVMVarExprAST* GetValueExpr(int32 InIndex) const { return ChildAt(1 + InIndex)->To<FRigVMVarExprAST>(); }
-	const FRigVMVarExprAST* GetResultExpr() const { return ChildAt(NumChildren() - 1)->To<FRigVMVarExprAST>(); }
-
-protected:
-
-	// default constructor (protected so that only parser can access it)
-	FRigVMSelectExprAST(const FRigVMASTProxy& InProxy)
-		: FRigVMNodeExprAST(EType::Select, InProxy)
-	{}
-
-private:
-
-	friend class FRigVMParserAST;
-};
-
 /*
 * An abstract syntax tree array expression that can be used to represent
 * any array operation, such as Add, GetNum etc
diff --git a/Engine/Source/Developer/RigVMDeveloper/Public/RigVMCompiler/RigVMCompiler.h b/Engine/Source/Developer/RigVMDeveloper/Public/RigVMCompiler/RigVMCompiler.h
--- a/Engine/Source/Developer/RigVMDeveloper/Public/RigVMCompiler/RigVMCompiler.h
+++ b/Engine/Source/Developer/RigVMDeveloper/Public/RigVMCompiler/RigVMCompiler.h
@@ -215,8 +215,6 @@ private:
 	void TraverseCopy(const FRigVMCopyExprAST* InExpr, FRigVMCompilerWorkData& WorkData);
 	void TraverseCachedValue(const FRigVMCachedValueExprAST* InExpr, FRigVMCompilerWorkData& WorkData);
 	void TraverseExit(const FRigVMExitExprAST* InExpr, FRigVMCompilerWorkData& WorkData);
-	void TraverseIf(const FRigVMIfExprAST* InExpr, FRigVMCompilerWorkData& WorkData);
-	void TraverseSelect(const FRigVMSelectExprAST* InExpr, FRigVMCompilerWorkData& WorkData);
 	void TraverseArray(const FRigVMArrayExprAST* InExpr, FRigVMCompilerWorkData& WorkData);
 	void TraverseInvokeEntry(const FRigVMInvokeEntryExprAST* InExpr, FRigVMCompilerWorkData& WorkData);
 
diff --git a/Engine/Source/Developer/RigVMDeveloper/Public/RigVMModel/Nodes/RigVMIfNode.h b/Engine/Source/Developer/RigVMDeveloper/Public/RigVMModel/Nodes/RigVMIfNode.h
--- a/Engine/Source/Developer/RigVMDeveloper/Public/RigVMModel/Nodes/RigVMIfNode.h
+++ b/Engine/Source/Developer/RigVMDeveloper/Public/RigVMModel/Nodes/RigVMIfNode.h
@@ -17,8 +17,8 @@ struct FRigVMTemplate;
 /**
  * A if node is used to pick between two values
  */
-UCLASS(BlueprintType)
-class RIGVMDEVELOPER_API URigVMIfNode : public URigVMTemplateNode
+UCLASS(BlueprintType, Deprecated)
+class RIGVMDEVELOPER_API UDEPRECATED_RigVMIfNode : public URigVMTemplateNode
 {
 	GENERATED_BODY()
 
diff --git a/Engine/Source/Developer/RigVMDeveloper/Public/RigVMModel/Nodes/RigVMSelectNode.h b/Engine/Source/Developer/RigVMDeveloper/Public/RigVMModel/Nodes/RigVMSelectNode.h
--- a/Engine/Source/Developer/RigVMDeveloper/Public/RigVMModel/Nodes/RigVMSelectNode.h
+++ b/Engine/Source/Developer/RigVMDeveloper/Public/RigVMModel/Nodes/RigVMSelectNode.h
@@ -18,8 +18,8 @@ struct FRigVMTemplate;
 /**
  * A select node is used to select between multiple values
  */
-UCLASS(BlueprintType)
-class RIGVMDEVELOPER_API URigVMSelectNode : public URigVMTemplateNode
+UCLASS(BlueprintType, Deprecated)
+class RIGVMDEVELOPER_API UDEPRECATED_RigVMSelectNode : public URigVMTemplateNode
 {
 	GENERATED_BODY()
 
diff --git a/Engine/Source/Developer/RigVMDeveloper/Public/RigVMModel/RigVMController.h b/Engine/Source/Developer/RigVMDeveloper/Public/RigVMModel/RigVMController.h
--- a/Engine/Source/Developer/RigVMDeveloper/Public/RigVMModel/RigVMController.h
+++ b/Engine/Source/Developer/RigVMDeveloper/Public/RigVMModel/RigVMController.h
@@ -338,23 +338,22 @@ public:
 	// Adds a branch node to the graph.
 	// Branch nodes can be used to split the execution of into multiple branches,
 	// allowing to drive behavior by logic.
-	UE_DEPRECATED(5.2, "Please use AddUnitNode with FRigVMFunction_ControlFlowBranch::StaticStruct()")
-	UFUNCTION(BlueprintCallable, Category = RigVMController, meta=(DeprecatedFunction))
+	UFUNCTION(BlueprintCallable, Category = RigVMController)
 	URigVMNode* AddBranchNode(const FVector2D& InPosition = FVector2D::ZeroVector, const FString& InNodeName = TEXT(""), bool bSetupUndoRedo = true, bool bPrintPythonCommand = false);
 
 	// Adds an if node to the graph.
 	// If nodes can be used to pick between two values based on a condition.
 	UFUNCTION(BlueprintCallable, Category = RigVMController)
-	URigVMIfNode* AddIfNode(const FString& InCPPType, const FName& InCPPTypeObjectPath, const FVector2D& InPosition = FVector2D::ZeroVector, const FString& InNodeName = TEXT(""), bool bSetupUndoRedo = true, bool bPrintPythonCommand = false);
+	URigVMNode* AddIfNode(const FString& InCPPType, const FName& InCPPTypeObjectPath, const FVector2D& InPosition = FVector2D::ZeroVector, const FString& InNodeName = TEXT(""), bool bSetupUndoRedo = true, bool bPrintPythonCommand = false);
 	UFUNCTION(BlueprintCallable, Category = RigVMController)
-	URigVMIfNode* AddIfNodeFromStruct(UScriptStruct* InScriptStruct, const FVector2D& InPosition = FVector2D::ZeroVector, const FString& InNodeName = TEXT(""), bool bSetupUndoRedo = true);
+	URigVMNode* AddIfNodeFromStruct(UScriptStruct* InScriptStruct, const FVector2D& InPosition = FVector2D::ZeroVector, const FString& InNodeName = TEXT(""), bool bSetupUndoRedo = true);
 
 	// Adds a select node to the graph.
 	// Select nodes can be used to pick between multiple values based on an index.
 	UFUNCTION(BlueprintCallable, Category = RigVMController)
-	URigVMSelectNode* AddSelectNode(const FString& InCPPType, const FName& InCPPTypeObjectPath, const FVector2D& InPosition = FVector2D::ZeroVector, const FString& InNodeName = TEXT(""), bool bSetupUndoRedo = true, bool bPrintPythonCommand = false);
+	URigVMNode* AddSelectNode(const FString& InCPPType, const FName& InCPPTypeObjectPath, const FVector2D& InPosition = FVector2D::ZeroVector, const FString& InNodeName = TEXT(""), bool bSetupUndoRedo = true, bool bPrintPythonCommand = false);
 	UFUNCTION(BlueprintCallable, Category = RigVMController)
-	URigVMSelectNode* AddSelectNodeFromStruct(UScriptStruct* InScriptStruct, const FVector2D& InPosition = FVector2D::ZeroVector, const FString& InNodeName = TEXT(""), bool bSetupUndoRedo = true);
+	URigVMNode* AddSelectNodeFromStruct(UScriptStruct* InScriptStruct, const FVector2D& InPosition = FVector2D::ZeroVector, const FString& InNodeName = TEXT(""), bool bSetupUndoRedo = true);
 
 	// Adds a template node to the graph.
 	UFUNCTION(BlueprintCallable, Category = RigVMController)
@@ -1168,6 +1167,7 @@ protected:
 	// backwards compatibility code
 	void PatchDispatchNodesOnLoad();
 	void PatchBranchNodesOnLoad();
+	void PatchIfSelectNodesOnLoad();
 	
 private: 
 	UPROPERTY(transient)
diff --git a/Engine/Source/Developer/RigVMDeveloper/Public/RigVMModel/RigVMControllerActions.h b/Engine/Source/Developer/RigVMDeveloper/Public/RigVMModel/RigVMControllerActions.h
--- a/Engine/Source/Developer/RigVMDeveloper/Public/RigVMModel/RigVMControllerActions.h
+++ b/Engine/Source/Developer/RigVMDeveloper/Public/RigVMModel/RigVMControllerActions.h
@@ -523,66 +523,6 @@ public:
 	FString NodePath;
 };
 
-/**
- * An action adding an if node to the graph.
- */
-USTRUCT()
-struct FRigVMAddIfNodeAction : public FRigVMBaseAction
-{
-	GENERATED_BODY()
-
-public:
-
-	FRigVMAddIfNodeAction();
-	FRigVMAddIfNodeAction(URigVMIfNode* InNode);
-	virtual ~FRigVMAddIfNodeAction() {};
-	virtual UScriptStruct* GetScriptStruct() const override { return FRigVMAddIfNodeAction::StaticStruct(); }
-	virtual bool Undo(URigVMController* InController) override;
-	virtual bool Redo(URigVMController* InController) override;
-
-	UPROPERTY()
-	FString CPPType;
-
-	UPROPERTY()
-	FName CPPTypeObjectPath;
-
-	UPROPERTY()
-	FVector2D Position;
-
-	UPROPERTY()
-	FString NodePath;
-};
-
-/**
- * An action adding a select node to the graph.
- */
-USTRUCT()
-struct FRigVMAddSelectNodeAction : public FRigVMBaseAction
-{
-	GENERATED_BODY()
-
-public:
-
-	FRigVMAddSelectNodeAction();
-	FRigVMAddSelectNodeAction(URigVMSelectNode* InNode);
-	virtual ~FRigVMAddSelectNodeAction() {};
-	virtual UScriptStruct* GetScriptStruct() const override { return FRigVMAddSelectNodeAction::StaticStruct(); }
-	virtual bool Undo(URigVMController* InController) override;
-	virtual bool Redo(URigVMController* InController) override;
-
-	UPROPERTY()
-	FString CPPType;
-
-	UPROPERTY()
-	FName CPPTypeObjectPath;
-
-	UPROPERTY()
-	FVector2D Position;
-
-	UPROPERTY()
-	FString NodePath;
-};
-
 /**
  * An action adding an enum node to the graph.
  */
diff --git a/Engine/Source/Runtime/RigVM/Private/RigVMCore/RigVMDispatchFactory.cpp b/Engine/Source/Runtime/RigVM/Private/RigVMCore/RigVMDispatchFactory.cpp
--- a/Engine/Source/Runtime/RigVM/Private/RigVMCore/RigVMDispatchFactory.cpp
+++ b/Engine/Source/Runtime/RigVM/Private/RigVMCore/RigVMDispatchFactory.cpp
@@ -160,10 +160,16 @@ const TArray<FName>& FRigVMDispatchFactory::GetControlFlowBlocks_Impl(const FRig
 TArray<FRigVMExecuteArgument> FRigVMDispatchFactory::GetExecuteArguments(const FRigVMDispatchContext& InContext) const
 {
 	TArray<FRigVMExecuteArgument> Arguments = GetExecuteArguments_Impl(InContext);
-	const TRigVMTypeIndex ExecuteTypeIndex = FRigVMRegistry::Get().GetTypeIndex<FRigVMExecuteContext>();
 	for(FRigVMExecuteArgument& Argument : Arguments)
 	{
-		Argument.TypeIndex = ExecuteTypeIndex;
+		if(Argument.TypeIndex != INDEX_NONE && FRigVMRegistry::Get().IsArrayType(Argument.TypeIndex))
+		{
+			Argument.TypeIndex = RigVMTypeUtils::TypeIndex::ExecuteArray;
+		}
+		else
+		{
+			Argument.TypeIndex = RigVMTypeUtils::TypeIndex::Execute;
+		}
 	}
 	return Arguments;
 }
diff --git a/Engine/Source/Runtime/RigVM/Private/RigVMCore/RigVMRegistry.cpp b/Engine/Source/Runtime/RigVM/Private/RigVMCore/RigVMRegistry.cpp
--- a/Engine/Source/Runtime/RigVM/Private/RigVMCore/RigVMRegistry.cpp
+++ b/Engine/Source/Runtime/RigVM/Private/RigVMCore/RigVMRegistry.cpp
@@ -133,6 +133,7 @@ void FRigVMRegistry::InitializeIfNeeded()
 	ArgumentsPerCategory.Add(FRigVMTemplateArgument::ETypeCategory_ArrayArrayObjectValue, TArray<TPair<int32,int32>>()).Reserve(64);
 
 	RigVMTypeUtils::TypeIndex::Execute = FindOrAddType(FRigVMTemplateArgumentType(FRigVMExecuteContext::StaticStruct()));
+	RigVMTypeUtils::TypeIndex::ExecuteArray = FindOrAddType(FRigVMTemplateArgumentType(FRigVMExecuteContext::StaticStruct()).ConvertToArray());
 	RigVMTypeUtils::TypeIndex::Bool = FindOrAddType(FRigVMTemplateArgumentType(RigVMTypeUtils::BoolTypeName, nullptr));
 	RigVMTypeUtils::TypeIndex::Float = FindOrAddType(FRigVMTemplateArgumentType(RigVMTypeUtils::FloatTypeName, nullptr));
 	RigVMTypeUtils::TypeIndex::Double = FindOrAddType(FRigVMTemplateArgumentType(RigVMTypeUtils::DoubleTypeName, nullptr));
diff --git a/Engine/Source/Runtime/RigVM/Private/RigVMFunctions/RigVMDispatch_Core.cpp b/Engine/Source/Runtime/RigVM/Private/RigVMFunctions/RigVMDispatch_Core.cpp
--- a/Engine/Source/Runtime/RigVM/Private/RigVMFunctions/RigVMDispatch_Core.cpp
+++ b/Engine/Source/Runtime/RigVM/Private/RigVMFunctions/RigVMDispatch_Core.cpp
@@ -6,6 +6,8 @@
 
 #include UE_INLINE_GENERATED_CPP_BY_NAME(RigVMDispatch_Core)
 
+const FString FRigVMDispatch_CoreBase::TrueString = TEXT("True");
+
 TArray<FRigVMTemplateArgument> FRigVMDispatch_CoreEquals::GetArguments() const
 {
 	const TArray<FRigVMTemplateArgument::ETypeCategory> ValueCategories = {
diff --git a/Engine/Source/Runtime/RigVM/Private/RigVMFunctions/RigVMDispatch_If.cpp b/Engine/Source/Runtime/RigVM/Private/RigVMFunctions/RigVMDispatch_If.cpp
new file mode 100644
--- /dev/null
+++ b/Engine/Source/Runtime/RigVM/Private/RigVMFunctions/RigVMDispatch_If.cpp
@@ -0,0 +1,84 @@
+// Copyright Epic Games, Inc. All Rights Reserved.
+
+#include "RigVMFunctions/RigVMDispatch_If.h"
+#include "RigVMCore/RigVMRegistry.h"
+#include "RigVMCore/RigVMStruct.h"
+
+#include UE_INLINE_GENERATED_CPP_BY_NAME(RigVMDispatch_If)
+
+const FName FRigVMDispatch_If::ConditionName = TEXT("Condition");
+const FName FRigVMDispatch_If::TrueName = TEXT("True");
+const FName FRigVMDispatch_If::FalseName = TEXT("False");
+const FName FRigVMDispatch_If::ResultName = TEXT("Result");
+
+TArray<FRigVMTemplateArgument> FRigVMDispatch_If::GetArguments() const
+{
+	static const TArray<FRigVMTemplateArgument::ETypeCategory> ValueCategories = {
+		FRigVMTemplateArgument::ETypeCategory_SingleAnyValue,
+		FRigVMTemplateArgument::ETypeCategory_ArrayAnyValue
+	};
+	return {
+		FRigVMTemplateArgument(ConditionName, ERigVMPinDirection::Input, RigVMTypeUtils::TypeIndex::Bool),
+		FRigVMTemplateArgument(TrueName, ERigVMPinDirection::Input, ValueCategories),
+		FRigVMTemplateArgument(FalseName, ERigVMPinDirection::Input, ValueCategories),
+		FRigVMTemplateArgument(ResultName, ERigVMPinDirection::Output, ValueCategories)
+	};
+}
+
+FRigVMTemplateTypeMap FRigVMDispatch_If::OnNewArgumentType(const FName& InArgumentName,
+	TRigVMTypeIndex InTypeIndex) const
+{
+	FRigVMTemplateTypeMap Types;
+	Types.Add(ConditionName, RigVMTypeUtils::TypeIndex::Bool);
+	Types.Add(TrueName, InTypeIndex);
+	Types.Add(FalseName, InTypeIndex);
+	Types.Add(ResultName, InTypeIndex);
+	return Types;
+}
+
+#if WITH_EDITOR
+
+FString FRigVMDispatch_If::GetArgumentMetaData(const FName& InArgumentName, const FName& InMetaDataKey) const
+{
+	if(InArgumentName == TrueName || InArgumentName == FalseName)
+	{
+		if(InMetaDataKey == FRigVMStruct::ComputeLazilyMetaName)
+		{
+			return TrueString;
+		}
+	}
+	return FRigVMDispatch_CoreBase::GetArgumentMetaData(InArgumentName, InMetaDataKey);
+}
+
+#endif
+
+FRigVMFunctionPtr FRigVMDispatch_If::GetDispatchFunctionImpl(const FRigVMTemplateTypeMap& InTypes) const
+{
+	return &FRigVMDispatch_If::Execute;
+}
+
+void FRigVMDispatch_If::Execute(FRigVMExtendedExecuteContext& InContext, FRigVMMemoryHandleArray Handles)
+{
+	const FProperty* Property = Handles[1].GetResolvedProperty(); 
+
+#if WITH_EDITOR
+	check(Handles[0].IsBool());
+	const FProperty* PropertyB = Handles[2].GetResolvedProperty(); 
+	const FProperty* PropertyC = Handles[3].GetResolvedProperty(); 
+	check(Property->SameType(PropertyB));
+	check(Property->SameType(PropertyC));
+#endif
+
+	const bool& Condition = *(const bool*)Handles[0].GetData();
+
+	FRigVMMemoryHandle& InputHandle = Condition ? Handles[1] : Handles[2];
+	if(InputHandle.IsLazy())
+	{
+		InputHandle.ComputeLazyValueIfNecessary();
+	}
+
+	const uint8* Input = InputHandle.GetData();
+	uint8* Result = Handles[3].GetData();
+	Property->CopyCompleteValue(Result, Input);
+}
+
diff --git a/Engine/Source/Runtime/RigVM/Private/RigVMFunctions/RigVMDispatch_Select.cpp b/Engine/Source/Runtime/RigVM/Private/RigVMFunctions/RigVMDispatch_Select.cpp
new file mode 100644
--- /dev/null
+++ b/Engine/Source/Runtime/RigVM/Private/RigVMFunctions/RigVMDispatch_Select.cpp
@@ -0,0 +1,147 @@
+// Copyright Epic Games, Inc. All Rights Reserved.
+
+#include "RigVMFunctions/RigVMDispatch_Select.h"
+#include "RigVMStringUtils.h"
+#include "RigVMCore/RigVMRegistry.h"
+#include "RigVMCore/RigVMStruct.h"
+
+#include UE_INLINE_GENERATED_CPP_BY_NAME(RigVMDispatch_Select)
+
+const FName FRigVMDispatch_SelectInt32::IndexName = TEXT("Index");
+const FName FRigVMDispatch_SelectInt32::ValuesName = TEXT("Values");
+const FName FRigVMDispatch_SelectInt32::ResultName = TEXT("Result");
+
+TArray<FRigVMTemplateArgument> FRigVMDispatch_SelectInt32::GetArguments() const
+{
+	static const TArray<FRigVMTemplateArgument::ETypeCategory> ArrayCategories = {
+		FRigVMTemplateArgument::ETypeCategory_ArrayAnyValue,
+		FRigVMTemplateArgument::ETypeCategory_ArrayArrayAnyValue
+	};
+	static const TArray<FRigVMTemplateArgument::ETypeCategory> SingleCategories = {
+		FRigVMTemplateArgument::ETypeCategory_SingleAnyValue,
+		FRigVMTemplateArgument::ETypeCategory_ArrayAnyValue
+	};
+	
+	return {
+		FRigVMTemplateArgument(IndexName, ERigVMPinDirection::Input, RigVMTypeUtils::TypeIndex::Int32),
+		FRigVMTemplateArgument(ValuesName, ERigVMPinDirection::Input, ArrayCategories),
+		FRigVMTemplateArgument(ResultName, ERigVMPinDirection::Output, SingleCategories)
+	};
+}
+
+FRigVMTemplateTypeMap FRigVMDispatch_SelectInt32::OnNewArgumentType(const FName& InArgumentName,
+	TRigVMTypeIndex InTypeIndex) const
+{
+	FRigVMTemplateTypeMap Types;
+	Types.Add(IndexName, RigVMTypeUtils::TypeIndex::Int32);
+	if(InArgumentName == ValuesName)
+	{
+		Types.Add(ValuesName, InTypeIndex);
+		Types.Add(ResultName, FRigVMRegistry::Get().GetBaseTypeFromArrayTypeIndex(InTypeIndex));
+	}
+	else
+	{
+		Types.Add(ValuesName, FRigVMRegistry::Get().GetArrayTypeFromBaseTypeIndex(InTypeIndex));
+		Types.Add(ResultName, InTypeIndex);
+	}
+	return Types;
+}
+
+#if WITH_EDITOR
+
+FString FRigVMDispatch_SelectInt32::GetArgumentMetaData(const FName& InArgumentName, const FName& InMetaDataKey) const
+{
+	if(InArgumentName == ValuesName)
+	{
+		if(InMetaDataKey == FRigVMStruct::ComputeLazilyMetaName)
+		{
+			return TrueString;
+		}
+		if(InMetaDataKey == FRigVMStruct::FixedSizeArrayMetaName)
+		{
+			return TrueString;
+		}
+	}
+	return FRigVMDispatch_CoreBase::GetArgumentMetaData(InArgumentName, InMetaDataKey);
+}
+
+FString FRigVMDispatch_SelectInt32::GetArgumentDefaultValue(const FName& InArgumentName,
+	TRigVMTypeIndex InTypeIndex) const
+{
+	if(InArgumentName == IndexName)
+	{
+		static const FString ZeroValue = TEXT("0");
+		return ZeroValue;
+	}
+	if(InArgumentName == ValuesName)
+	{
+		const TRigVMTypeIndex ElementTypeIndex = FRigVMRegistry::Get().GetBaseTypeFromArrayTypeIndex(InTypeIndex);
+		FString ElementDefaultValue = FRigVMDispatch_CoreBase::GetArgumentDefaultValue(InArgumentName, ElementTypeIndex);
+		if(ElementDefaultValue.IsEmpty())
+		{
+			static const FString EmptyBracesString = TEXT("()");
+			ElementDefaultValue = EmptyBracesString;
+		}
+
+		// default to two entries
+		return RigVMStringUtils::JoinDefaultValue({ElementDefaultValue, ElementDefaultValue});
+	}
+	return FRigVMDispatch_CoreBase::GetArgumentDefaultValue(InArgumentName, InTypeIndex);
+}
+
+#endif
+
+FName FRigVMDispatch_SelectInt32::GetArgumentNameForOperandIndex(int32 InOperandIndex, int32 InTotalOperands) const
+{
+	if(InOperandIndex > 0 && InOperandIndex < InTotalOperands - 1)
+	{
+		return FRigVMBranchInfo::GetFixedArrayLabel(ValuesName, *FString::FromInt(InOperandIndex - 1));
+	}
+	return FRigVMDispatch_CoreBase::GetArgumentNameForOperandIndex(InOperandIndex == 0 ? 0 : 2, 3);
+}
+
+FRigVMFunctionPtr FRigVMDispatch_SelectInt32::GetDispatchFunctionImpl(const FRigVMTemplateTypeMap& InTypes) const
+{
+	return &FRigVMDispatch_SelectInt32::Execute;
+}
+
+void FRigVMDispatch_SelectInt32::Execute(FRigVMExtendedExecuteContext& InContext, FRigVMMemoryHandleArray Handles)
+{
+	const FProperty* Property = Handles.Last().GetResolvedProperty();
+	const int32 NumCases = Handles.Num() - 2;
+
+	if(NumCases < 1)
+	{
+		static constexpr TCHAR InvalidIndexMessage[] = TEXT("SelectInt32: No cases to select from.");
+		InContext.GetPublicData<>().Log(EMessageSeverity::Error, InvalidIndexMessage);
+		return;
+	}
+
+#if WITH_EDITOR
+	check(Handles[0].IsInt32());
+	for(int32 Case = 0; Case < NumCases; Case++)
+	{
+		const FProperty* PropertyB = Handles[1 + Case].GetResolvedProperty(); 
+		check(Property->SameType(PropertyB));
+	}
+#endif
+
+	int32 Index = *(const int32*)Handles[0].GetData();
+	if(!FMath::IsWithin(Index, 0, NumCases))
+	{
+		static constexpr TCHAR InvalidIndexMessage[] = TEXT("SelectInt32: Invalid index %d - num cases is %d. Falling back to index 0.");
+		InContext.GetPublicData<>().Logf(EMessageSeverity::Warning, InvalidIndexMessage, Index, NumCases);
+		Index = 0;
+	}
+
+	FRigVMMemoryHandle& InputHandle = Handles[1 + Index];
+	if(InputHandle.IsLazy())
+	{
+		InputHandle.ComputeLazyValueIfNecessary();
+	}
+
+	const uint8* Input = InputHandle.GetData();
+	uint8* Result = Handles.Last().GetData();
+	Property->CopyCompleteValue(Result, Input);
+}
+
diff --git a/Engine/Source/Runtime/RigVM/Private/RigVMStringUtils.cpp b/Engine/Source/Runtime/RigVM/Private/RigVMStringUtils.cpp
new file mode 100644
--- /dev/null
+++ b/Engine/Source/Runtime/RigVM/Private/RigVMStringUtils.cpp
@@ -0,0 +1,189 @@
+﻿// Copyright Epic Games, Inc. All Rights Reserved.
+
+#include "RigVMStringUtils.h"
+
+bool RigVMStringUtils::SplitNodePathAtStart(const FString& InNodePath, FString& LeftMost, FString& Right)
+{
+	return InNodePath.Split(TEXT("|"), &LeftMost, &Right, ESearchCase::IgnoreCase, ESearchDir::FromStart);
+}
+
+bool RigVMStringUtils::SplitNodePathAtEnd(const FString& InNodePath, FString& Left, FString& RightMost)
+{
+	return InNodePath.Split(TEXT("|"), &Left, &RightMost, ESearchCase::IgnoreCase, ESearchDir::FromEnd);
+}
+
+bool RigVMStringUtils::SplitNodePath(const FString& InNodePath, TArray<FString>& Parts)
+{
+	int32 OriginalPartsCount = Parts.Num();
+	FString NodePathRemaining = InNodePath;
+	FString Left, Right;
+	Right = NodePathRemaining;
+
+	while (SplitNodePathAtStart(NodePathRemaining, Left, Right))
+	{
+		Parts.Add(Left);
+		Left.Empty();
+		NodePathRemaining = Right;
+	}
+
+	if (!Right.IsEmpty())
+	{
+		Parts.Add(Right);
+	}
+
+	return Parts.Num() > OriginalPartsCount;
+}
+
+FString RigVMStringUtils::JoinNodePath(const FString& Left, const FString& Right)
+{
+	ensure(!Left.IsEmpty() && !Right.IsEmpty());
+	return Left + TEXT("|") + Right;
+}
+
+FString RigVMStringUtils::JoinNodePath(const TArray<FString>& InParts)
+{
+	if (InParts.Num() == 0)
+	{
+		return FString();
+	}
+
+	FString Result = InParts[0];
+	for (int32 PartIndex = 1; PartIndex < InParts.Num(); PartIndex++)
+	{
+		Result += TEXT("|") + InParts[PartIndex];
+	}
+
+	return Result;
+}
+
+bool RigVMStringUtils::SplitPinPathAtStart(const FString& InPinPath, FString& LeftMost, FString& Right)
+{
+	return InPinPath.Split(TEXT("."), &LeftMost, &Right, ESearchCase::IgnoreCase, ESearchDir::FromStart);
+}
+
+bool RigVMStringUtils::SplitPinPathAtEnd(const FString& InPinPath, FString& Left, FString& RightMost)
+{
+	return InPinPath.Split(TEXT("."), &Left, &RightMost, ESearchCase::IgnoreCase, ESearchDir::FromEnd);
+}
+
+bool RigVMStringUtils::SplitPinPath(const FString& InPinPath, TArray<FString>& Parts)
+{
+	int32 OriginalPartsCount = Parts.Num();
+	FString PinPathRemaining = InPinPath;
+	FString Left, Right;
+	while(SplitPinPathAtStart(PinPathRemaining, Left, Right))
+	{
+		Parts.Add(Left);
+		Left.Empty();
+		PinPathRemaining = Right;
+	}
+
+	if (!Right.IsEmpty())
+	{
+		Parts.Add(Right);
+	}
+
+	return Parts.Num() > OriginalPartsCount;
+}
+
+FString RigVMStringUtils::JoinPinPath(const FString& Left, const FString& Right)
+{
+	ensure(!Left.IsEmpty() && !Right.IsEmpty());
+	return Left + TEXT(".") + Right;
+}
+
+FString RigVMStringUtils::JoinPinPath(const TArray<FString>& InParts)
+{
+	if (InParts.Num() == 0)
+	{
+		return FString();
+	}
+
+	FString Result = InParts[0];
+	for (int32 PartIndex = 1; PartIndex < InParts.Num(); PartIndex++)
+	{
+		Result += TEXT(".") + InParts[PartIndex];
+	}
+
+	return Result;
+}
+
+FString RigVMStringUtils::JoinDefaultValue(const TArray<FString>& InParts)
+{
+	static constexpr TCHAR EmptyBraces[] = TEXT("()");
+	if(InParts.IsEmpty())
+	{
+		return EmptyBraces;
+	}
+
+	static constexpr TCHAR Format[] = TEXT("(%s)");
+	return FString::Printf(Format, *FString::Join(InParts, TEXT(",")));
+}
+
+TArray<FString> RigVMStringUtils::SplitDefaultValue(const FString& InDefaultValue)
+{
+	TArray<FString> Parts;
+	if (InDefaultValue.IsEmpty())
+	{
+		return Parts;
+	}
+
+	if(InDefaultValue[0] != TCHAR('('))
+	{
+		return Parts;
+	}
+	if(InDefaultValue[InDefaultValue.Len() - 1] != TCHAR(')'))
+	{
+		return Parts;
+	}
+
+	FString Content = InDefaultValue.Mid(1, InDefaultValue.Len() - 2);
+	int32 BraceCount = 0;
+	int32 QuoteCount = 0;
+
+	int32 LastPartStartIndex = 0;
+	for (int32 CharIndex = 0; CharIndex < Content.Len(); CharIndex++)
+	{
+		TCHAR Char = Content[CharIndex];
+		if (QuoteCount > 0)
+		{
+			if (Char == TCHAR('"'))
+			{
+				QuoteCount = 0;
+			}
+		}
+		else if (Char == TCHAR('"'))
+		{
+			QuoteCount = 1;
+		}
+
+		if (Char == TCHAR('('))
+		{
+			if (QuoteCount == 0)
+			{
+				BraceCount++;
+			}
+		}
+		else if (Char == TCHAR(')'))
+		{
+			if (QuoteCount == 0)
+			{
+				BraceCount--;
+				BraceCount = FMath::Max<int32>(BraceCount, 0);
+			}
+		}
+		else if (Char == TCHAR(',') && BraceCount == 0 && QuoteCount == 0)
+		{
+			// ignore whitespaces
+			Parts.Add(Content.Mid(LastPartStartIndex, CharIndex - LastPartStartIndex).Replace(TEXT(" "), TEXT("")));
+			LastPartStartIndex = CharIndex + 1;
+		}
+	}
+
+	if (!Content.IsEmpty())
+	{
+		// ignore whitespaces from the start and end of the string
+		Parts.Add(Content.Mid(LastPartStartIndex).TrimStartAndEnd());
+	}
+	return Parts;
+}
diff --git a/Engine/Source/Runtime/RigVM/Private/RigVMTypeUtils.cpp b/Engine/Source/Runtime/RigVM/Private/RigVMTypeUtils.cpp
--- a/Engine/Source/Runtime/RigVM/Private/RigVMTypeUtils.cpp
+++ b/Engine/Source/Runtime/RigVM/Private/RigVMTypeUtils.cpp
@@ -9,6 +9,7 @@
 namespace RigVMTypeUtils
 {
 	TRigVMTypeIndex TypeIndex::Execute = INDEX_NONE;	
+	TRigVMTypeIndex TypeIndex::ExecuteArray = INDEX_NONE;	
 	TRigVMTypeIndex TypeIndex::Bool = INDEX_NONE;	
 	TRigVMTypeIndex TypeIndex::Float = INDEX_NONE;	
 	TRigVMTypeIndex TypeIndex::Double = INDEX_NONE;	
diff --git a/Engine/Source/Runtime/RigVM/Public/RigVMCore/RigVMMemoryStorage.h b/Engine/Source/Runtime/RigVM/Public/RigVMCore/RigVMMemoryStorage.h
--- a/Engine/Source/Runtime/RigVM/Public/RigVMCore/RigVMMemoryStorage.h
+++ b/Engine/Source/Runtime/RigVM/Public/RigVMCore/RigVMMemoryStorage.h
@@ -95,6 +95,17 @@ struct RIGVM_API FRigVMBranchInfo
 		return Hash;
 	}
 
+	static FString GetFixedArrayLabel(const FString& InArrayName, const FString& InElementName)
+	{
+		static constexpr TCHAR Format[] = TEXT("%s_%s");
+		return FString::Printf(Format, *InArrayName, *InElementName);
+	}
+
+	static FName GetFixedArrayLabel(const FName& InArrayName, const FName& InElementName)
+	{
+		return *GetFixedArrayLabel(InArrayName.ToString(), InElementName.ToString());
+	}
+
 	UPROPERTY()
 	int32 Index;
 
diff --git a/Engine/Source/Runtime/RigVM/Public/RigVMFunctions/RigVMDispatch_Core.h b/Engine/Source/Runtime/RigVM/Public/RigVMFunctions/RigVMDispatch_Core.h
--- a/Engine/Source/Runtime/RigVM/Public/RigVMFunctions/RigVMDispatch_Core.h
+++ b/Engine/Source/Runtime/RigVM/Public/RigVMFunctions/RigVMDispatch_Core.h
@@ -9,6 +9,10 @@ USTRUCT(meta=(Abstract, Category = "Core", NodeColor = "0.762745, 1,0, 0.329412"
 struct FRigVMDispatch_CoreBase : public FRigVMDispatchFactory
 {
 	GENERATED_BODY()
+
+protected:
+
+	static const FString TrueString;
 };
 
 /*
diff --git a/Engine/Source/Runtime/RigVM/Public/RigVMFunctions/RigVMDispatch_If.h b/Engine/Source/Runtime/RigVM/Public/RigVMFunctions/RigVMDispatch_If.h
new file mode 100644
--- /dev/null
+++ b/Engine/Source/Runtime/RigVM/Public/RigVMFunctions/RigVMDispatch_If.h
@@ -0,0 +1,35 @@
+// Copyright Epic Games, Inc. All Rights Reserved.
+
+#pragma once
+
+#include "RigVMCore/RigVMDispatchFactory.h"
+#include "RigVMDispatch_Core.h"
+#include "RigVMDispatch_If.generated.h"
+
+/*
+ * Chooses between two values based on a condition
+ */
+USTRUCT(meta=(DisplayName = "If", Category = "Execution", Keywords = "Branch,Condition", NodeColor = "0,1,0,1"))
+struct FRigVMDispatch_If : public FRigVMDispatch_CoreBase
+{
+	GENERATED_BODY()
+
+public:
+
+	virtual TArray<FRigVMTemplateArgument> GetArguments() const override;
+	virtual FRigVMTemplateTypeMap OnNewArgumentType(const FName& InArgumentName, TRigVMTypeIndex InTypeIndex) const override;
+
+#if WITH_EDITOR
+	virtual FString GetArgumentMetaData(const FName& InArgumentName, const FName& InMetaDataKey) const override;
+#endif
+
+protected:
+
+	virtual FRigVMFunctionPtr GetDispatchFunctionImpl(const FRigVMTemplateTypeMap& InTypes) const override;
+	static void Execute(FRigVMExtendedExecuteContext& InContext, FRigVMMemoryHandleArray Handles);
+
+	static const FName ConditionName;
+	static const FName TrueName;
+	static const FName FalseName;
+	static const FName ResultName;
+};
diff --git a/Engine/Source/Runtime/RigVM/Public/RigVMFunctions/RigVMDispatch_Select.h b/Engine/Source/Runtime/RigVM/Public/RigVMFunctions/RigVMDispatch_Select.h
new file mode 100644
--- /dev/null
+++ b/Engine/Source/Runtime/RigVM/Public/RigVMFunctions/RigVMDispatch_Select.h
@@ -0,0 +1,35 @@
+// Copyright Epic Games, Inc. All Rights Reserved.
+
+#pragma once
+
+#include "RigVMCore/RigVMDispatchFactory.h"
+#include "RigVMDispatch_Core.h"
+#include "RigVMDispatch_Select.generated.h"
+
+/*
+ * Pick from a list of values based on an integer index
+ */
+USTRUCT(meta=(DisplayName = "Select", Category = "Execution", Keywords = "Switch,Case", NodeColor = "0,1,0,1"))
+struct FRigVMDispatch_SelectInt32 : public FRigVMDispatch_CoreBase
+{
+	GENERATED_BODY()
+
+public:
+
+	virtual TArray<FRigVMTemplateArgument> GetArguments() const override;
+	virtual FRigVMTemplateTypeMap OnNewArgumentType(const FName& InArgumentName, TRigVMTypeIndex InTypeIndex) const override;
+#if WITH_EDITOR
+	virtual FString GetArgumentMetaData(const FName& InArgumentName, const FName& InMetaDataKey) const override;
+	virtual FString GetArgumentDefaultValue(const FName& InArgumentName, TRigVMTypeIndex InTypeIndex) const override;
+#endif
+	virtual FName GetArgumentNameForOperandIndex(int32 InOperandIndex, int32 InTotalOperands) const override;
+
+protected:
+
+	virtual FRigVMFunctionPtr GetDispatchFunctionImpl(const FRigVMTemplateTypeMap& InTypes) const override;
+	static void Execute(FRigVMExtendedExecuteContext& InContext, FRigVMMemoryHandleArray Handles);
+
+	static const FName IndexName;
+	static const FName ValuesName;
+	static const FName ResultName;
+};
diff --git a/Engine/Source/Runtime/RigVM/Public/RigVMStringUtils.h b/Engine/Source/Runtime/RigVM/Public/RigVMStringUtils.h
new file mode 100644
--- /dev/null
+++ b/Engine/Source/Runtime/RigVM/Public/RigVMStringUtils.h
@@ -0,0 +1,44 @@
+// Copyright Epic Games, Inc. All Rights Reserved.
+
+#pragma once
+
+#include "CoreMinimal.h"
+
+namespace RigVMStringUtils
+{
+	// Splits a NodePath at the start, so for example "CollapseNodeA|CollapseNodeB|CollapseNodeC" becomes "CollapseNodeA" and "CollapseNodeB|CollapseNodeC"
+	RIGVM_API bool SplitNodePathAtStart(const FString& InNodePath, FString& LeftMost, FString& Right);
+
+	// Splits a NodePath at the end, so for example "CollapseNodeA|CollapseNodeB|CollapseNodeC" becomes "CollapseNodeA|CollapseNodeB" and "CollapseNodeC"
+	RIGVM_API bool SplitNodePathAtEnd(const FString& InNodePath, FString& Left, FString& RightMost);
+
+	// Splits a NodePath into all segments, so for example "Node.Color.R" becomes ["Node", "Color", "R"]
+	RIGVM_API bool SplitNodePath(const FString& InNodePath, TArray<FString>& Parts);
+
+	// Joins a NodePath from to segments, so for example "CollapseNodeA" and "CollapseNodeB|CollapseNodeC" becomes "CollapseNodeA|CollapseNodeB|CollapseNodeC"
+	RIGVM_API FString JoinNodePath(const FString& Left, const FString& Right);
+
+	// Joins a NodePath from to segments, so for example ["CollapseNodeA", "CollapseNodeB", "CollapseNodeC"] becomes "CollapseNodeA|CollapseNodeB|CollapseNodeC"
+	RIGVM_API FString JoinNodePath(const TArray<FString>& InParts);
+
+	// Splits a PinPath at the start, so for example "Node.Color.R" becomes "Node" and "Color.R"
+	RIGVM_API bool SplitPinPathAtStart(const FString& InPinPath, FString& LeftMost, FString& Right);
+
+	// Splits a PinPath at the start, so for example "Node.Color.R" becomes "Node.Color" and "R"
+	RIGVM_API bool SplitPinPathAtEnd(const FString& InPinPath, FString& Left, FString& RightMost);
+
+	// Splits a PinPath into all segments, so for example "Node.Color.R" becomes ["Node", "Color", "R"]
+	RIGVM_API bool SplitPinPath(const FString& InPinPath, TArray<FString>& Parts);
+
+	// Joins a PinPath from to segments, so for example "Node.Color" and "R" becomes "Node.Color.R"
+	RIGVM_API FString JoinPinPath(const FString& Left, const FString& Right);
+
+	// Joins a PinPath from to segments, so for example ["Node", "Color", "R"] becomes "Node.Color.R"
+	RIGVM_API FString JoinPinPath(const TArray<FString>& InParts);
+
+	// Joins the name value pairs into a single default value string
+	RIGVM_API FString JoinDefaultValue(const TArray<FString>& InParts);
+
+	// Splits the default value into name-value pairs
+	RIGVM_API TArray<FString> SplitDefaultValue(const FString& InDefaultValue);
+}
\ No newline at end of file
diff --git a/Engine/Source/Runtime/RigVM/Public/RigVMTypeUtils.h b/Engine/Source/Runtime/RigVM/Public/RigVMTypeUtils.h
--- a/Engine/Source/Runtime/RigVM/Public/RigVMTypeUtils.h
+++ b/Engine/Source/Runtime/RigVM/Public/RigVMTypeUtils.h
@@ -51,6 +51,7 @@ namespace RigVMTypeUtils
 	{
 	public:
 		static TRigVMTypeIndex Execute;	
+		static TRigVMTypeIndex ExecuteArray;	
 		static TRigVMTypeIndex Bool;	
 		static TRigVMTypeIndex Float;	
 		static TRigVMTypeIndex Double;	
