Control Rig: Fix hidden properties initialization in nativized rigs

diff --git a/Engine/Source/Developer/RigVMDeveloper/Private/RigVMCompiler/RigVMCodeGenerator.cpp b/Engine/Source/Developer/RigVMDeveloper/Private/RigVMCompiler/RigVMCodeGenerator.cpp
--- a/Engine/Source/Developer/RigVMDeveloper/Private/RigVMCompiler/RigVMCodeGenerator.cpp
+++ b/Engine/Source/Developer/RigVMDeveloper/Private/RigVMCompiler/RigVMCodeGenerator.cpp
@@ -24,7 +24,8 @@ static constexpr TCHAR RigVM_WrappedArrayTypeFormat[] = TEXT("struct {0}_API {1}
 static constexpr TCHAR RigVM_WrappedTypeNameFormat[] = TEXT("{0}Array_{1}");
 static constexpr TCHAR RigVM_DeclareExternalVariableFormat[] = TEXT("\t{0}* {1} = nullptr;");
 static constexpr TCHAR RigVM_UpdateExternalVariableFormat[] = TEXT("\t{0} = &GetExternalVariableRef<{1}>(TEXT(\"{2}\"), TEXT(\"{1}\"));");
-static constexpr TCHAR RigVM_MemberPropertyFormat[] = TEXT("\t{0} {1};");
+static constexpr TCHAR RigVM_MemberPropertyFormat[] = TEXT("\t{0} {1} = {2};");
+static constexpr TCHAR RigVM_MemberPropertyFormatNoDefault[] = TEXT("\t{0} {1};");
 static constexpr TCHAR RigVM_DeclareEntryNameFormat[] = TEXT("\tstatic const FName EntryName_{0};");
 static constexpr TCHAR RigVM_DefineEntryNameFormat[] = TEXT("const FName U{0}::EntryName_{1} = TEXT(\"{2}\");");
 static constexpr TCHAR RigVM_DeclareBlockNameFormat[] = TEXT("\tstatic const FName BlockName_{0};");
@@ -52,8 +53,8 @@ static constexpr TCHAR RigVM_CopyOpAssignFormat[] = TEXT("\t{0} = {1}{2};");
 static constexpr TCHAR RigVM_IncrementOpFormat[] = TEXT("\t{0}++;");
 static constexpr TCHAR RigVM_DecrementOpFormat[] = TEXT("\t{0}--;");
 static constexpr TCHAR RigVM_EqualsOpFormat[] = TEXT("\t{0} = {1} == {2};");
-static constexpr TCHAR RigVM_InvokeEntryFormat[] = TEXT("\tif (InEntryName == EntryName_{0}) return ExecuteEntry_{0}();");
-static constexpr TCHAR RigVM_InvokeEntryByNameFormat[] = TEXT("\tERigVMExecuteResult EntryResult = InvokeEntryByName(InEntryName{0});\r\n\tSetInstructionIndex(INDEX_NONE);\r\n\tStopProfiling();\r\n\treturn ExecuteResult;");
+static constexpr TCHAR RigVM_InvokeEntryFormat[] = TEXT("\tif (InEntryName == EntryName_{0}) return ExecuteEntry_{0}(PublicContext);");
+static constexpr TCHAR RigVM_InvokeEntryByNameFormat[] = TEXT("\tERigVMExecuteResult EntryResult = InvokeEntryByName(InEntryName{0});\r\n\tSetInstructionIndex(0);\r\n\tStopProfiling();\r\n\treturn EntryResult;");
 static constexpr TCHAR RigVM_InvokeEntryByNameFormat2[] = TEXT("\tif(!InvokeEntryByName({0}{1})) return ERigVMExecuteResult::Failed;");
 static constexpr TCHAR RigVM_CanExecuteEntryFormat[] = TEXT("\tif(!CanExecuteEntry(InEntryName, false)) { return ERigVMExecuteResult::Failed; }");
 static constexpr TCHAR RigVM_EntryExecuteGuardFormat[] = TEXT("\tFEntryExecuteGuard EntryExecuteGuard(EntriesBeingExecuted, FindEntry(InEntryName));");
@@ -70,7 +71,7 @@ static constexpr TCHAR RigVM_EnumTypeSuffixFormat[] = TEXT("::Type");
 static constexpr TCHAR RigVM_IsValidArraySizeFormat[] = TEXT("IsValidArraySize({0})");
 static constexpr TCHAR RigVM_IsValidArrayIndexFormat[] = TEXT("IsValidArrayIndex<{0}>(TemporaryArrayIndex, {1})");
 static constexpr TCHAR RigVM_TemporaryArrayIndexFormat[] = TEXT("\tTemporaryArrayIndex = {0};");
-static constexpr TCHAR RigVM_StartProfilingFormat[] = TEXT("\tStartProfiling");
+static constexpr TCHAR RigVM_StartProfilingFormat[] = TEXT("\tStartProfiling();");
 static constexpr TCHAR RigVM_ExecuteReachedExitFormat[] = TEXT("\tBroadcastExecutionReachedExit();");
 static constexpr TCHAR RigVM_InstructionLabelFormat[] = TEXT("\tInstruction{0}Label:");
 static constexpr TCHAR RigVM_SetInstructionIndexFormat[] = TEXT("\tSetInstructionIndex({0});");
@@ -112,12 +113,12 @@ static constexpr TCHAR RigVM_GetVMHashFormat[] = TEXT("\tvirtual uint32 GetVMHas
 static constexpr TCHAR RigVM_GetEntryNamesFormat[] = TEXT("\tvirtual const TArray<FName>& GetEntryNames() const override\r\n\t{\r\n\t\tstatic const TArray<FName> StaticEntryNames = { {0} };\r\n\t\treturn StaticEntryNames;\r\n\t}");
 static constexpr TCHAR RigVM_DeclareUpdateExternalVariablesFormat[] = TEXT("\tvirtual void UpdateExternalVariables() override;");
 static constexpr TCHAR RigVM_DeclareInvokeEntryByNameFormat[] = TEXT("\tERigVMExecuteResult InvokeEntryByName(const FName& InEntryName{0});");
-static constexpr TCHAR RigVM_DeclareInitializeFormat[] = TEXT("\tvirtual bool Initialize(TArrayView<URigVMMemoryStorage*> Memory, TArrayView<void*> bool bInitializeMemory = true) override;");
-static constexpr TCHAR RigVM_DefineInitializeFormat[] = TEXT("bool U{0}::Initialize(TArrayView<URigVMMemoryStorage*> Memory, TArrayView<void*> bool bInitializeMemory)\r\n{");
-static constexpr TCHAR RigVM_DeclareExecuteFormat[] = TEXT("\tvirtual ERigVMExecuteResult Execute(TArrayView<URigVMMemoryStorage*> Memory, TArrayView<void*> const FName& InEntryName) override;");
+static constexpr TCHAR RigVM_DeclareInitializeFormat[] = TEXT("\tvirtual bool Initialize(TArrayView<URigVMMemoryStorage*> Memory, bool bInitializeMemory = true) override;");
+static constexpr TCHAR RigVM_DefineInitializeFormat[] = TEXT("bool U{0}::Initialize(TArrayView<URigVMMemoryStorage*> Memory, bool bInitializeMemory)\r\n{");
+static constexpr TCHAR RigVM_DeclareExecuteFormat[] = TEXT("\tvirtual ERigVMExecuteResult Execute(TArrayView<URigVMMemoryStorage*> Memory, const FName& InEntryName) override;");
 static constexpr TCHAR RigVM_DefineUpdateExternalVariablesFormat[] = TEXT("void U{0}::UpdateExternalVariables()\r\n{");
 static constexpr TCHAR RigVM_DefineInvokeEntryByNameFormat[] = TEXT("ERigVMExecuteResult U{0}::InvokeEntryByName(const FName& InEntryName{1})\r\n{");
-static constexpr TCHAR RigVM_DefineExecuteFormat[] = TEXT("ERigVMExecuteResult U{0}::Execute(TArrayView<URigVMMemoryStorage*> Memory, TArrayView<void*> const FName& InEntryName)\r\n{");
+static constexpr TCHAR RigVM_DefineExecuteFormat[] = TEXT("ERigVMExecuteResult U{0}::Execute(TArrayView<URigVMMemoryStorage*> Memory, const FName& InEntryName)\r\n{");
 static constexpr TCHAR RigVM_DeclareExecuteEntryFormat[] = TEXT("\tERigVMExecuteResult ExecuteEntry_{0}({1});");
 static constexpr TCHAR RigVM_DefineExecuteEntryFormat[] = TEXT("ERigVMExecuteResult U{0}::ExecuteEntry_{1}({2})\r\n{");
 static constexpr TCHAR RigVM_DeclareExecuteGroupFormat[] = TEXT("\tERigVMExecuteResult ExecuteGroup_{0}_{1}({2});");
@@ -127,7 +128,7 @@ static constexpr TCHAR RigVM_RigVMCoreIncludeFormat[] = TEXT("RigVMCore/RigVMCor
 static constexpr TCHAR RigVM_RigVMModuleIncludeFormat[] = TEXT("RigVMModule.h");
 static constexpr TCHAR RigVM_RigVMCoreLibraryFormat[] = TEXT("RigVM");
 static constexpr TCHAR RigVM_JoinFilePathFormat[] = TEXT("{0}/{1}");
-static constexpr TCHAR RigVM_GetOperandSliceFormat[] = TEXT("GetOperandSlice<{0}>({1}){2}");
+static constexpr TCHAR RigVM_GetOperandSliceFormat[] = TEXT("GetOperandSlice<{0}>({1},&{1}_Const){2}");
 static constexpr TCHAR RigVM_ExternalVariableFormat[] = TEXT("(*External_{0})");
 static constexpr TCHAR RigVM_JoinSegmentPathFormat[] = TEXT("{0}.{1}");
 static constexpr TCHAR RigVM_GetArrayElementSafeFormat[] = TEXT("GetArrayElementSafe<{0}>({1}, {2})");
@@ -258,18 +259,13 @@ FString FRigVMCodeGenerator::DumpProperties(bool bForHeader, int32 InInstruction
 	FStringArray Lines;
 	for(int32 Index = 0; Index < Properties.Num(); Index++)
 	{
-		if(!IsPropertyPartOfGroup(Index, InInstructionGroup))
-		{
-			continue;
-		}
-
 		const FPropertyInfo& PropertyInfo = Properties[Index];
 		
 		// in headers we only dump the work / sliced properties,
-		// and for source files we only dump the non-sliced
-		if(bForHeader != (PropertyInfo.PropertyType == ERigVMNativizedPropertyType::Sliced ||
-							(PropertyInfo.PropertyType == ERigVMNativizedPropertyType::Work &&
-								PropertyInfo.Groups.Num() > 1)))
+		// and for source files we only dump the non-sliced (and initialized sliced)
+		if(PropertyInfo.PropertyType != ERigVMNativizedPropertyType::Sliced &&
+			bForHeader != (PropertyInfo.PropertyType == ERigVMNativizedPropertyType::Work &&
+								PropertyInfo.Groups.Num() > 1))
 		{
 			continue;
 		}
@@ -277,20 +273,38 @@ FString FRigVMCodeGenerator::DumpProperties(bool bForHeader, int32 InInstruction
 		const FRigVMPropertyDescription& Property = PropertyInfo.Description;
 		check(Property.IsValid());
 
-		if(PropertyInfo.PropertyType == ERigVMNativizedPropertyType::Literal)
+		if(PropertyInfo.PropertyType == ERigVMNativizedPropertyType::Literal ||
+			(!bForHeader && PropertyInfo.PropertyType == ERigVMNativizedPropertyType::Sliced))
 		{
-			FRigVMOperand Operand(ERigVMMemoryType::Literal, PropertyInfo.MemoryPropertyIndex, INDEX_NONE);
+			FRigVMOperand Operand;
+			if (PropertyInfo.PropertyType == ERigVMNativizedPropertyType::Literal)
+			{
+				Operand = FRigVMOperand(ERigVMMemoryType::Literal, PropertyInfo.MemoryPropertyIndex, INDEX_NONE);
+			}
+			else if (PropertyInfo.PropertyType == ERigVMNativizedPropertyType::Sliced)
+			{
+				Operand = FRigVMOperand(ERigVMMemoryType::Work, PropertyInfo.MemoryPropertyIndex, INDEX_NONE);
+			}
 			FString OperandName = GetOperandName(Operand, false);
 			FString CPPType = GetOperandCPPType(Operand);
 
 			FString BaseCPPType = CPPType;
-			const bool bIsArray = RigVMTypeUtils::IsArrayType(CPPType);
+			bool bIsArray = RigVMTypeUtils::IsArrayType(CPPType);
 			if (bIsArray)
 			{
 				BaseCPPType = RigVMTypeUtils::BaseTypeFromArrayType(CPPType);
 			}
 			
 			FString DefaultValue = Property.DefaultValue;
+			if (!bForHeader && PropertyInfo.PropertyType == ERigVMNativizedPropertyType::Sliced)
+			{
+				// The const definition of a slice should have the element type
+				OperandName += TEXT("_Const");
+				CPPType = BaseCPPType;
+				bIsArray = false;
+				DefaultValue = DefaultValue.LeftChop(1);
+				DefaultValue = DefaultValue.RightChop(1);
+			}			
 
 			if (UScriptStruct* ScriptStruct = Cast<UScriptStruct>(Property.CPPTypeObject))
 			{
@@ -488,8 +502,17 @@ FString FRigVMCodeGenerator::DumpProperties(bool bForHeader, int32 InInstruction
 			FString CPPType = GetOperandCPPType(Operand);
 
 			const FString MappedType = GetMappedType(Property.CPPType);
-			const FString Line = Format(RigVM_MemberPropertyFormat, *MappedType, *SanitizeName(Property.Name.ToString(), Property.CPPType));
-			Lines.Add(Line);
+
+			if (bForHeader && PropertyInfo.PropertyType == ERigVMNativizedPropertyType::Sliced)
+			{
+				const FString Line = Format(RigVM_MemberPropertyFormatNoDefault, *MappedType, *SanitizeName(Property.Name.ToString(), Property.CPPType));
+				Lines.Add(Line);
+			}
+			else
+			{				
+				const FString Line = Format(RigVM_MemberPropertyFormat, *MappedType, *SanitizeName(Property.Name.ToString(), Property.CPPType), PropertyInfo.Description.DefaultValue);
+				Lines.Add(Line);
+			}
 		}
 	}
 
diff --git a/Engine/Source/Runtime/RigVM/Public/RigVMCore/RigVMNativized.h b/Engine/Source/Runtime/RigVM/Public/RigVMCore/RigVMNativized.h
--- a/Engine/Source/Runtime/RigVM/Public/RigVMCore/RigVMNativized.h
+++ b/Engine/Source/Runtime/RigVM/Public/RigVMCore/RigVMNativized.h
@@ -188,12 +188,19 @@ protected:
 	}
 
 	template<typename T>
-	T& GetOperandSlice(TArray<T>& InOutArray)
+	T& GetOperandSlice(TArray<T>& InOutArray, const T* InDefaultValue = nullptr)
 	{
 		const int32 SliceIndex = Context.GetSlice().GetIndex();
 		if(InOutArray.Num() <= SliceIndex)
 		{
-			InOutArray.AddDefaulted(1 + SliceIndex - InOutArray.Num());
+			const int32 FirstInsertedItem = InOutArray.AddDefaulted(1 + SliceIndex - InOutArray.Num());
+			if (InDefaultValue)
+			{
+				for (int32 i=FirstInsertedItem; i<InOutArray.Num();++i)
+				{
+					InOutArray[i] = *InDefaultValue;
+				}
+			}
 		}
 		return InOutArray[SliceIndex];
 	}
