RigVM: Fix patch select / if node code

diff --git a/Engine/Source/Developer/RigVMDeveloper/Private/RigVMModel/RigVMController.cpp b/Engine/Source/Developer/RigVMDeveloper/Private/RigVMModel/RigVMController.cpp
--- a/Engine/Source/Developer/RigVMDeveloper/Private/RigVMModel/RigVMController.cpp
+++ b/Engine/Source/Developer/RigVMDeveloper/Private/RigVMModel/RigVMController.cpp
@@ -19724,12 +19724,26 @@ void URigVMController::PatchIfSelectNodesOnLoad()
 				false, 
 				false);
 
-			for(URigVMPin* Pin : NewNode->GetPins())
+			if(!FRigVMRegistry::Get().IsWildCardType(TypeIndex))
 			{
-				if(Pin->IsWildCard() && !FRigVMRegistry::Get().IsWildCardType(TypeIndex))
+				TArray<int32> Permutations;
+				FRigVMTemplateTypeMap Types;
+				Types.Add(IfOrSelectNode->GetPins().Last()->GetFName(), TypeIndex);
+				Template->Resolve(Types, Permutations, false);
+
+				for(URigVMPin* Pin : NewNode->GetPins())
 				{
-					ChangePinType(Pin, TypeIndex, false, false);
-					NewNode->UpdateFilteredPermutations(Pin, {TypeIndex});
+					if(Pin->IsWildCard())
+					{
+						if(const TRigVMTypeIndex* ResolvedTypeIndex = Types.Find(Pin->GetFName()))
+						{
+							if(!FRigVMRegistry::Get().IsWildCardType(*ResolvedTypeIndex))
+							{
+								ChangePinType(Pin, *ResolvedTypeIndex, false, false);
+								NewNode->UpdateFilteredPermutations(Pin, {*ResolvedTypeIndex});
+							}
+						}
+					}
 				}
 			}
 
