Control Rig: Fix writting at invalid address during compilation

diff --git a/Engine/Source/Developer/RigVMDeveloper/Private/RigVMCompiler/RigVMCompiler.cpp b/Engine/Source/Developer/RigVMDeveloper/Private/RigVMCompiler/RigVMCompiler.cpp
--- a/Engine/Source/Developer/RigVMDeveloper/Private/RigVMCompiler/RigVMCompiler.cpp
+++ b/Engine/Source/Developer/RigVMDeveloper/Private/RigVMCompiler/RigVMCompiler.cpp
@@ -1770,11 +1770,14 @@ int32 URigVMCompiler::TraverseCallExtern(const FRigVMCallExternExprAST* InExpr,
 			for(int32 BlockIndex = 0; BlockIndex < BlockNames.Num(); BlockIndex++)
 			{
 				const FName& BlockName = BlockNames[BlockIndex];
-
-				FRigVMBranchInfo& BranchInfo = WorkData.VM->GetByteCode().BranchInfos[BranchIndices[BlockIndex]];
-				BranchInfo.Label = BlockName;
-				BranchInfo.InstructionIndex = JumpToBranchInstructionIndex;
-				BranchInfo.FirstInstruction = WorkData.VM->GetByteCode().GetNumInstructions();
+				int32 BranchIndex = BranchIndices[BlockIndex];
+				{
+					FRigVMBranchInfo& BranchInfo = WorkData.VM->GetByteCode().BranchInfos[BranchIndex];
+					BranchInfo.Label = BlockName;
+					BranchInfo.InstructionIndex = JumpToBranchInstructionIndex;
+					BranchInfo.FirstInstruction = WorkData.VM->GetByteCode().GetNumInstructions();
+					// BranchInfo can be invalidated by ByteCode array reallocs in the code below, so do not keep a reference to it
+				}
 
 				// check if the block requires slicing or not.
 				// (do we want the private state of the nodes to be unique per run of the block)
@@ -1806,7 +1809,7 @@ int32 URigVMCompiler::TraverseCallExtern(const FRigVMCallExternExprAST* InExpr,
 					WorkData.VM->GetByteCode().AddJumpOp(ERigVMOpCode::JumpBackward, JumpToCallExternInstruction - CallExternInstructionIndex);
 				}
 
-				BranchInfo.LastInstruction = WorkData.VM->GetByteCode().GetNumInstructions() - 1;
+				WorkData.VM->GetByteCode().BranchInfos[BranchIndex].LastInstruction = WorkData.VM->GetByteCode().GetNumInstructions() - 1;
 			}
 		}
 	}
