Control Rig: Fix default values in nativized rigs

diff --git a/Engine/Source/Developer/RigVMDeveloper/Private/RigVMCompiler/RigVMCodeGenerator.cpp b/Engine/Source/Developer/RigVMDeveloper/Private/RigVMCompiler/RigVMCodeGenerator.cpp
--- a/Engine/Source/Developer/RigVMDeveloper/Private/RigVMCompiler/RigVMCodeGenerator.cpp
+++ b/Engine/Source/Developer/RigVMDeveloper/Private/RigVMCompiler/RigVMCodeGenerator.cpp
@@ -7,6 +7,7 @@
 #include "Algo/Count.h"
 #include "Animation/Rig.h"
 #include "RigVMModel/Nodes/RigVMDispatchNode.h"
+#include "RigVMStringUtils.h"
 
 static constexpr TCHAR RigVM_CommaSeparator[] = TEXT(", ");
 static constexpr TCHAR RigVM_NewLineFormat[] = TEXT("\r\n");
@@ -31,8 +32,9 @@ static constexpr TCHAR RigVM_DefineEntryNameFormat[] = TEXT("const FName U{0}::E
 static constexpr TCHAR RigVM_DeclareBlockNameFormat[] = TEXT("\tstatic const FName BlockName_{0};");
 static constexpr TCHAR RigVM_DefineBlockNameFormat[] = TEXT("const FName U{0}::BlockName_{1} = TEXT(\"{2}\");");
 static constexpr TCHAR RigVM_DefineConstFormatNoDefault[] = TEXT("\tstatic const {0} {1};");
-static constexpr TCHAR RigVM_DefineStructConstFormat[] = TEXT("\tstatic const {0} {1} = URigVMNativized::GetStructConstant<{0}>(TEXT(\"{2}\"));");
-static constexpr TCHAR RigVM_DefineStructArrayConstFormat[] = TEXT("\tstatic const {0} {1} = URigVMNativized::GetStructArrayConstant<{2}>(TEXT(\"{3}\"));");
+static constexpr TCHAR RigVM_StructConstantArrayArrayValue[] = TEXT("URigVMNativized::GetStructArrayArrayConstant<{0}>(TEXT(\"{1}\"))");
+static constexpr TCHAR RigVM_StructConstantArrayValue[] = TEXT("URigVMNativized::GetStructArrayConstant<{0}>(TEXT(\"{1}\"))");
+static constexpr TCHAR RigVM_StructConstantValue[] = TEXT("URigVMNativized::GetStructConstant<{0}>(TEXT(\"{1}\"))");
 static constexpr TCHAR RigVM_DefineConstFormat[] = TEXT("\tstatic const {0} {1} = {2};");
 static constexpr TCHAR RigVM_NameNoneFormat[] = TEXT("FName(NAME_None)");
 static constexpr TCHAR RigVM_EmptyStringFormat[] = TEXT("FString()");
@@ -149,24 +151,6 @@ FString FRigVMCodeGenerator::DumpIncludes(bool bLog)
 	return DumpLines(Lines, bLog);
 }
 
-FString FRigVMCodeGenerator::DumpWrappedArrayTypes(bool bLog)
-{
-	FStringArray Lines;
-	for(const FString& WrappedArrayType : WrappedArrayCPPTypes)
-	{
-		if (!Lines.IsEmpty())
-		{
-			Lines.Add(FString());
-		}
-
-		const FString TypeName = GetMappedArrayTypeName(WrappedArrayType);
-		const FString Line = Format(RigVM_WrappedArrayTypeFormat, *ModuleName.ToUpper(), *TypeName, *WrappedArrayType);
-		Lines.Add(Line);
-	}
-
-	return DumpLines(Lines, bLog);
-}
-
 FString FRigVMCodeGenerator::DumpExternalVariables(bool bForHeader, bool bLog)
 {
 	FStringArray Lines;
@@ -290,9 +274,15 @@ FString FRigVMCodeGenerator::DumpProperties(bool bForHeader, int32 InInstruction
 
 			FString BaseCPPType = CPPType;
 			bool bIsArray = RigVMTypeUtils::IsArrayType(CPPType);
+			bool bIsDoubleArray = false;
 			if (bIsArray)
 			{
 				BaseCPPType = RigVMTypeUtils::BaseTypeFromArrayType(CPPType);
+				bIsDoubleArray = RigVMTypeUtils::IsArrayType(BaseCPPType);
+				if (bIsDoubleArray)
+				{
+					BaseCPPType = RigVMTypeUtils::BaseTypeFromArrayType(BaseCPPType);
+				}
 			}
 			
 			FString DefaultValue = Property.DefaultValue;
@@ -300,68 +290,44 @@ FString FRigVMCodeGenerator::DumpProperties(bool bForHeader, int32 InInstruction
 			{
 				// The const definition of a slice should have the element type
 				OperandName += TEXT("_Const");
-				CPPType = BaseCPPType;
-				bIsArray = false;
+				if (bIsDoubleArray)
+				{
+					CPPType = RigVMTypeUtils::ArrayTypeFromBaseType(BaseCPPType);
+					bIsArray = true;
+					bIsDoubleArray = false;
+				}
+				else
+				{
+					CPPType = BaseCPPType;
+					bIsArray = RigVMTypeUtils::IsArrayType(CPPType);
+					if (bIsArray)
+					{
+						BaseCPPType = RigVMTypeUtils::BaseTypeFromArrayType(BaseCPPType);
+					}
+				}
 				DefaultValue = DefaultValue.LeftChop(1);
 				DefaultValue = DefaultValue.RightChop(1);
-			}			
+			}
+			DefaultValue = SanitizeValue(DefaultValue, CPPType, PropertyInfo.Description.CPPTypeObject);
 
 			if (UScriptStruct* ScriptStruct = Cast<UScriptStruct>(Property.CPPTypeObject))
 			{
-				DefaultValue = DefaultValue.ReplaceCharWithEscapedChar();
-
 				if(DefaultValue.IsEmpty())
 				{
 					Lines.Add(Format(
 						RigVM_DefineConstFormatNoDefault,
 						*CPPType,
-						*GetOperandName(Operand, false)
+						*OperandName
 					));
 				}
 				else
 				{
-					if(const TStructConstGenerator* StructConstGenerator = GetStructConstGenerators().Find(*BaseCPPType))
-					{
-						if(bIsArray)
-						{
-							FStringArray DefaultValues = URigVMPin::SplitDefaultValue(DefaultValue);
-							for(int32 DefaultValueIndex = 0; DefaultValueIndex < DefaultValues.Num(); DefaultValueIndex++)
-							{
-								DefaultValues[DefaultValueIndex] = (*StructConstGenerator)(DefaultValues[DefaultValueIndex]); 
-							}
-							DefaultValue = Format(RigVM_CurlyBracesFormat, FString::Join(DefaultValues, RigVM_CommaSeparator));
-						}
-						else
-						{
-							DefaultValue = (*StructConstGenerator)(DefaultValue);
-						}
-						
-						Lines.Add(Format(
-							RigVM_DefineConstFormat,
-							*CPPType,
-							*OperandName,
-							*DefaultValue
-						));
-					}
-					else if (bIsArray)
-					{
-						Lines.Add(Format(
-							RigVM_DefineStructArrayConstFormat,
-							*CPPType,
-							*OperandName,
-							*BaseCPPType,
-							*DefaultValue
-						));
-					}
-					else
-					{
-						Lines.Add(Format(
-							RigVM_DefineStructConstFormat,
-							*CPPType,
-							*OperandName,
-							*DefaultValue
-						));
-					}
+					Lines.Add(Format(
+						RigVM_DefineConstFormat,
+						*CPPType,
+						*OperandName,
+						*DefaultValue
+					));
 				}
 			}
 			else if (const UEnum* Enum = Cast<UEnum>(Property.CPPTypeObject))
@@ -377,41 +343,11 @@ FString FRigVMCodeGenerator::DumpProperties(bool bForHeader, int32 InInstruction
 				{
 					CPPType = RigVMTypeUtils::ArrayTypeFromBaseType(BaseCPPType);
 				}
-
-				FStringArray DefaultValues;
-				if (bIsArray)
-				{
-					DefaultValues = URigVMPin::SplitDefaultValue(DefaultValue);
-				}
-				else
+				if (bIsDoubleArray)
 				{
-					DefaultValues.Add(DefaultValue);
+					CPPType = RigVMTypeUtils::ArrayTypeFromBaseType(CPPType);
 				}
 
-				for(int32 DefaultValueIndex = 0; DefaultValueIndex < DefaultValues.Num(); DefaultValueIndex++)
-				{
-					if (DefaultValues[DefaultValueIndex].IsNumeric())
-					{
-						const int64 EnumIndex = FCString::Atoi64(*DefaultValues[DefaultValueIndex]);
-						const FString EnumName = Enum->GetNameStringByValue(EnumIndex);
-						DefaultValues[DefaultValueIndex] = Enum->GenerateFullEnumName(*EnumName);
-					}
-					else if (!UEnum::IsFullEnumName(*DefaultValues[DefaultValueIndex]))
-					{
-						DefaultValues[DefaultValueIndex] = Enum->GenerateFullEnumName(*DefaultValues[DefaultValueIndex]);
-					}
-				}
-				
-				if (bIsArray)
-				{
-					DefaultValue = Format(RigVM_CurlyBracesFormat, *FString::Join(DefaultValues, TEXT(",")));
-				}
-				else
-				{
-					DefaultValue = DefaultValues[0];
-				}
-
-
 				if(DefaultValue.IsEmpty())
 				{
 					Lines.Add(Format(
@@ -437,42 +373,10 @@ FString FRigVMCodeGenerator::DumpProperties(bool bForHeader, int32 InInstruction
 				if (bIsArray)
 				{
 					bUseConstExpr = false;
-					
-					if (DefaultValue.StartsWith(TEXT("(")) && DefaultValue.EndsWith(TEXT(")")))
-					{
-						DefaultValue = Format(RigVM_CurlyBracesFormat, *DefaultValue.Mid(1, DefaultValue.Len() - 2));
-					}
-				}
-
-				if (BaseCPPType == RigVMTypeUtils::BoolType)
-				{
-					DefaultValue.ToLowerInline();
 				}
 
 				if (CPPType == RigVMTypeUtils::FNameType || CPPType == RigVMTypeUtils::FStringType)
 				{
-					if(DefaultValue.IsEmpty())
-					{
-						if(CPPType == RigVMTypeUtils::FNameType)
-						{
-							DefaultValue = RigVM_NameNoneFormat;
-						}
-						else
-						{
-							DefaultValue = RigVM_EmptyStringFormat;
-						}
-					}
-					else
-					{
-						if (DefaultValue.StartsWith(TEXT("\"")) && DefaultValue.EndsWith(TEXT("\"")))
-						{
-							DefaultValue = Format(RigVM_TextFormat, *DefaultValue);
-						}
-						else
-						{
-							DefaultValue = Format(RigVM_QuotedTextFormat, *DefaultValue);
-						}
-					}
 					bUseConstExpr = false;
 				}
 
@@ -509,8 +413,9 @@ FString FRigVMCodeGenerator::DumpProperties(bool bForHeader, int32 InInstruction
 				Lines.Add(Line);
 			}
 			else
-			{				
-				const FString Line = Format(RigVM_MemberPropertyFormat, *MappedType, *SanitizeName(Property.Name.ToString(), Property.CPPType), PropertyInfo.Description.DefaultValue);
+			{
+				const FString DefaultValue = SanitizeValue(PropertyInfo.Description.DefaultValue, Property.CPPType, Property.CPPTypeObject);
+				const FString Line = Format(RigVM_MemberPropertyFormat, *MappedType, *SanitizeName(Property.Name.ToString(), Property.CPPType), DefaultValue);
 				Lines.Add(Line);
 			}
 		}
@@ -525,8 +430,9 @@ FString FRigVMCodeGenerator::DumpDispatches(bool bLog)
 
 	if(!Dispatches.IsEmpty())
 	{
-		for(const FRigVMDispatchInfo& Info : Dispatches)
+		for(const TPair<FString, FRigVMDispatchInfo>& Pair : Dispatches)
 		{
+			const FRigVMDispatchInfo& Info = Pair.Value;
 			const FString FunctionName = Info.Function->GetName();
 			const FRigVMDispatchFactory* Factory = Info.Function->Factory;
 
@@ -850,11 +756,11 @@ FString FRigVMCodeGenerator::DumpInstructions(const FString& InPrefix, const TAr
 				else if(Function->Factory)
 				{
 					FString DispatchName;
-					for(const FRigVMDispatchInfo& Info : Dispatches)
+					for(const TPair<FString, FRigVMDispatchInfo>& Pair : Dispatches)
 					{
-						if(Info.Function == Function)
+						if(Pair.Value.Function == Function)
 						{
-							DispatchName = Info.Name;
+							DispatchName = Pair.Value.Name;
 							break;
 						}
 					}
@@ -1205,12 +1111,7 @@ FString FRigVMCodeGenerator::DumpHeader(bool bLog)
 	Lines.Add(DumpIncludes());
 	Lines.Add(Format(RigVM_GeneratedIncludeFormat, *ClassName));
 	Lines.Emplace();
-	if(WrappedArrayCPPTypes.Num() > 0)
-	{
-		Lines.Add(DumpWrappedArrayTypes());
-		Lines.Emplace();
-	}
-
+	
 	Lines.Add(Format(RigVM_UClassDefinitionFormat, *ModuleName.ToUpper(), *ClassName));
 	Lines.Add(Format(RigVM_GetVMHashFormat, FString::Printf(TEXT("%lu"), VM->GetVMHash())));
 	Lines.Add(Format(RigVM_GetEntryNamesFormat, *FString::Join(FormattedEntries, RigVM_CommaSeparator)));
@@ -1385,7 +1286,6 @@ void FRigVMCodeGenerator::Reset()
 	Dispatches.Reset();
 	RequiredUProperties.Reset();
 	InstructionGroups.Reset();
-	WrappedArrayCPPTypes.Reset();
 	MappedCPPTypes.Reset();
 	Properties.Reset();
 	PropertyNameToIndex.Reset();
@@ -1494,7 +1394,17 @@ void FRigVMCodeGenerator::ParseRequiredUProperties()
 				const FString FunctionName = Function->GetName();
 				const int32 PermutationIndex = Factory->GetTemplate()->FindPermutation(Function); 
 				check(PermutationIndex != INDEX_NONE);
+
+				for(const FRigVMFunctionArgument& Argument : Function->Arguments)
+				{
+					AddRequiredUProperty(Argument.Type);
+				}
+				
 				const FString DispatchKey = Format(RigVM_DispatchKeyFormat, Factory->GetFactoryName().ToString(), PermutationIndex);
+				if (Dispatches.Contains(DispatchKey))
+				{
+					continue;
+				}
 
 				FRigVMDispatchContext Context;
 				if(URigVMDispatchNode* DispatchNode = Cast<URigVMDispatchNode>(ByteCode.GetSubjectForInstruction(InstructionIndex)))
@@ -1502,12 +1412,7 @@ void FRigVMCodeGenerator::ParseRequiredUProperties()
 					Context = DispatchNode->GetDispatchContext();
 				}
 
-				Dispatches.Add({DispatchKey, Function, Context});
-
-				for(const FRigVMFunctionArgument& Argument : Function->Arguments)
-				{
-					AddRequiredUProperty(Argument.Type);
-				}
+				Dispatches.Add(DispatchKey, {DispatchKey, Function, Context});
 			}
 		}
 	}
@@ -1586,21 +1491,6 @@ void FRigVMCodeGenerator::ParseProperty(ERigVMMemoryType InMemoryType, const FPr
 
 		const int32 LookupIndex = Properties.Add(Info);
 		PropertyNameToIndex.Add(Info.Description.Name, LookupIndex);
-
-		// if this property is a nested array, we need to inject out wrapped types
-		if(PropertyDescription.Containers.Num() > 1)
-		{
-			check(PropertyDescription.Containers[0] == EPinContainerType::Array);
-			check(PropertyDescription.Containers[1] == EPinContainerType::Array);
-
-			FString BaseElementType = RigVMTypeUtils::BaseTypeFromArrayType(PropertyDescription.CPPType);
-			BaseElementType = RigVMTypeUtils::BaseTypeFromArrayType(BaseElementType);
-			WrappedArrayCPPTypes.AddUnique(BaseElementType);
-			
-			const FString ArrayArray = RigVMTypeUtils::ArrayTypeFromBaseType(RigVMTypeUtils::ArrayTypeFromBaseType(BaseElementType));
-			const FString MappedType = RigVMTypeUtils::ArrayTypeFromBaseType(GetMappedArrayTypeName(BaseElementType));
-			MappedCPPTypes.Add(ArrayArray, FMappedType(MappedType, TEXT(".Array")));
-		}
 	}
 }
 
@@ -2033,7 +1923,7 @@ FString FRigVMCodeGenerator::GetOperandName(const FRigVMOperand& InOperand, bool
 		while(!RemainingSegmentPath.IsEmpty() && (CurrentProperty != nullptr))
 		{
 			FString Left, Right;
-			if(!URigVMPin::SplitPinPathAtStart(RemainingSegmentPath, Left, Right))
+			if(!RigVMStringUtils::SplitPinPathAtStart(RemainingSegmentPath, Left, Right))
 			{
 				Left = RemainingSegmentPath;
 				Right.Reset();
@@ -2144,6 +2034,168 @@ FString FRigVMCodeGenerator::SanitizeName(const FString& InName, const FString&
 	return Name;
 }
 
+FString FRigVMCodeGenerator::SanitizeValue(const FString& InValue, const FString& InCPPType, const UObject* InCPPTypeObject)
+{
+	FString DefaultValue = InValue;
+	FString CPPType = InCPPType;
+	FString BaseCPPType = CPPType;
+	bool bIsArray = RigVMTypeUtils::IsArrayType(CPPType);
+	bool bIsDoubleArray = false;
+	if (bIsArray)
+	{
+		BaseCPPType = RigVMTypeUtils::BaseTypeFromArrayType(CPPType);
+		bIsDoubleArray = RigVMTypeUtils::IsArrayType(BaseCPPType);
+		if (bIsDoubleArray)
+		{
+			BaseCPPType = RigVMTypeUtils::BaseTypeFromArrayType(BaseCPPType);			
+		}
+	}
+	
+	if (const UScriptStruct* ScriptStruct = Cast<UScriptStruct>(InCPPTypeObject))
+	{
+		DefaultValue = DefaultValue.ReplaceCharWithEscapedChar();
+
+		if(!DefaultValue.IsEmpty())
+		{
+			if(const TStructConstGenerator* StructConstGenerator = GetStructConstGenerators().Find(*BaseCPPType))
+			{
+				if (bIsDoubleArray)
+				{
+					FStringArray ArrayDefaultValues = RigVMStringUtils::SplitDefaultValue(DefaultValue);
+					for(int32 ArrayDefaultValueIndex = 0; ArrayDefaultValueIndex < ArrayDefaultValues.Num(); ArrayDefaultValueIndex++)
+					{
+						FStringArray DefaultValues = RigVMStringUtils::SplitDefaultValue(ArrayDefaultValues[ArrayDefaultValueIndex]);
+						for(int32 DefaultValueIndex = 0; DefaultValueIndex < DefaultValues.Num(); DefaultValueIndex++)
+						{
+							DefaultValues[DefaultValueIndex] = (*StructConstGenerator)(DefaultValues[DefaultValueIndex]); 
+						}
+						ArrayDefaultValues[ArrayDefaultValueIndex] = Format(RigVM_CurlyBracesFormat, FString::Join(DefaultValues, RigVM_CommaSeparator));
+					}
+					DefaultValue = Format(RigVM_CurlyBracesFormat, FString::Join(ArrayDefaultValues, RigVM_CommaSeparator));
+				}
+				else if(bIsArray)
+				{
+					FStringArray DefaultValues = RigVMStringUtils::SplitDefaultValue(DefaultValue);
+					for(int32 DefaultValueIndex = 0; DefaultValueIndex < DefaultValues.Num(); DefaultValueIndex++)
+					{
+						DefaultValues[DefaultValueIndex] = (*StructConstGenerator)(DefaultValues[DefaultValueIndex]); 
+					}
+					DefaultValue = Format(RigVM_CurlyBracesFormat, FString::Join(DefaultValues, RigVM_CommaSeparator));
+				}
+				else
+				{
+					DefaultValue = (*StructConstGenerator)(DefaultValue);
+				}
+			}
+			else
+			{
+				if (bIsDoubleArray)
+				{
+					DefaultValue = Format(RigVM_StructConstantArrayArrayValue, *BaseCPPType, *DefaultValue);
+				}
+				else if(bIsArray)
+				{
+					DefaultValue = Format(RigVM_StructConstantArrayValue, *BaseCPPType, *DefaultValue);
+				}
+				else
+				{
+					DefaultValue = Format(RigVM_StructConstantValue, *BaseCPPType, *DefaultValue);
+				}
+			}
+		}
+	}
+	else if (const UEnum* Enum = Cast<UEnum>(InCPPTypeObject))
+	{
+		BaseCPPType = Enum->GetName();
+		if (Enum->GetCppForm() == UEnum::ECppForm::Namespaced)
+		{
+			BaseCPPType += RigVM_EnumTypeSuffixFormat;
+		}
+		
+		CPPType = BaseCPPType;
+		if (bIsArray)
+		{
+			CPPType = RigVMTypeUtils::ArrayTypeFromBaseType(BaseCPPType);
+		}
+
+		FStringArray DefaultValues;
+		if (bIsArray)
+		{
+			DefaultValues = RigVMStringUtils::SplitDefaultValue(DefaultValue);
+		}
+		else
+		{
+			DefaultValues.Add(DefaultValue);
+		}
+
+		for(int32 DefaultValueIndex = 0; DefaultValueIndex < DefaultValues.Num(); DefaultValueIndex++)
+		{
+			if (DefaultValues[DefaultValueIndex].IsNumeric())
+			{
+				const int64 EnumIndex = FCString::Atoi64(*DefaultValues[DefaultValueIndex]);
+				const FString EnumName = Enum->GetNameStringByValue(EnumIndex);
+				DefaultValues[DefaultValueIndex] = Enum->GenerateFullEnumName(*EnumName);
+			}
+			else if (!UEnum::IsFullEnumName(*DefaultValues[DefaultValueIndex]))
+			{
+				DefaultValues[DefaultValueIndex] = Enum->GenerateFullEnumName(*DefaultValues[DefaultValueIndex]);
+			}
+		}
+		
+		if (bIsArray)
+		{
+			DefaultValue = Format(RigVM_CurlyBracesFormat, *FString::Join(DefaultValues, TEXT(",")));
+		}
+		else
+		{
+			DefaultValue = DefaultValues[0];
+		}
+	}
+	else
+	{
+		if (bIsArray)
+		{
+			if (DefaultValue.StartsWith(TEXT("(")) && DefaultValue.EndsWith(TEXT(")")))
+			{
+				DefaultValue = Format(RigVM_CurlyBracesFormat, *DefaultValue.Mid(1, DefaultValue.Len() - 2));
+			}
+		}
+
+		if (BaseCPPType == RigVMTypeUtils::BoolType)
+		{
+			DefaultValue.ToLowerInline();
+		}
+
+		if (CPPType == RigVMTypeUtils::FNameType || CPPType == RigVMTypeUtils::FStringType)
+		{
+			if(DefaultValue.IsEmpty())
+			{
+				if(CPPType == RigVMTypeUtils::FNameType)
+				{
+					DefaultValue = RigVM_NameNoneFormat;
+				}
+				else
+				{
+					DefaultValue = RigVM_EmptyStringFormat;
+				}
+			}
+			else
+			{
+				if (DefaultValue.StartsWith(TEXT("\"")) && DefaultValue.EndsWith(TEXT("\"")))
+				{
+					DefaultValue = Format(RigVM_TextFormat, *DefaultValue);
+				}
+				else
+				{
+					DefaultValue = Format(RigVM_QuotedTextFormat, *DefaultValue);
+				}
+			}
+		}
+	}
+
+	return DefaultValue;
+}
+
 FRigVMPropertyDescription FRigVMCodeGenerator::GetPropertyForOperand(const FRigVMOperand& InOperand) const
 {
 	CheckOperand(InOperand);
@@ -2438,7 +2490,7 @@ const TMap<FName, FRigVMCodeGenerator::TStructConstGenerator>& FRigVMCodeGenerat
 	{
 		static FStringArray SplitIntoArray(const FString& InValue)
 		{
-			return URigVMPin::SplitDefaultValue(InValue);
+			return RigVMStringUtils::SplitDefaultValue(InValue);
 		}
 		
 		static FStringMap SplitIntoMap(const FString& InValue)
diff --git a/Engine/Source/Developer/RigVMDeveloper/Public/RigVMCompiler/RigVMCodeGenerator.h b/Engine/Source/Developer/RigVMDeveloper/Public/RigVMCompiler/RigVMCodeGenerator.h
--- a/Engine/Source/Developer/RigVMDeveloper/Public/RigVMCompiler/RigVMCodeGenerator.h
+++ b/Engine/Source/Developer/RigVMDeveloper/Public/RigVMCompiler/RigVMCodeGenerator.h
@@ -22,7 +22,6 @@ public:
 	}
 
 	FString DumpIncludes(bool bLog = false);
-	FString DumpWrappedArrayTypes(bool bLog = false);
 	FString DumpExternalVariables(bool bForHeader, bool bLog = false);
 	FString DumpEntries(bool bForHeader, bool bLog = false);
 	FString DumpBlockNames(bool bForHeader, bool bLog = false);
@@ -103,6 +102,7 @@ private:
 	FString GetOperandCPPType(const FRigVMOperand& InOperand) const;
 	FString GetOperandCPPBaseType(const FRigVMOperand& InOperand) const;
 	static FString SanitizeName(const FString& InName, const FString& CPPType = FString()); 
+	static FString SanitizeValue(const FString& InValue, const FString& CPPType, const UObject* CPPTypeObject); 
 	FRigVMPropertyDescription GetPropertyForOperand(const FRigVMOperand& InOperand) const;
 	const FRigVMPropertyPathDescription& GetPropertyPathForOperand(const FRigVMOperand& InOperand) const;
 	int32 GetPropertyIndex(const FRigVMPropertyDescription& InProperty) const;
@@ -246,10 +246,9 @@ private:
 		FRigVMDispatchContext Context;
 	};
 	
-	TArray<FRigVMDispatchInfo> Dispatches;
+	TMap<FString, FRigVMDispatchInfo> Dispatches;
 	TMap<TRigVMTypeIndex,TTuple<FString,FString>> RequiredUProperties;
 	TArray<FInstructionGroup> InstructionGroups;
-	FStringArray WrappedArrayCPPTypes;
 	TMap<FString, FMappedType> MappedCPPTypes;
 	TArray<FPropertyInfo> Properties;
 	TMap<FName, int32> PropertyNameToIndex;
diff --git a/Engine/Source/Runtime/RigVM/Private/RigVMCore/RigVMMemoryStorage.cpp b/Engine/Source/Runtime/RigVM/Private/RigVMCore/RigVMMemoryStorage.cpp
--- a/Engine/Source/Runtime/RigVM/Private/RigVMCore/RigVMMemoryStorage.cpp
+++ b/Engine/Source/Runtime/RigVM/Private/RigVMCore/RigVMMemoryStorage.cpp
@@ -178,7 +178,7 @@ FRigVMPropertyDescription::FRigVMPropertyDescription(const FProperty* InProperty
 	if(CPPTypeObject == nullptr && Property != nullptr)
 	{
 		const FProperty* ValueProperty = Property;
-		if(const FArrayProperty* ArrayProperty = CastField<FArrayProperty>(ValueProperty))
+		while(const FArrayProperty* ArrayProperty = CastField<FArrayProperty>(ValueProperty))
 		{
 			ValueProperty = ArrayProperty->Inner;
 		}
diff --git a/Engine/Source/Runtime/RigVM/Public/RigVMCore/RigVMNativized.h b/Engine/Source/Runtime/RigVM/Public/RigVMCore/RigVMNativized.h
--- a/Engine/Source/Runtime/RigVM/Public/RigVMCore/RigVMNativized.h
+++ b/Engine/Source/Runtime/RigVM/Public/RigVMCore/RigVMNativized.h
@@ -371,6 +371,36 @@ public:
 		return Value;
 	}
 
+	template <
+		typename T,
+		typename TEnableIf<TRigVMIsBaseStructure<T>::Value, T>::Type* = nullptr
+	>
+	static TArray<TArray<T>> GetStructArrayArrayConstant(const FString& InDefaultValue)
+	{
+		TArray<TArray<T>> Value;
+		if(!InDefaultValue.IsEmpty())
+		{
+			FErrorPipe ErrorPipe;
+			TBaseStructure<T>::Get()->ImportText(*InDefaultValue, &Value, nullptr, PPF_None, &ErrorPipe, FString());
+		}
+		return Value;
+	}
+
+	template <
+		typename T,
+		typename TEnableIf<TModels<CRigVMUStruct, T>::Value>::Type* = nullptr
+	>
+	static TArray<TArray<T>> GetStructArrayArrayConstant(const FString& InDefaultValue)
+	{
+		TArray<TArray<T>> Value;
+		if(!InDefaultValue.IsEmpty())
+		{
+			FErrorPipe ErrorPipe;
+			T::StaticStruct()->ImportText(*InDefaultValue, &Value, nullptr, PPF_None, &ErrorPipe, FString());
+		}
+		return Value;
+	}
+
 protected:
 	
 	class FTransformSetter
