Control Rig: Fix orphan pins on templates triggering link breaks

diff --git a/Engine/Source/Developer/RigVMDeveloper/Private/RigVMModel/RigVMController.cpp b/Engine/Source/Developer/RigVMDeveloper/Private/RigVMModel/RigVMController.cpp
--- a/Engine/Source/Developer/RigVMDeveloper/Private/RigVMModel/RigVMController.cpp
+++ b/Engine/Source/Developer/RigVMDeveloper/Private/RigVMModel/RigVMController.cpp
@@ -16446,7 +16446,7 @@ bool URigVMController::PrepareTemplatePinForType(URigVMPin* InPin, const TArray<
 
 	// It might be a subpin of a struct, in that case we just want to make sure
 	// the type required matches the one in the pin
-	if (InPin->IsStructMember())
+	if (InPin->IsStructMember() || InPin->IsOrphanPin())
 	{
 		TRigVMTypeIndex PinTypeIndex = InPin->GetTypeIndex();
 		
@@ -16772,7 +16772,7 @@ bool URigVMController::PrepareTemplatePinForType(URigVMPin* InPin, const TArray<
 
 TArray<TRigVMTypeIndex> URigVMController::GetFilteredTypes(URigVMPin* InPin)
 {
-	if (!InPin->IsStructMember())
+	if (!InPin->IsStructMember() && !InPin->IsOrphanPin())
 	{
 		if (URigVMTemplateNode* TemplateNode = Cast<URigVMTemplateNode>(InPin->GetNode()))
 		{
