Control Rig: Save serialized graph function in function data

diff --git a/Engine/Plugins/Animation/ControlRig/Source/ControlRigDeveloper/Private/ControlRigBlueprint.cpp b/Engine/Plugins/Animation/ControlRig/Source/ControlRigDeveloper/Private/ControlRigBlueprint.cpp
--- a/Engine/Plugins/Animation/ControlRig/Source/ControlRigDeveloper/Private/ControlRigBlueprint.cpp
+++ b/Engine/Plugins/Animation/ControlRig/Source/ControlRigDeveloper/Private/ControlRigBlueprint.cpp
@@ -616,17 +616,11 @@ bool UControlRigBlueprint::TryImportGraphFromText(const FString& InClipboardText
 	{
 		TGuardValue<FRigVMController_RequestLocalizeFunctionDelegate> RequestLocalizeDelegateGuard(
             FunctionLibraryController->RequestLocalizeFunctionDelegate,
-            FRigVMController_RequestLocalizeFunctionDelegate::CreateLambda([this](TSoftObjectPtr<URigVMLibraryNode> InFunctionToLocalize)
+            FRigVMController_RequestLocalizeFunctionDelegate::CreateLambda([this](FRigVMGraphFunctionIdentifier& InFunctionToLocalize)
             {
-            	if (InFunctionToLocalize.IsValid())
-            	{
-            		URigVMLibraryNode* LibraryNode = InFunctionToLocalize.Get();
-            		BroadcastRequestLocalizeFunctionDialog(LibraryNode);
-
-				   const URigVMLibraryNode* LocalizedFunctionNode = GetLocalFunctionLibrary()->FindPreviouslyLocalizedFunction(LibraryNode);
-				   return LocalizedFunctionNode != nullptr;
-            	}
-            	return false;
+            	BroadcastRequestLocalizeFunctionDialog(InFunctionToLocalize);
+				const URigVMLibraryNode* LocalizedFunctionNode = GetLocalFunctionLibrary()->FindPreviouslyLocalizedFunction(InFunctionToLocalize);
+				return LocalizedFunctionNode != nullptr;
             })
         );
 		
@@ -5182,7 +5176,7 @@ void UControlRigBlueprint::BroadcastPostEditChangeChainProperty(FPropertyChanged
 	PostEditChangeChainPropertyEvent.Broadcast(PropertyChangedChainEvent);
 }
 
-void UControlRigBlueprint::BroadcastRequestLocalizeFunctionDialog(URigVMLibraryNode* InFunction, bool bForce)
+void UControlRigBlueprint::BroadcastRequestLocalizeFunctionDialog(FRigVMGraphFunctionIdentifier InFunction, bool bForce)
 {
 	RequestLocalizeFunctionDialog.Broadcast(InFunction, this, bForce);
 }
diff --git a/Engine/Plugins/Animation/ControlRig/Source/ControlRigDeveloper/Private/Graph/ControlRigGraph.cpp b/Engine/Plugins/Animation/ControlRig/Source/ControlRigDeveloper/Private/Graph/ControlRigGraph.cpp
--- a/Engine/Plugins/Animation/ControlRig/Source/ControlRigDeveloper/Private/Graph/ControlRigGraph.cpp
+++ b/Engine/Plugins/Animation/ControlRig/Source/ControlRigDeveloper/Private/Graph/ControlRigGraph.cpp
@@ -504,9 +504,12 @@ void UControlRigGraph::HandleModifiedEvent(ERigVMGraphNotifType InNotifType, URi
 
 					NewNode->SetFlags(RF_Transactional);
 
-					if (UEdGraphPin* ValuePin = NewNode->FindPin(ModelNode->FindPin("Value")->GetPinPath()))
+					if (URigVMPin* ModelPin = ModelNode->FindPin("Value"))
 					{
-						NewNode->SetColorFromModel(GetSchema()->GetPinTypeColor(ValuePin->PinType));
+						if (UEdGraphPin* ValuePin = NewNode->FindPin(ModelPin->GetPinPath()))
+						{
+							NewNode->SetColorFromModel(GetSchema()->GetPinTypeColor(ValuePin->PinType));
+						}
 					}
 
 					ModelNodePathToEdNode.Add(ModelNode->GetFName(), NewNode);
diff --git a/Engine/Plugins/Animation/ControlRig/Source/ControlRigDeveloper/Private/Graph/ControlRigGraphSchema.cpp b/Engine/Plugins/Animation/ControlRig/Source/ControlRigDeveloper/Private/Graph/ControlRigGraphSchema.cpp
--- a/Engine/Plugins/Animation/ControlRig/Source/ControlRigDeveloper/Private/Graph/ControlRigGraphSchema.cpp
+++ b/Engine/Plugins/Animation/ControlRig/Source/ControlRigDeveloper/Private/Graph/ControlRigGraphSchema.cpp
@@ -496,8 +496,9 @@ FReply FControlRigFunctionDragDropAction::DroppedOnPanel(const TSharedRef< class
 								{
 									if(!FunctionRigBlueprint->IsFunctionPublic(FunctionDefinitionNode->GetFName()))
 									{
-										TargetRigBlueprint->BroadcastRequestLocalizeFunctionDialog(FunctionDefinitionNode);
-										FunctionDefinitionNode = TargetRigBlueprint->GetLocalFunctionLibrary()->FindPreviouslyLocalizedFunction(FunctionDefinitionNode);
+										FRigVMGraphFunctionIdentifier FunctionIdentifier = FunctionDefinitionNode->GetFunctionIdentifier();
+										TargetRigBlueprint->BroadcastRequestLocalizeFunctionDialog(FunctionIdentifier);
+										FunctionDefinitionNode = TargetRigBlueprint->GetLocalFunctionLibrary()->FindPreviouslyLocalizedFunction(FunctionIdentifier);
 									}
 								}
 #endif
diff --git a/Engine/Plugins/Animation/ControlRig/Source/ControlRigDeveloper/Public/ControlRigBlueprint.h b/Engine/Plugins/Animation/ControlRig/Source/ControlRigDeveloper/Public/ControlRigBlueprint.h
--- a/Engine/Plugins/Animation/ControlRig/Source/ControlRigDeveloper/Public/ControlRigBlueprint.h
+++ b/Engine/Plugins/Animation/ControlRig/Source/ControlRigDeveloper/Public/ControlRigBlueprint.h
@@ -39,7 +39,7 @@ DECLARE_EVENT_OneParam(UControlRigBlueprint, FOnExternalVariablesChanged, const
 DECLARE_EVENT_TwoParams(UControlRigBlueprint, FOnNodeDoubleClicked, UControlRigBlueprint*, URigVMNode*);
 DECLARE_EVENT_OneParam(UControlRigBlueprint, FOnGraphImported, UEdGraph*);
 DECLARE_EVENT_OneParam(UControlRigBlueprint, FOnPostEditChangeChainProperty, FPropertyChangedChainEvent&);
-DECLARE_EVENT_ThreeParams(UControlRigBlueprint, FOnLocalizeFunctionDialogRequested, URigVMLibraryNode*, UControlRigBlueprint*, bool);
+DECLARE_EVENT_ThreeParams(UControlRigBlueprint, FOnLocalizeFunctionDialogRequested, FRigVMGraphFunctionIdentifier&, UControlRigBlueprint*, bool);
 DECLARE_EVENT_ThreeParams(UControlRigBlueprint, FOnReportCompilerMessage, EMessageSeverity::Type, UObject*, const FString&);
 DECLARE_DELEGATE_RetVal_FourParams(FRigVMController_BulkEditResult, FControlRigOnBulkEditDialogRequestedDelegate, UControlRigBlueprint*, URigVMController*, URigVMLibraryNode*, ERigVMControllerBulkEditType);
 DECLARE_DELEGATE_RetVal_OneParam(bool, FControlRigOnBreakLinksDialogRequestedDelegate, TArray<URigVMLink*>);
@@ -699,7 +699,7 @@ public:
 	void BroadcastPostEditChangeChainProperty(FPropertyChangedChainEvent& PropertyChangedChainEvent);
 
 	FOnLocalizeFunctionDialogRequested& OnRequestLocalizeFunctionDialog() { return RequestLocalizeFunctionDialog; }
-	void BroadcastRequestLocalizeFunctionDialog(URigVMLibraryNode* InFunction, bool bForce = false);
+	void BroadcastRequestLocalizeFunctionDialog(FRigVMGraphFunctionIdentifier InFunction, bool bForce = false);
 
 	FControlRigOnBulkEditDialogRequestedDelegate& OnRequestBulkEditDialog() { return RequestBulkEditDialog; }
 
diff --git a/Engine/Plugins/Animation/ControlRig/Source/ControlRigEditor/Private/ControlRigEditorModule.cpp b/Engine/Plugins/Animation/ControlRig/Source/ControlRigEditor/Private/ControlRigEditorModule.cpp
--- a/Engine/Plugins/Animation/ControlRig/Source/ControlRigEditor/Private/ControlRigEditorModule.cpp
+++ b/Engine/Plugins/Animation/ControlRig/Source/ControlRigEditor/Private/ControlRigEditorModule.cpp
@@ -2405,18 +2405,16 @@ void FControlRigEditorModule::GetContextMenuActions(const UControlRigGraphSchema
 						TSoftObjectPtr<URigVMFunctionReferenceNode> RefPtr(FunctionReferenceNode);
 						if(RefPtr.GetLongPackageName() != FunctionReferenceNode->GetReferencedFunctionHeader().LibraryPointer.LibraryNode.GetLongPackageName())
 						{
-							if (URigVMLibraryNode* LibraryNode = FunctionReferenceNode->LoadReferencedNode())
-							{
-								OrganizationSection.AddMenuEntry(
-								   "Localize Function",
-								   LOCTEXT("LocalizeFunction", "Localize Function"),
-								   LOCTEXT("LocalizeFunction_Tooltip", "Creates a local copy of the function backing the node."),
-								   FSlateIcon(),
-								   FUIAction(FExecuteAction::CreateLambda([RigBlueprint, LibraryNode]() {
-									   RigBlueprint->BroadcastRequestLocalizeFunctionDialog(LibraryNode, true);
-								   })
-							   ));
-							}
+							OrganizationSection.AddMenuEntry(
+							   "Localize Function",
+							   LOCTEXT("LocalizeFunction", "Localize Function"),
+							   LOCTEXT("LocalizeFunction_Tooltip", "Creates a local copy of the function backing the node."),
+							   FSlateIcon(),
+							   FUIAction(FExecuteAction::CreateLambda([RigBlueprint, FunctionReferenceNode]() {
+								   RigBlueprint->BroadcastRequestLocalizeFunctionDialog(FunctionReferenceNode->GetFunctionIdentifier(), true);
+							   })
+						    ));
+							
 
 							if(!FunctionReferenceNode->IsFullyRemapped() && FunctionReferenceNode->GetReferencedFunctionHeader().LibraryPointer.LibraryNode.ResolveObject())
 							{
diff --git a/Engine/Plugins/Animation/ControlRig/Source/ControlRigEditor/Private/Editor/ControlRigEditor.cpp b/Engine/Plugins/Animation/ControlRig/Source/ControlRigEditor/Private/Editor/ControlRigEditor.cpp
--- a/Engine/Plugins/Animation/ControlRig/Source/ControlRigEditor/Private/Editor/ControlRigEditor.cpp
+++ b/Engine/Plugins/Animation/ControlRig/Source/ControlRigEditor/Private/Editor/ControlRigEditor.cpp
@@ -2725,17 +2725,13 @@ void FControlRigEditor::PasteNodes()
 
 	TGuardValue<FRigVMController_RequestLocalizeFunctionDelegate> RequestLocalizeDelegateGuard(
 		GetFocusedController()->RequestLocalizeFunctionDelegate,
-		FRigVMController_RequestLocalizeFunctionDelegate::CreateLambda([this](TSoftObjectPtr<URigVMLibraryNode> InFunctionToLocalize)
+		FRigVMController_RequestLocalizeFunctionDelegate::CreateLambda([this](FRigVMGraphFunctionIdentifier& InFunctionToLocalize)
 		{
-			if (InFunctionToLocalize.IsValid())
-			{
-				URigVMLibraryNode* LibraryNode = InFunctionToLocalize.Get();
-				OnRequestLocalizeFunctionDialog(LibraryNode, GetControlRigBlueprint(), true);
+			OnRequestLocalizeFunctionDialog(InFunctionToLocalize, GetControlRigBlueprint(), true);
 
-			   const URigVMLibraryNode* LocalizedFunctionNode = GetControlRigBlueprint()->GetLocalFunctionLibrary()->FindPreviouslyLocalizedFunction(LibraryNode);
-			   return LocalizedFunctionNode != nullptr;
-			}
-			return false;
+		   const URigVMLibraryNode* LocalizedFunctionNode = GetControlRigBlueprint()->GetLocalFunctionLibrary()->FindPreviouslyLocalizedFunction(InFunctionToLocalize);
+		   return LocalizedFunctionNode != nullptr;
+		
 		})
 	);
 	
@@ -5046,9 +5042,8 @@ void FControlRigEditor::OnWrappedPropertyChangedChainEvent(UDetailsViewWrapperOb
 	}
 }
 
-void FControlRigEditor::OnRequestLocalizeFunctionDialog(URigVMLibraryNode* InFunction, UControlRigBlueprint* InTargetBlueprint, bool bForce)
+void FControlRigEditor::OnRequestLocalizeFunctionDialog(FRigVMGraphFunctionIdentifier& InFunction, UControlRigBlueprint* InTargetBlueprint, bool bForce)
 {
-	check(InFunction);
 	check(InTargetBlueprint);
 
 	if(InTargetBlueprint != GetControlRigBlueprint())
@@ -5058,23 +5053,18 @@ void FControlRigEditor::OnRequestLocalizeFunctionDialog(URigVMLibraryNode* InFun
 	
 	if(URigVMController* TargetController = InTargetBlueprint->GetController(InTargetBlueprint->GetDefaultModel()))
 	{
-		if(URigVMFunctionLibrary* FunctionLibrary = Cast<URigVMFunctionLibrary>(InFunction->GetOuter()))
+		bool bIsPublic;
+		if (FRigVMGraphFunctionData::FindFunctionData(InFunction, &bIsPublic))
 		{
-			if(UControlRigBlueprint* FunctionRigBlueprint = Cast<UControlRigBlueprint>(FunctionLibrary->GetOuter()))
+			if (bForce || bIsPublic)
 			{
-				if(FunctionRigBlueprint != InTargetBlueprint)
-				{
-					if(bForce || !FunctionRigBlueprint->IsFunctionPublic(InFunction->GetFName()))
-					{
-                        TSharedRef<SControlRigFunctionLocalizationDialog> LocalizationDialog = SNew(SControlRigFunctionLocalizationDialog)
-                        .Function(InFunction)
-                        .TargetBlueprint(InTargetBlueprint);
+				TSharedRef<SControlRigFunctionLocalizationDialog> LocalizationDialog = SNew(SControlRigFunctionLocalizationDialog)
+							.Function(InFunction)
+							.TargetBlueprint(InTargetBlueprint);
 
-						if (LocalizationDialog->ShowModal() != EAppReturnType::Cancel)
-						{
-							TargetController->LocalizeFunctions(LocalizationDialog->GetFunctionsToLocalize(), true, true, true);
-						}
-					}
+				if (LocalizationDialog->ShowModal() != EAppReturnType::Cancel)
+				{
+					TargetController->LocalizeFunctions(LocalizationDialog->GetFunctionsToLocalize(), true, true, true);
 				}
 			}
 		}
diff --git a/Engine/Plugins/Animation/ControlRig/Source/ControlRigEditor/Private/Editor/ControlRigEditor.h b/Engine/Plugins/Animation/ControlRig/Source/ControlRigEditor/Private/Editor/ControlRigEditor.h
--- a/Engine/Plugins/Animation/ControlRig/Source/ControlRigEditor/Private/Editor/ControlRigEditor.h
+++ b/Engine/Plugins/Animation/ControlRig/Source/ControlRigEditor/Private/Editor/ControlRigEditor.h
@@ -430,7 +430,7 @@ protected:
 	virtual void OnFinishedChangingProperties(const FPropertyChangedEvent& PropertyChangedEvent) override;
 	void OnPropertyChanged(UObject* InObject, FPropertyChangedEvent& InEvent);
 	void OnWrappedPropertyChangedChainEvent(UDetailsViewWrapperObject* InWrapperObject, const FString& InPropertyPath, FPropertyChangedChainEvent& InPropertyChangedChainEvent);
-	void OnRequestLocalizeFunctionDialog(URigVMLibraryNode* InFunction, UControlRigBlueprint* InTargetBlueprint, bool bForce);
+	void OnRequestLocalizeFunctionDialog(FRigVMGraphFunctionIdentifier& InFunction, UControlRigBlueprint* InTargetBlueprint, bool bForce);
 	FRigVMController_BulkEditResult OnRequestBulkEditDialog(UControlRigBlueprint* InBlueprint, URigVMController* InController, URigVMLibraryNode* InFunction, ERigVMControllerBulkEditType InEditType);
 	bool OnRequestBreakLinksDialog(TArray<URigVMLink*> InLinks);
 	void HandleJumpToHyperlink(const UObject* InSubject);
diff --git a/Engine/Plugins/Animation/ControlRig/Source/ControlRigEditor/Private/Editor/SControlRigFunctionLocalizationWidget.cpp b/Engine/Plugins/Animation/ControlRig/Source/ControlRigEditor/Private/Editor/SControlRigFunctionLocalizationWidget.cpp
--- a/Engine/Plugins/Animation/ControlRig/Source/ControlRigEditor/Private/Editor/SControlRigFunctionLocalizationWidget.cpp
+++ b/Engine/Plugins/Animation/ControlRig/Source/ControlRigEditor/Private/Editor/SControlRigFunctionLocalizationWidget.cpp
@@ -18,24 +18,22 @@
 /// SControlRigFunctionLocalizationItem
 ///////////////////////////////////////////////////////////
 
-SControlRigFunctionLocalizationItem::SControlRigFunctionLocalizationItem(URigVMLibraryNode* InFunction)
+SControlRigFunctionLocalizationItem::SControlRigFunctionLocalizationItem(const FRigVMGraphFunctionIdentifier& InFunction)
     : Function(InFunction)
 {
 	FString OuterName;
+	FString FunctionName;
 	bool bIsPublic = false;
 
-	if(URigVMFunctionLibrary* Library = Cast<URigVMFunctionLibrary>(InFunction->GetOuter()))
+	if (FRigVMGraphFunctionData* FunctionData = FRigVMGraphFunctionData::FindFunctionData(InFunction, &bIsPublic))
 	{
-		if(UControlRigBlueprint* Blueprint = Cast<UControlRigBlueprint>(Library->GetOuter()))
-		{
-			OuterName = Blueprint->GetName();
-			bIsPublic = Blueprint->IsFunctionPublic(InFunction->GetFName());
-		}
+		OuterName = InFunction.LibraryNode.GetAssetPathString();
+		FunctionName = FunctionData->Header.Name.ToString();
+		DisplayText = FText::FromString(FString::Printf(TEXT("%s :: %s"), *OuterName, *FunctionName));
+		ToolTipText = bIsPublic ?
+			FText::FromString(FString::Printf(TEXT("%s :: %s is public. Check this to create a local copy."), *OuterName, *FunctionName)) :
+			FText::FromString(FString::Printf(TEXT("%s :: %s is private. A local copy has to be created. To avoid this you can turn it public in the source Control Rig."), *OuterName, *FunctionName));
 	}
-	DisplayText = FText::FromString(FString::Printf(TEXT("%s :: %s"), *OuterName, *InFunction->GetName()));
-	ToolTipText = bIsPublic ?
-		FText::FromString(FString::Printf(TEXT("%s :: %s is public. Check this to create a local copy."), *OuterName, *InFunction->GetName())) :
-		FText::FromString(FString::Printf(TEXT("%s :: %s is private. A local copy has to be created. To avoid this you can turn it public in the source Control Rig."), *OuterName, *InFunction->GetName()));
 }
 
 //////////////////////////////////////////////////////////////
@@ -88,74 +86,58 @@ void SControlRigFunctionLocalizationTableRow::Construct(const FArguments& InArgs
 /// SControlRigFunctionLocalizationWidget
 ///////////////////////////////////////////////////////////
 
-void SControlRigFunctionLocalizationWidget::Construct(const FArguments& InArgs, URigVMLibraryNode* InFunctionToLocalize, UControlRigBlueprint* InTargetBlueprint)
+void SControlRigFunctionLocalizationWidget::Construct(const FArguments& InArgs, const FRigVMGraphFunctionIdentifier& InFunctionToLocalize, UControlRigBlueprint* InTargetBlueprint)
 {
 	FunctionsToLocalize.Add(InFunctionToLocalize);
 	FunctionItems.Reset();
 
-	TArray<URigVMLibraryNode*> NodesToVisit;
-	TArray<URigVMLibraryNode*> FunctionsForTable;
+	TArray<FRigVMGraphFunctionIdentifier> NodesToVisit;
+	TArray<FRigVMGraphFunctionData*> FunctionsForTable;
 
-	FunctionsForTable.Add(InFunctionToLocalize);
-	NodesToVisit.Add(InFunctionToLocalize);
+	{
+		FRigVMGraphFunctionData* FunctionData = FRigVMGraphFunctionData::FindFunctionData(InFunctionToLocalize);
+		FunctionsForTable.Add(FunctionData);
+		NodesToVisit.Add(InFunctionToLocalize);
+	}
 
+	IRigVMGraphFunctionHost* TargetFunctionHost = InTargetBlueprint->GetRigVMGraphFunctionHost();
 	for(int32 NodeToVisitIndex=0; NodeToVisitIndex<NodesToVisit.Num(); NodeToVisitIndex++)
 	{
-		URigVMLibraryNode* NodeToVisit = NodesToVisit[NodeToVisitIndex];
-		TArray<URigVMNode*> ContainedNodes = NodeToVisit->GetContainedNodes();
-		if (URigVMFunctionReferenceNode* ReferenceNode = Cast<URigVMFunctionReferenceNode>(NodeToVisit))
+		const FRigVMGraphFunctionIdentifier& NodeToVisit = NodesToVisit[NodeToVisitIndex];
+
+		bool bIsPublic;
+		FRigVMGraphFunctionData* FunctionData = FRigVMGraphFunctionData::FindFunctionData(NodeToVisit, &bIsPublic);
+		if (!FunctionData)
 		{
-			if (ReferenceNode->GetReferencedFunctionHeader().LibraryPointer.LibraryNode.ResolveObject())
-			{
-				URigVMLibraryNode* LibraryNode = Cast<URigVMLibraryNode>(ReferenceNode->GetReferencedFunctionHeader().LibraryPointer.LibraryNode.ResolveObject());
-				ContainedNodes = LibraryNode->GetContainedNodes();
-			}
+			continue;
 		}
-		for(URigVMNode* ContainedNode : ContainedNodes)
+
+		for (TPair<FRigVMGraphFunctionIdentifier, uint32>& Pair : FunctionData->Header.Dependencies)
 		{
-			if(URigVMLibraryNode* ContainedLibraryNode = Cast<URigVMLibraryNode>(ContainedNode))
-			{
-				NodesToVisit.AddUnique(ContainedLibraryNode);
-			}
+			NodesToVisit.AddUnique(Pair.Key);
 		}
-
-		if(URigVMFunctionReferenceNode* FunctionReferenceNode = Cast<URigVMFunctionReferenceNode>(NodeToVisit))
+		
+		if (NodeToVisit.HostObject != Cast<UObject>(TargetFunctionHost))
 		{
-			FSoftObjectPath ReferencedNodePath = FunctionReferenceNode->GetReferencedFunctionHeader().LibraryPointer.LibraryNode;
-			if(ReferencedNodePath.ResolveObject())
+			FunctionsForTable.AddUnique(FunctionData);
+			if (bIsPublic)
 			{
-				URigVMLibraryNode* ReferencedNode = Cast<URigVMLibraryNode>(ReferencedNodePath.ResolveObject());
-				if(URigVMFunctionLibrary* ReferencedLibrary = ReferencedNode->GetLibrary())
-				{
-					if(UControlRigBlueprint* ReferencedBlueprint = Cast<UControlRigBlueprint>(ReferencedLibrary->GetOuter()))
-					{
-						if(ReferencedBlueprint != InTargetBlueprint)
-						{
-							FunctionsForTable.AddUnique(ReferencedNode);
-							if(!ReferencedBlueprint->IsFunctionPublic(ReferencedNode->GetFName()))
-							{
-								FunctionsToLocalize.AddUnique(ReferencedNode);
-							}
-						}
-					}
-				}
-
-				NodesToVisit.AddUnique(ReferencedNode);
+				FunctionsToLocalize.AddUnique(NodeToVisit);
 			}
 		}
-	}
+	} 
 
 	// sort the functions to localize based on their nesting
-	Algo::Sort(FunctionsForTable, [](URigVMLibraryNode* A, URigVMLibraryNode* B) -> bool
+	Algo::Sort(FunctionsForTable, [](FRigVMGraphFunctionData* A, FRigVMGraphFunctionData* B) -> bool
     {
-        check(A);
-        check(B);
-        return A->Contains(B);
+		check(A);
+		check(B);
+		return A->Header.Dependencies.Contains(B->Header.LibraryPointer);
     });
 
-	for (URigVMLibraryNode* FunctionForTable : FunctionsForTable)
+	for (FRigVMGraphFunctionData* FunctionForTable : FunctionsForTable)
 	{
-		FunctionItems.Add(MakeShared<SControlRigFunctionLocalizationItem>(FunctionForTable));
+		FunctionItems.Add(MakeShared<SControlRigFunctionLocalizationItem>(FunctionForTable->Header.LibraryPointer));
 	}
 
 	ChildSlot
@@ -186,12 +168,12 @@ TSharedRef<ITableRow> SControlRigFunctionLocalizationWidget::GenerateFunctionLis
 	return TableRow;
 }
 
-ECheckBoxState SControlRigFunctionLocalizationWidget::IsFunctionEnabled(URigVMLibraryNode* InFunction) const
+ECheckBoxState SControlRigFunctionLocalizationWidget::IsFunctionEnabled(const FRigVMGraphFunctionIdentifier InFunction) const
 {
 	return FunctionsToLocalize.Contains(InFunction) ? ECheckBoxState::Checked : ECheckBoxState::Unchecked;
 }
 
-void SControlRigFunctionLocalizationWidget::SetFunctionEnabled(ECheckBoxState NewState, URigVMLibraryNode* InFunction)
+void SControlRigFunctionLocalizationWidget::SetFunctionEnabled(ECheckBoxState NewState, const FRigVMGraphFunctionIdentifier InFunction)
 {
 	if(NewState == ECheckBoxState::Checked)
 	{
@@ -203,16 +185,19 @@ void SControlRigFunctionLocalizationWidget::SetFunctionEnabled(ECheckBoxState Ne
 	}
 }
 
-bool SControlRigFunctionLocalizationWidget::IsFunctionPublic(URigVMLibraryNode* InFunction) const
+bool SControlRigFunctionLocalizationWidget::IsFunctionPublic(const FRigVMGraphFunctionIdentifier InFunction) const
 {
-	if(URigVMFunctionLibrary* FunctionLibrary = Cast<URigVMFunctionLibrary>(InFunction->GetOuter()))
+	IRigVMGraphFunctionHost* FunctionHost = nullptr;
+	if (UObject* FunctionHostObj = InFunction.HostObject.TryLoad())
 	{
-		if(UControlRigBlueprint* Blueprint = Cast<UControlRigBlueprint>(FunctionLibrary->GetOuter()))
-		{
-			return Blueprint->IsFunctionPublic(InFunction->GetFName());
-		}
+		FunctionHost = Cast<IRigVMGraphFunctionHost>(FunctionHostObj);									
+	}
+	if (FunctionHost)
+	{
+		return FunctionHost->GetRigVMGraphFunctionStore()->IsFunctionPublic(InFunction);
 	}
-	return true;	
+			
+	return false;	
 }
 
 void SControlRigFunctionLocalizationDialog::Construct(const FArguments& InArgs)
@@ -313,7 +298,7 @@ EAppReturnType::Type SControlRigFunctionLocalizationDialog::ShowModal()
 	return UserResponse;
 }
 
-TArray<URigVMLibraryNode*>& SControlRigFunctionLocalizationDialog::GetFunctionsToLocalize()
+TArray<FRigVMGraphFunctionIdentifier>& SControlRigFunctionLocalizationDialog::GetFunctionsToLocalize()
 {
 	return FunctionsWidget->FunctionsToLocalize;
 }
diff --git a/Engine/Plugins/Animation/ControlRig/Source/ControlRigEditor/Private/Editor/SControlRigFunctionLocalizationWidget.h b/Engine/Plugins/Animation/ControlRig/Source/ControlRigEditor/Private/Editor/SControlRigFunctionLocalizationWidget.h
--- a/Engine/Plugins/Animation/ControlRig/Source/ControlRigEditor/Private/Editor/SControlRigFunctionLocalizationWidget.h
+++ b/Engine/Plugins/Animation/ControlRig/Source/ControlRigEditor/Private/Editor/SControlRigFunctionLocalizationWidget.h
@@ -13,11 +13,11 @@ class SControlRigFunctionLocalizationWidget;
 class SControlRigFunctionLocalizationItem : public TSharedFromThis<SControlRigFunctionLocalizationItem>
 {
 public:
-	SControlRigFunctionLocalizationItem(URigVMLibraryNode* InFunction);
+	SControlRigFunctionLocalizationItem(const FRigVMGraphFunctionIdentifier& InFunction);
 
 	FText DisplayText;
 	FText ToolTipText;
-	URigVMLibraryNode* Function;
+	const FRigVMGraphFunctionIdentifier Function;
 };
 
 class SControlRigFunctionLocalizationTableRow : public STableRow<TSharedPtr<SControlRigFunctionLocalizationItem>>
@@ -34,18 +34,18 @@ public:
 	SLATE_BEGIN_ARGS(SControlRigFunctionLocalizationWidget) {}
 	SLATE_END_ARGS()
 
-	void Construct(const FArguments& InArgs, URigVMLibraryNode* InFunctionToLocalize, UControlRigBlueprint* InTargetBlueprint);
+	void Construct(const FArguments& InArgs, const FRigVMGraphFunctionIdentifier& InFunctionToLocalize, UControlRigBlueprint* InTargetBlueprint);
 
 	TSharedRef<ITableRow> GenerateFunctionListRow(TSharedPtr<SControlRigFunctionLocalizationItem> InItem, const TSharedRef<STableViewBase>& InOwningTable);
-	ECheckBoxState IsFunctionEnabled(URigVMLibraryNode* InFunction) const;
-	void SetFunctionEnabled(ECheckBoxState NewState, URigVMLibraryNode* InFunction);
-	bool IsFunctionPublic(URigVMLibraryNode* InFunction) const;
+	ECheckBoxState IsFunctionEnabled(const FRigVMGraphFunctionIdentifier InFunction) const;
+	void SetFunctionEnabled(ECheckBoxState NewState, const FRigVMGraphFunctionIdentifier InFunction);
+	bool IsFunctionPublic(const FRigVMGraphFunctionIdentifier InFunction) const;
 
 private:
 
-	TArray<URigVMLibraryNode*> FunctionsToLocalize;
+	TArray<FRigVMGraphFunctionIdentifier> FunctionsToLocalize;
 	TArray<TSharedPtr<SControlRigFunctionLocalizationItem>> FunctionItems;
-	TMap<URigVMLibraryNode*, TSharedRef<SControlRigFunctionLocalizationTableRow>> TableRows;
+	TMap<FRigVMGraphFunctionIdentifier, TSharedRef<SControlRigFunctionLocalizationTableRow>> TableRows;
 
 	friend class SControlRigFunctionLocalizationTableRow;
 	friend class SControlRigFunctionLocalizationDialog;
@@ -59,14 +59,14 @@ public:
 	{
 	}
 
-	SLATE_ARGUMENT(URigVMLibraryNode*, Function)
+	SLATE_ARGUMENT(FRigVMGraphFunctionIdentifier, Function)
 	SLATE_ARGUMENT(UControlRigBlueprint*, TargetBlueprint)
 	SLATE_END_ARGS()
 
 	void Construct(const FArguments& InArgs);
 
 	EAppReturnType::Type ShowModal();
-	TArray<URigVMLibraryNode*>& GetFunctionsToLocalize();
+	TArray<FRigVMGraphFunctionIdentifier>& GetFunctionsToLocalize();
 
 protected:
 
diff --git a/Engine/Source/Developer/RigVMDeveloper/Private/RigVMModel/Nodes/RigVMFunctionReferenceNode.cpp b/Engine/Source/Developer/RigVMDeveloper/Private/RigVMModel/Nodes/RigVMFunctionReferenceNode.cpp
--- a/Engine/Source/Developer/RigVMDeveloper/Private/RigVMModel/Nodes/RigVMFunctionReferenceNode.cpp
+++ b/Engine/Source/Developer/RigVMDeveloper/Private/RigVMModel/Nodes/RigVMFunctionReferenceNode.cpp
@@ -152,6 +152,11 @@ FText URigVMFunctionReferenceNode::GetToolTipTextForPin(const URigVMPin* InPin)
 	return Super::GetToolTipTextForPin(InPin);
 }
 
+FRigVMGraphFunctionIdentifier URigVMFunctionReferenceNode::GetFunctionIdentifier() const
+{
+	return GetReferencedFunctionHeader().LibraryPointer;
+}
+
 bool URigVMFunctionReferenceNode::IsReferencedFunctionHostLoaded() const
 {
 	return ReferencedFunctionHeader.LibraryPointer.HostObject.ResolveObject() != nullptr;
diff --git a/Engine/Source/Developer/RigVMDeveloper/Private/RigVMModel/RigVMClient.cpp b/Engine/Source/Developer/RigVMDeveloper/Private/RigVMModel/RigVMClient.cpp
--- a/Engine/Source/Developer/RigVMDeveloper/Private/RigVMModel/RigVMClient.cpp
+++ b/Engine/Source/Developer/RigVMDeveloper/Private/RigVMModel/RigVMClient.cpp
@@ -2,6 +2,8 @@
 
 #include "RigVMModel/RigVMClient.h"
 #include "Misc/TransactionObjectEvent.h"
+#include "Exporters/Exporter.h"
+#include "UnrealExporter.h"
 
 #include UE_INLINE_GENERATED_CPP_BY_NAME(RigVMClient)
 
@@ -707,12 +709,14 @@ void FRigVMClient::HandleGraphModifiedEvent(ERigVMGraphNotifType InNotifType, UR
 					if (GetOuter()->Implements<URigVMClientHost>())
 					{
 						FunctionStore->AddFunction(CollapseNode->GetFunctionHeader(FunctionHost), false);
+						UpdateGraphFunctionSerializedGraph(CollapseNode);
 					}
 				}
 			}
 			// A node was added into the contained graph of a function
 			else if(URigVMLibraryNode* LibraryNode = Cast<URigVMNode>(InSubject)->FindFunctionForNode())
 			{
+				UpdateGraphFunctionSerializedGraph(LibraryNode);
 				DirtyGraphFunctionCompilationData(LibraryNode);
 				if (URigVMFunctionReferenceNode* FunctionReference = Cast<URigVMFunctionReferenceNode>(InSubject))
 				{
@@ -737,6 +741,7 @@ void FRigVMClient::HandleGraphModifiedEvent(ERigVMGraphNotifType InNotifType, UR
 			// A node was added into the contained graph of a function
 			else if(URigVMLibraryNode* LibraryNode = Cast<URigVMNode>(InSubject)->FindFunctionForNode())
 			{
+				UpdateGraphFunctionSerializedGraph(LibraryNode);
 				DirtyGraphFunctionCompilationData(LibraryNode);
 				if (URigVMFunctionReferenceNode* FunctionReference = Cast<URigVMFunctionReferenceNode>(InSubject))
 				{
@@ -753,6 +758,7 @@ void FRigVMClient::HandleGraphModifiedEvent(ERigVMGraphNotifType InNotifType, UR
 		{
 			if(URigVMLibraryNode* LibraryNode = Cast<URigVMNode>(InSubject)->FindFunctionForNode())
 			{
+				UpdateGraphFunctionSerializedGraph(LibraryNode);
 				DirtyGraphFunctionCompilationData(LibraryNode);
 				UpdateExternalVariablesForFunction(LibraryNode);				
 			}
@@ -776,6 +782,10 @@ void FRigVMClient::HandleGraphModifiedEvent(ERigVMGraphNotifType InNotifType, UR
 					UpdateGraphFunctionData(CollapseNode);
 				}
 			}
+			else if(URigVMLibraryNode* LibraryNode = Cast<URigVMNode>(InSubject)->FindFunctionForNode())
+			{
+				UpdateGraphFunctionSerializedGraph(LibraryNode);
+			}
 			break;	
 		}
 		case ERigVMGraphNotifType::NodeColorChanged: // A node's color has changed (Subject == URigVMNode)
@@ -790,13 +800,19 @@ void FRigVMClient::HandleGraphModifiedEvent(ERigVMGraphNotifType InNotifType, UR
 					UpdateGraphFunctionData(CollapseNode);
 				}
 			}
+			else if(URigVMLibraryNode* LibraryNode = Cast<URigVMNode>(InSubject)->FindFunctionForNode())
+			{
+				UpdateGraphFunctionSerializedGraph(LibraryNode);
+			}
 			break;
 		}
 		case ERigVMGraphNotifType::NodeReferenceChanged: // A node has changed it's referenced function (Subject == URigVMFunctionReferenceNode)
 		case ERigVMGraphNotifType::LibraryTemplateChanged: // The definition of a library node's template has changed (Subject == URigVMLibraryNode)
 		{
-			
-			
+			if(URigVMLibraryNode* LibraryNode = Cast<URigVMNode>(InSubject)->FindFunctionForNode())
+			{
+				UpdateGraphFunctionSerializedGraph(LibraryNode);
+			}
 			break;	
 		}
 		
@@ -813,6 +829,7 @@ void FRigVMClient::HandleGraphModifiedEvent(ERigVMGraphNotifType InNotifType, UR
 			{
 				if (URigVMLibraryNode* LibraryNode = Node->FindFunctionForNode())
 				{
+					UpdateGraphFunctionSerializedGraph(LibraryNode);
 					DirtyGraphFunctionCompilationData(LibraryNode);
 				}
 				if(Node->GetOuter()->IsA<URigVMFunctionLibrary>())
@@ -835,6 +852,7 @@ void FRigVMClient::HandleGraphModifiedEvent(ERigVMGraphNotifType InNotifType, UR
 				{
 					if (URigVMLibraryNode* LibraryNode = OuterNode->FindFunctionForNode())
 					{
+						UpdateGraphFunctionSerializedGraph(LibraryNode);
 						DirtyGraphFunctionCompilationData(LibraryNode);
 					}
 				}
@@ -993,6 +1011,23 @@ bool FRigVMClient::DirtyGraphFunctionCompilationData(URigVMLibraryNode* InLibrar
 	return false;
 }
 
+bool FRigVMClient::UpdateGraphFunctionSerializedGraph(URigVMLibraryNode* InLibraryNode)
+{
+	if (FRigVMGraphFunctionStore* Store = FindFunctionStore(InLibraryNode))
+	{
+		FRigVMGraphFunctionIdentifier Identifier = InLibraryNode->GetFunctionIdentifier();
+		if (FRigVMGraphFunctionData* Data = Store->FindFunction(Identifier))
+		{
+			FStringOutputDevice Archive;
+			const FExportObjectInnerContext Context;
+			UExporter::ExportToOutputDevice(&Context, InLibraryNode, NULL, Archive, TEXT("copy"), 0, PPF_ExportsNotFullyQualified | PPF_Copy | PPF_Delimited, false, nullptr);
+			Data->SerializedCollapsedNode = Archive;
+			return true;
+		}
+	}
+	return false;
+}
+
 bool FRigVMClient::IsFunctionPublic(URigVMLibraryNode* InLibraryNode)
 {
 	if (FRigVMGraphFunctionStore* Store = FindFunctionStore(InLibraryNode))
diff --git a/Engine/Source/Developer/RigVMDeveloper/Private/RigVMModel/RigVMController.cpp b/Engine/Source/Developer/RigVMDeveloper/Private/RigVMModel/RigVMController.cpp
--- a/Engine/Source/Developer/RigVMDeveloper/Private/RigVMModel/RigVMController.cpp
+++ b/Engine/Source/Developer/RigVMDeveloper/Private/RigVMModel/RigVMController.cpp
@@ -3771,8 +3771,53 @@ TArray<FName> URigVMController::ImportNodesFromText(const FString& InText, bool
 	return NodeNames;
 }
 
+URigVMLibraryNode* URigVMController::LocalizeFunctionFromPath(const FString& InHostPath, const FName& InFunctionName, bool bLocalizeDependentPrivateFunctions, bool bSetupUndoRedo, bool bPrintPythonCommand)
+{
+	if (!IsValidGraph())
+	{
+		return nullptr;
+	}
+
+	if (!bIsTransacting && !IsGraphEditable())
+	{
+		return nullptr;
+	}
+
+	URigVMGraph* Graph = GetGraph();
+	check(Graph);
+
+	if (Graph->IsA<URigVMFunctionLibrary>())
+	{
+		ReportError(TEXT("Cannot add function reference nodes to function library graphs."));
+		return nullptr;
+	}
+
+	UObject* HostObject = StaticLoadObject(UObject::StaticClass(), NULL, *InHostPath, NULL, LOAD_None, NULL);
+	if (!HostObject)
+	{
+		ReportErrorf(TEXT("Failed to load the Host object %s."), *InHostPath);
+		return nullptr;
+	}
+
+	IRigVMGraphFunctionHost* FunctionHost = Cast<IRigVMGraphFunctionHost>(HostObject);
+	if (!FunctionHost)
+	{
+		ReportError(TEXT("Host object is not a IRigVMGraphFunctionHost."));
+		return nullptr;
+	}
+
+	FRigVMGraphFunctionData* Data = FunctionHost->GetRigVMGraphFunctionStore()->FindFunctionByName(InFunctionName);
+	if (!Data)
+	{
+		ReportErrorf(TEXT("Function %s not found in host %s."), *InFunctionName.ToString(), *InHostPath);
+		return nullptr;
+	}
+
+	return LocalizeFunction(Data->Header.LibraryPointer, bLocalizeDependentPrivateFunctions, bSetupUndoRedo, bPrintPythonCommand);
+}
+
 URigVMLibraryNode* URigVMController::LocalizeFunction(
-	URigVMLibraryNode* InFunctionDefinition,
+	const FRigVMGraphFunctionIdentifier& InFunctionDefinition,
 	bool bLocalizeDependentPrivateFunctions,
 	bool bSetupUndoRedo,
 	bool bPrintPythonCommand)
@@ -3787,15 +3832,10 @@ URigVMLibraryNode* URigVMController::LocalizeFunction(
 		return nullptr;
 	}
 
-	if(InFunctionDefinition == nullptr)
-	{
-		return nullptr;
-	}
-
-	TArray<URigVMLibraryNode*> FunctionsToLocalize;
+	TArray<FRigVMGraphFunctionIdentifier> FunctionsToLocalize;
 	FunctionsToLocalize.Add(InFunctionDefinition);
 
-	TMap<URigVMLibraryNode*, URigVMLibraryNode*> Results = LocalizeFunctions(FunctionsToLocalize, bLocalizeDependentPrivateFunctions, bSetupUndoRedo, bPrintPythonCommand);
+	TMap<FRigVMGraphFunctionIdentifier, URigVMLibraryNode*> Results = LocalizeFunctions(FunctionsToLocalize, bLocalizeDependentPrivateFunctions, bSetupUndoRedo, bPrintPythonCommand);
 
 	URigVMLibraryNode** LocalizedFunctionPtr = Results.Find(FunctionsToLocalize[0]);
 	if(LocalizedFunctionPtr)
@@ -3805,12 +3845,12 @@ URigVMLibraryNode* URigVMController::LocalizeFunction(
 	return nullptr;
 }
 
-TMap<URigVMLibraryNode*, URigVMLibraryNode*> URigVMController::LocalizeFunctions(
-	TArray<URigVMLibraryNode*> InFunctionDefinitions,
+TMap<FRigVMGraphFunctionIdentifier, URigVMLibraryNode*> URigVMController::LocalizeFunctions(
+	TArray<FRigVMGraphFunctionIdentifier> InFunctionDefinitions,
 	bool bLocalizeDependentPrivateFunctions,
 	bool bSetupUndoRedo, bool bPrintPythonCommand)
 {
-	TMap<URigVMLibraryNode*, URigVMLibraryNode*> LocalizedFunctions;
+	TMap<FRigVMGraphFunctionIdentifier, URigVMLibraryNode*> LocalizedFunctions;
 
 	if(!IsValidGraph())
 	{
@@ -3831,88 +3871,65 @@ TMap<URigVMLibraryNode*, URigVMLibraryNode*> URigVMController::LocalizeFunctions
 		return LocalizedFunctions;
 	}
 
-	TArray<URigVMLibraryNode*> FunctionsToLocalize;
+	TArray<FRigVMGraphFunctionData*> FunctionsToLocalize;
 
-	TArray<URigVMLibraryNode*> NodesToVisit;
-	for(URigVMLibraryNode* FunctionDefinition : InFunctionDefinitions)
+	TArray<FRigVMGraphFunctionIdentifier> NodesToVisit;
+	for(const FRigVMGraphFunctionIdentifier& FunctionDefinition : InFunctionDefinitions)
 	{
 		NodesToVisit.AddUnique(FunctionDefinition);
-		FunctionsToLocalize.AddUnique(FunctionDefinition);
+		FunctionsToLocalize.AddUnique(FRigVMGraphFunctionData::FindFunctionData(FunctionDefinition));
 	}
 
+	const FSoftObjectPath ThisFunctionHost = ThisLibrary->GetFunctionHostObjectPath();
+	
 	// find all functions to localize
 	for(int32 NodeToVisitIndex=0; NodeToVisitIndex<NodesToVisit.Num(); NodeToVisitIndex++)
 	{
-		URigVMLibraryNode* NodeToVisit = NodesToVisit[NodeToVisitIndex];
+		FRigVMGraphFunctionIdentifier& NodeToVisit = NodesToVisit[NodeToVisitIndex];
 
-		if(URigVMCollapseNode* CollapseNode = Cast<URigVMCollapseNode>(NodeToVisit))
+		// Already local
+		if (NodeToVisit.HostObject == ThisFunctionHost)
 		{
-			const TArray<URigVMNode*>& ContainedNodes = CollapseNode->GetContainedNodes();
-			for(URigVMNode* ContainedNode : ContainedNodes)
-			{
-				if(URigVMLibraryNode* ContainedLibraryNode = Cast<URigVMLibraryNode>(ContainedNode))
-				{
-					NodesToVisit.AddUnique(ContainedLibraryNode);
-				}
-			}
+			continue;
+		}
 
-			if(URigVMFunctionLibrary* OtherLibrary = Cast<URigVMFunctionLibrary>(CollapseNode->GetOuter()))
-			{
-				if(OtherLibrary != ThisLibrary)
-				{
-					bool bIsAvailable = false;
-					FRigVMGraphFunctionIdentifier Identifier = CollapseNode->GetFunctionIdentifier();
-					if (IRigVMGraphFunctionHost* Host = Cast<IRigVMGraphFunctionHost>(Identifier.HostObject.ResolveObject()))
-					{
-						bIsAvailable = Host->GetRigVMGraphFunctionStore()->IsFunctionPublic(Identifier);		
-					}
+		bool bIsPublic;
+		FRigVMGraphFunctionData* FunctionData = FRigVMGraphFunctionData::FindFunctionData(NodeToVisit, &bIsPublic);
+		if (!FunctionData)
+		{
+			ReportAndNotifyErrorf(TEXT("Cannot localize function - could not find function %s in host %s."), *NodeToVisit.LibraryNode.ToString(), *NodeToVisit.HostObject.ToString());
+			return LocalizedFunctions;
+		}
 
-					if(!bIsAvailable)
-					{
-						if(!bLocalizeDependentPrivateFunctions)
-						{
-							ReportAndNotifyErrorf(TEXT("Cannot localize function - dependency %s is private."), *CollapseNode->GetPathName());
-							return LocalizedFunctions;
-						}
-						
-						FunctionsToLocalize.AddUnique(CollapseNode);
-					}
-				}
-			}
+		// Do not localize public functions
+		if (bIsPublic)
+		{
+			continue;
 		}
 
-		else if(const URigVMFunctionReferenceNode* FunctionReferencedNode = Cast<URigVMFunctionReferenceNode>(NodeToVisit))
+		if (!bLocalizeDependentPrivateFunctions)
 		{
-			if(FunctionReferencedNode->GetReferencedFunctionHeader().LibraryPointer.HostObject != ThisLibrary->GetFunctionHostObjectPath())
-			{
-				if(URigVMCollapseNode* FunctionDefinition = Cast<URigVMCollapseNode>(FunctionReferencedNode->LoadReferencedNode()))
-				{
-					NodesToVisit.AddUnique(FunctionDefinition);
-				}
-			}
+			ReportAndNotifyErrorf(TEXT("Cannot localize function - dependency %s is private."), *NodeToVisit.LibraryNode.ToString());
+			return LocalizedFunctions;
+		}
+
+		FunctionsToLocalize.AddUnique(FunctionData);
+
+		// Look for its dependencies
+		for (TPair<FRigVMGraphFunctionIdentifier, uint32>& Pair : FunctionData->Header.Dependencies)
+		{
+			NodesToVisit.AddUnique(Pair.Key);
 		}
 	}
 	
 	// sort the functions to localize based on their nesting
-	Algo::Sort(FunctionsToLocalize, [](URigVMLibraryNode* A, URigVMLibraryNode* B) -> bool
+	Algo::Sort(FunctionsToLocalize, [](FRigVMGraphFunctionData* A, FRigVMGraphFunctionData* B) -> bool
 	{
 		check(A);
 		check(B);
-		return B->Contains(A);
+		return B->Header.Dependencies.Contains(A->Header.LibraryPointer);
 	});
 
-	// export all of the content for each node
-	TMap<const URigVMLibraryNode*, FString> ExportedTextPerFunction;
-	for(const URigVMLibraryNode* FunctionToLocalize : FunctionsToLocalize)
-	{
-		URigVMFunctionLibrary* OtherLibrary = Cast<URigVMFunctionLibrary>(FunctionToLocalize->GetOuter());
-		FRigVMControllerGraphGuard GraphGuard(this, OtherLibrary, false);
-
-		const TArray<FName> NodeNamesToExport = {FunctionToLocalize->GetFName()};
-		const FString ExportedText = ExportNodesToText(NodeNamesToExport);
-		ExportedTextPerFunction.Add(FunctionToLocalize, ExportedText);
-	}
-
 	FRigVMControllerCompileBracketScope CompileScope(this);
 	if (bSetupUndoRedo)
 	{
@@ -3922,25 +3939,26 @@ TMap<URigVMLibraryNode*, URigVMLibraryNode*> URigVMController::LocalizeFunctions
 	// import the functions to our local function library
 	{
 		FRigVMControllerGraphGuard GraphGuard(this, ThisLibrary, bSetupUndoRedo);
-		for(URigVMLibraryNode* FunctionToLocalize : FunctionsToLocalize)
+		for(FRigVMGraphFunctionData* FunctionToLocalize : FunctionsToLocalize)
 		{
-			const FString& ExportedText = ExportedTextPerFunction.FindChecked(FunctionToLocalize);
-			TArray<FName> ImportedNodeNames = ImportNodesFromText(ExportedText);
-			if(ImportedNodeNames.Num() != 1)
+			if (FunctionToLocalize->SerializedCollapsedNode.IsEmpty())
 			{
-				ReportErrorf(TEXT("Not possible to localize function %s"), *FunctionToLocalize->GetPathName());
-				continue;
+				if (URigVMLibraryNode* ReferencedFunction = Cast<URigVMLibraryNode>(FunctionToLocalize->Header.LibraryPointer.LibraryNode.TryLoad()))
+				{
+					if (IRigVMClientHost* ClientHost = ReferencedFunction->GetImplementingOuter<IRigVMClientHost>())
+					{
+						ClientHost->GetRigVMClient()->UpdateGraphFunctionSerializedGraph(ReferencedFunction);
+					}
+				}				
 			}
-
-			URigVMLibraryNode* LocalizedFunction = Cast<URigVMLibraryNode>(GetGraph()->FindNodeByName(ImportedNodeNames[0]));
-			if(LocalizedFunction == nullptr)
+			
+			TArray<FName> NodeNames = ImportNodesFromText(FunctionToLocalize->SerializedCollapsedNode, false);
+			if (NodeNames.Num() > 0)
 			{
-				ReportErrorf(TEXT("Not possible to localize function %s"), *FunctionToLocalize->GetPathName());
-				continue;
+				URigVMLibraryNode* LocalizedFunction = ThisLibrary->FindFunction(NodeNames[0]);
+				LocalizedFunctions.Add(FunctionToLocalize->Header.LibraryPointer, LocalizedFunction);
+				ThisLibrary->LocalizedFunctions.FindOrAdd(FunctionToLocalize->Header.LibraryPointer.LibraryNode.ToString(), LocalizedFunction);
 			}
-
-			LocalizedFunctions.Add(FunctionToLocalize, LocalizedFunction);
-			ThisLibrary->LocalizedFunctions.FindOrAdd(FunctionToLocalize->GetPathName(), LocalizedFunction);
 		}
 	}
 
@@ -3964,8 +3982,7 @@ TMap<URigVMLibraryNode*, URigVMLibraryNode*> URigVMController::LocalizeFunctions
 			}
 			else if(URigVMFunctionReferenceNode* FunctionReferenceNode = Cast<URigVMFunctionReferenceNode>(NodeToUpdate))
 			{
-				const URigVMLibraryNode* ReferencedNode = FunctionReferenceNode->LoadReferencedNode();
-				URigVMLibraryNode** RemappedNodePtr = LocalizedFunctions.Find(ReferencedNode);
+				URigVMLibraryNode** RemappedNodePtr = LocalizedFunctions.Find(FunctionReferenceNode->GetReferencedFunctionHeader().LibraryPointer);
 				if(RemappedNodePtr)
 				{
 					URigVMLibraryNode* RemappedNode = *RemappedNodePtr;
@@ -3982,26 +3999,19 @@ TMap<URigVMLibraryNode*, URigVMLibraryNode*> URigVMController::LocalizeFunctions
 
 	if (bPrintPythonCommand)
 	{
-		FString FunctionNames = TEXT("[");
-		for (auto It = InFunctionDefinitions.CreateConstIterator(); It; ++It)
+		for (const FRigVMGraphFunctionIdentifier& Identifier : InFunctionDefinitions)
 		{
-			FunctionNames += FString::Printf(TEXT("unreal.load_object(name = '%s', outer = None).get_local_function_library().find_function('%s')"),
-				*(*It)->GetLibrary()->GetOuter()->GetPathName(),
-				*(*It)->GetName());
-			if (It.GetIndex() < InFunctionDefinitions.Num() - 1)
-			{
-				FunctionNames += TEXT(", ");
-			}
-		}
-		FunctionNames += TEXT("]");
-
-		const FString GraphName = GetSanitizedGraphName(GetGraph()->GetGraphName());
+			const FString GraphName = GetSanitizedGraphName(GetGraph()->GetGraphName());
+			FRigVMGraphFunctionData* FunctionData = FRigVMGraphFunctionData::FindFunctionData(Identifier);
+			const FString FunctionDefinitionName = GetSanitizedNodeName(FunctionData->Header.Name.ToString());
 
-		RigVMPythonUtils::Print(GetGraphOuterName(), 
-							FString::Printf(TEXT("blueprint.get_controller_by_name('%s').localize_functions(%s, %s)"),
-											*GraphName,
-											*FunctionNames,
-											(bLocalizeDependentPrivateFunctions) ? TEXT("True") : TEXT("False")));
+			RigVMPythonUtils::Print(GetGraphOuterName(), 
+				FString::Printf(TEXT("blueprint.get_controller_by_name('%s').localize_function_from_path('%s', '%s', %s)"),
+						*GraphName,
+						*Identifier.HostObject.ToString(),
+						*FunctionDefinitionName,
+						bLocalizeDependentPrivateFunctions ? TEXT("True") : TEXT("False")));
+		}
 	}
 
 	return LocalizedFunctions;
@@ -13820,27 +13830,24 @@ bool URigVMController::CanAddNode(URigVMNode* InNode, bool bReportErrors, bool b
 			FRigVMGraphFunctionHeader FunctionDefinition = FunctionRefNode->GetReferencedFunctionHeader();
 			if(!CanAddFunctionRefForDefinition(FunctionDefinition, false))
 			{
-				if (URigVMLibraryNode* LibraryNode = FunctionRefNode->LoadReferencedNode())
+				URigVMFunctionLibrary* TargetLibrary = Graph->GetDefaultFunctionLibrary();
+				URigVMLibraryNode* NewFunctionDefinition = TargetLibrary->FindPreviouslyLocalizedFunction(FunctionDefinition.LibraryPointer);
+			
+				if((NewFunctionDefinition == nullptr) && RequestLocalizeFunctionDelegate.IsBound())
 				{
-					URigVMFunctionLibrary* TargetLibrary = Graph->GetDefaultFunctionLibrary();
-					URigVMLibraryNode* NewFunctionDefinition = TargetLibrary->FindPreviouslyLocalizedFunction(LibraryNode);
-				
-					if((NewFunctionDefinition == nullptr) && RequestLocalizeFunctionDelegate.IsBound())
+					if(RequestLocalizeFunctionDelegate.Execute(FunctionDefinition.LibraryPointer))
 					{
-						if(RequestLocalizeFunctionDelegate.Execute(LibraryNode))
-						{
-							NewFunctionDefinition = TargetLibrary->FindPreviouslyLocalizedFunction(LibraryNode);
-						}
+						NewFunctionDefinition = TargetLibrary->FindPreviouslyLocalizedFunction(FunctionDefinition.LibraryPointer);
 					}
+				}
 
-					if(NewFunctionDefinition == nullptr)
-					{
-						return false;
-					}
-				
-					SetReferencedFunction(FunctionRefNode, NewFunctionDefinition, false);
-					FunctionDefinition = FunctionRefNode->GetReferencedFunctionHeader();
+				if(NewFunctionDefinition == nullptr)
+				{
+					return false;
 				}
+			
+				SetReferencedFunction(FunctionRefNode, NewFunctionDefinition, false);
+				FunctionDefinition = FunctionRefNode->GetReferencedFunctionHeader();
 			}
 			
 			if(!CanAddFunctionRefForDefinition(FunctionDefinition, bReportErrors))
diff --git a/Engine/Source/Developer/RigVMDeveloper/Private/RigVMModel/RigVMFunctionLibrary.cpp b/Engine/Source/Developer/RigVMDeveloper/Private/RigVMModel/RigVMFunctionLibrary.cpp
--- a/Engine/Source/Developer/RigVMDeveloper/Private/RigVMModel/RigVMFunctionLibrary.cpp
+++ b/Engine/Source/Developer/RigVMDeveloper/Private/RigVMModel/RigVMFunctionLibrary.cpp
@@ -119,36 +119,37 @@ void URigVMFunctionLibrary::ForEachReferenceSoftPtr(const FName& InFunctionName,
 	}
 }
 
-URigVMLibraryNode* URigVMFunctionLibrary::FindPreviouslyLocalizedFunction(URigVMLibraryNode* InFunctionToLocalize)
+URigVMLibraryNode* URigVMFunctionLibrary::FindPreviouslyLocalizedFunction(FRigVMGraphFunctionIdentifier InFunctionToLocalize)
 {
-	if(InFunctionToLocalize == nullptr)
+	const FString PathName = InFunctionToLocalize.LibraryNode.ToString();
+	if(!LocalizedFunctions.Contains((PathName)))
 	{
 		return nullptr;
 	}
-	
-	const FString PathName = InFunctionToLocalize->GetPathName();
 
-	if(!LocalizedFunctions.Contains((PathName)))
+	FRigVMGraphFunctionData* FunctionData = FRigVMGraphFunctionData::FindFunctionData(InFunctionToLocalize);
+	if (!FunctionData)
 	{
 		return nullptr;
 	}
-
+	FRigVMGraphFunctionHeader& Header = FunctionData->Header;
+	
 	URigVMLibraryNode* LocalizedFunction = LocalizedFunctions.FindChecked(PathName);
 
 	// once we found the function - let's make sure it's notation is right
-	if(LocalizedFunction->GetPins().Num() != InFunctionToLocalize->GetPins().Num())
+	if(LocalizedFunction->GetPins().Num() != Header.Arguments.Num())
 	{
 		return nullptr;
 	}
-	for(int32 PinIndex=0; PinIndex < InFunctionToLocalize->GetPins().Num(); PinIndex++)
+	for(int32 PinIndex=0; PinIndex < Header.Arguments.Num(); PinIndex++)
 	{
-		URigVMPin* PinA = InFunctionToLocalize->GetPins()[PinIndex];
-		URigVMPin* PinB = LocalizedFunction->GetPins()[PinIndex];
+		FRigVMGraphFunctionArgument& Argument = Header.Arguments[PinIndex];
+		URigVMPin* Pin = LocalizedFunction->GetPins()[PinIndex];
 
-		if((PinA->GetFName() != PinB->GetFName()) ||
-			(PinA->GetCPPType() != PinB->GetCPPType()) ||
-			(PinA->GetCPPTypeObject() != PinB->GetCPPTypeObject()) ||
-			(PinA->IsArray() != PinB->IsArray()))
+		if((Argument.Name != Pin->GetFName()) ||
+			(Argument.CPPType.ToString() != Pin->GetCPPType()) ||
+			(Argument.CPPTypeObject != Pin->GetCPPTypeObject()) ||
+			(Argument.bIsArray != Pin->IsArray()))
 		{
 			return nullptr;
 		}
diff --git a/Engine/Source/Developer/RigVMDeveloper/Public/RigVMModel/Nodes/RigVMFunctionReferenceNode.h b/Engine/Source/Developer/RigVMDeveloper/Public/RigVMModel/Nodes/RigVMFunctionReferenceNode.h
--- a/Engine/Source/Developer/RigVMDeveloper/Public/RigVMModel/Nodes/RigVMFunctionReferenceNode.h
+++ b/Engine/Source/Developer/RigVMDeveloper/Public/RigVMModel/Nodes/RigVMFunctionReferenceNode.h
@@ -31,6 +31,7 @@ public:
 	virtual FString GetNodeKeywords() const override;
 	virtual TArray<FRigVMExternalVariable> GetExternalVariables() const override;
 	virtual const FRigVMTemplate* GetTemplate() const override { return nullptr; }
+	virtual FRigVMGraphFunctionIdentifier GetFunctionIdentifier() const override;
 	// end URigVMLibraryNode interface
 
 	bool IsReferencedFunctionHostLoaded() const;
diff --git a/Engine/Source/Developer/RigVMDeveloper/Public/RigVMModel/Nodes/RigVMLibraryNode.h b/Engine/Source/Developer/RigVMDeveloper/Public/RigVMModel/Nodes/RigVMLibraryNode.h
--- a/Engine/Source/Developer/RigVMDeveloper/Public/RigVMModel/Nodes/RigVMLibraryNode.h
+++ b/Engine/Source/Developer/RigVMDeveloper/Public/RigVMModel/Nodes/RigVMLibraryNode.h
@@ -58,7 +58,7 @@ public:
 	virtual TArray<FRigVMExternalVariable> GetExternalVariables() const;
 	virtual TMap<FRigVMGraphFunctionIdentifier, uint32> GetDependencies() const;
 
-	FRigVMGraphFunctionIdentifier GetFunctionIdentifier() const;
+	virtual FRigVMGraphFunctionIdentifier GetFunctionIdentifier() const;
 	FRigVMGraphFunctionHeader GetFunctionHeader(IRigVMGraphFunctionHost* InHostObject = nullptr) const;
 
 protected:
diff --git a/Engine/Source/Developer/RigVMDeveloper/Public/RigVMModel/RigVMClient.h b/Engine/Source/Developer/RigVMDeveloper/Public/RigVMModel/RigVMClient.h
--- a/Engine/Source/Developer/RigVMDeveloper/Public/RigVMModel/RigVMClient.h
+++ b/Engine/Source/Developer/RigVMDeveloper/Public/RigVMModel/RigVMClient.h
@@ -153,6 +153,7 @@ public:
 	bool UpdateDependenciesForFunction(const URigVMLibraryNode* InLibraryNode);
 	bool UpdateFunctionReferences(const FRigVMGraphFunctionHeader& InHeader, bool bUpdateDependencies, bool bUpdateExternalVariables);
 	bool DirtyGraphFunctionCompilationData(URigVMLibraryNode* InLibraryNode);
+	bool UpdateGraphFunctionSerializedGraph(URigVMLibraryNode* InLibraryNode);
 	bool IsFunctionPublic(URigVMLibraryNode* InLibraryNode);
 
 private:
diff --git a/Engine/Source/Developer/RigVMDeveloper/Public/RigVMModel/RigVMController.h b/Engine/Source/Developer/RigVMDeveloper/Public/RigVMModel/RigVMController.h
--- a/Engine/Source/Developer/RigVMDeveloper/Public/RigVMModel/RigVMController.h
+++ b/Engine/Source/Developer/RigVMDeveloper/Public/RigVMModel/RigVMController.h
@@ -79,7 +79,7 @@ private:
 DECLARE_DELEGATE_RetVal_OneParam(bool, FRigVMController_ShouldStructUnfoldDelegate, const UStruct*)
 DECLARE_DELEGATE_RetVal_OneParam(TArray<FRigVMExternalVariable>, FRigVMController_GetExternalVariablesDelegate, URigVMGraph*)
 DECLARE_DELEGATE_RetVal(const FRigVMByteCode*, FRigVMController_GetByteCodeDelegate)
-DECLARE_DELEGATE_RetVal_OneParam(bool, FRigVMController_RequestLocalizeFunctionDelegate, TSoftObjectPtr<URigVMLibraryNode>)
+DECLARE_DELEGATE_RetVal_OneParam(bool, FRigVMController_RequestLocalizeFunctionDelegate, FRigVMGraphFunctionIdentifier&)
 DECLARE_DELEGATE_RetVal_ThreeParams(FName, FRigVMController_RequestNewExternalVariableDelegate, FRigVMGraphVariableDescription, bool, bool);
 DECLARE_DELEGATE_RetVal_TwoParams(bool, FRigVMController_IsDependencyCyclicDelegate, const FRigVMGraphFunctionHeader& Dependent, const FRigVMGraphFunctionHeader& Dependency)
 DECLARE_DELEGATE_RetVal_TwoParams(FRigVMController_BulkEditResult, FRigVMController_RequestBulkEditDialogDelegate, URigVMLibraryNode*, ERigVMControllerBulkEditType)
@@ -472,19 +472,23 @@ public:
 
 	// Copies a function declaration into this graph's local function library
 	UFUNCTION(BlueprintCallable, Category = RigVMController)
+	URigVMLibraryNode* LocalizeFunctionFromPath(const FString& InHostPath, const FName& InFunctionName, bool bLocalizeDependentPrivateFunctions = true, bool bSetupUndoRedo = true, bool bPrintPythonCommand = false);
+
+	UFUNCTION(BlueprintCallable, Category = RigVMController)
     URigVMLibraryNode* LocalizeFunction(
-    	URigVMLibraryNode* InFunctionDefinition,
+		const FRigVMGraphFunctionIdentifier& InFunctionDefinition,
 		bool bLocalizeDependentPrivateFunctions = true,
-    	bool bSetupUndoRedo = true,
-    	bool bPrintPythonCommand = false);
+		bool bSetupUndoRedo = true,
+		bool bPrintPythonCommand = false);
 
 	// Copies a series of function declaratioms into this graph's local function library
 	UFUNCTION(BlueprintCallable, Category = RigVMController)
-    TMap<URigVMLibraryNode*, URigVMLibraryNode*> LocalizeFunctions(
-        TArray<URigVMLibraryNode*> InFunctionDefinitions,
-        bool bLocalizeDependentPrivateFunctions = true,
-        bool bSetupUndoRedo = true,
-        bool bPrintPythonCommand = false);
+	TMap<FRigVMGraphFunctionIdentifier, URigVMLibraryNode*> LocalizeFunctions(
+		TArray<FRigVMGraphFunctionIdentifier> InFunctionDefinitions,
+		bool bLocalizeDependentPrivateFunctions = true,
+		bool bSetupUndoRedo = true,
+		bool bPrintPythonCommand = false);
+
 
 	// Returns a unique name
 	static FName GetUniqueName(const FName& InName, TFunction<bool(const FName&)> IsNameAvailableFunction, bool bAllowPeriod, bool bAllowSpace);
@@ -765,7 +769,7 @@ public:
 	URigVMFunctionReferenceNode* AddFunctionReferenceNodeFromDescription(const FRigVMGraphFunctionHeader& InFunctionDefinition, const FVector2D& InNodePosition = FVector2D::ZeroVector, const
 	                                                                     FString& InNodeName = TEXT(""), bool bSetupUndoRedo = true, bool bPrintPythonCommand = false, bool bAllowPrivateFunctions = false);
 	UFUNCTION(BlueprintCallable, Category = RigVMController)
-	URigVMFunctionReferenceNode* AddExternalFunctionReferenceNode(const FString& InHostPath, const FName& InFunctionName, const FVector2D& InNodePosition = FVector2D::ZeroVector, const	                                                              FString& InNodeName = TEXT(""), bool bSetupUndoRedo = true, bool bPrintPythonCommand = false);
+	URigVMFunctionReferenceNode* AddExternalFunctionReferenceNode(const FString& InHostPath, const FName& InFunctionName, const FVector2D& InNodePosition = FVector2D::ZeroVector, const FString& InNodeName = TEXT(""), bool bSetupUndoRedo = true, bool bPrintPythonCommand = false);
 	
 	UFUNCTION(BlueprintCallable, Category = RigVMController)
 	URigVMFunctionReferenceNode* AddFunctionReferenceNode(URigVMLibraryNode* InFunctionDefinition, const FVector2D& InNodePosition = FVector2D::ZeroVector, const FString& InNodeName = TEXT(""), bool bSetupUndoRedo = true, bool bPrintPythonCommand = false);
diff --git a/Engine/Source/Developer/RigVMDeveloper/Public/RigVMModel/RigVMFunctionLibrary.h b/Engine/Source/Developer/RigVMDeveloper/Public/RigVMModel/RigVMFunctionLibrary.h
--- a/Engine/Source/Developer/RigVMDeveloper/Public/RigVMModel/RigVMFunctionLibrary.h
+++ b/Engine/Source/Developer/RigVMDeveloper/Public/RigVMModel/RigVMFunctionLibrary.h
@@ -67,7 +67,7 @@ public:
 	// Returns a function that has been previously localized based on the provided function to localize.
 	// We maintain meta data on what functions have been created locally based on which other ones,
 	// and use this method to avoid redundant localizations.
-	URigVMLibraryNode* FindPreviouslyLocalizedFunction(URigVMLibraryNode* InFunctionToLocalize);
+	URigVMLibraryNode* FindPreviouslyLocalizedFunction(FRigVMGraphFunctionIdentifier InFunctionToLocalize);
 
 	const FSoftObjectPath GetFunctionHostObjectPath() const;
 	URigVMFunctionLibrary_GetFunctionHostObjectPath GetFunctionHostObjectPathDelegate;
diff --git a/Engine/Source/Runtime/Core/Public/UObject/UE5MainStreamObjectVersion.h b/Engine/Source/Runtime/Core/Public/UObject/UE5MainStreamObjectVersion.h
--- a/Engine/Source/Runtime/Core/Public/UObject/UE5MainStreamObjectVersion.h
+++ b/Engine/Source/Runtime/Core/Public/UObject/UE5MainStreamObjectVersion.h
@@ -316,6 +316,9 @@ struct CORE_API FUE5MainStreamObjectVersion
 		// Add refraction mode None, converted from legacy when the refraction pin is not plugged.
 		MaterialRefractionModeNone,
 
+		// Store serialized graph function in the function data
+		RigVMSaveSerializedGraphInGraphFunctionData,
+
 		// -----<new versions can be added above this line>-------------------------------------------------
 		VersionPlusOne,
 		LatestVersion = VersionPlusOne - 1
diff --git a/Engine/Source/Runtime/RigVM/Private/RigVMCore/RigVMGraphFunctionDefinition.cpp b/Engine/Source/Runtime/RigVM/Private/RigVMCore/RigVMGraphFunctionDefinition.cpp
--- a/Engine/Source/Runtime/RigVM/Private/RigVMCore/RigVMGraphFunctionDefinition.cpp
+++ b/Engine/Source/Runtime/RigVM/Private/RigVMCore/RigVMGraphFunctionDefinition.cpp
@@ -93,3 +93,26 @@ void FRigVMGraphFunctionData::PostDuplicateHost(const FString& InOldHostPathName
 {
 	Header.PostDuplicateHost(InOldHostPathName, InNewHostPathName);
 }
+
+FRigVMGraphFunctionData* FRigVMGraphFunctionData::FindFunctionData(const FRigVMGraphFunctionIdentifier& InIdentifier, bool* bOutIsPublic)
+{
+	IRigVMGraphFunctionHost* FunctionHost = nullptr;
+	if (UObject* FunctionHostObj = InIdentifier.HostObject.TryLoad())
+	{
+		FunctionHost = Cast<IRigVMGraphFunctionHost>(FunctionHostObj);									
+	}
+
+	if (!FunctionHost)
+	{
+		return nullptr;
+	}
+
+	FRigVMGraphFunctionStore* FunctionStore = FunctionHost->GetRigVMGraphFunctionStore();
+	if (!FunctionStore)
+	{
+		return nullptr;
+	}
+
+	return FunctionStore->FindFunction(InIdentifier, bOutIsPublic);
+}
+
diff --git a/Engine/Source/Runtime/RigVM/Public/RigVMCore/RigVMGraphFunctionDefinition.h b/Engine/Source/Runtime/RigVM/Public/RigVMCore/RigVMGraphFunctionDefinition.h
--- a/Engine/Source/Runtime/RigVM/Public/RigVMCore/RigVMGraphFunctionDefinition.h
+++ b/Engine/Source/Runtime/RigVM/Public/RigVMCore/RigVMGraphFunctionDefinition.h
@@ -4,6 +4,7 @@
 
 #include "RigVMCore/RigVMExternalVariable.h"
 #include "RigVMCore/RigVMByteCode.h"
+#include "UObject/UE5MainStreamObjectVersion.h"
 #include "RigVMGraphFunctionDefinition.generated.h"
 
 class IRigVMGraphFunctionHost;
@@ -423,6 +424,9 @@ struct RIGVM_API FRigVMGraphFunctionData
 	UPROPERTY()
 	FRigVMFunctionCompilationData CompilationData;
 
+	UPROPERTY()
+	FString SerializedCollapsedNode;
+
 	bool IsMutable() const;
 
 	bool operator==(const FRigVMGraphFunctionData& Other) const
@@ -436,8 +440,17 @@ struct RIGVM_API FRigVMGraphFunctionData
 	{
 		Ar << Data.Header;
 		Ar << Data.CompilationData;
+
+		if (Ar.CustomVer(FUE5MainStreamObjectVersion::GUID) < FUE5MainStreamObjectVersion::RigVMSaveSerializedGraphInGraphFunctionData)
+		{
+			return Ar;
+		}
+
+		Ar << Data.SerializedCollapsedNode;
 		return Ar;
 	}
 
 	void PostDuplicateHost(const FString& InOldPathName, const FString& InNewPathName);
+
+	static FRigVMGraphFunctionData* FindFunctionData(const FRigVMGraphFunctionIdentifier& InIdentifier, bool* bOutIsPublic = nullptr);	
 };
