Control Rig: Fix pin redirectors on old structs not working

diff --git a/Engine/Source/Developer/RigVMDeveloper/Private/RigVMModel/RigVMController.cpp b/Engine/Source/Developer/RigVMDeveloper/Private/RigVMModel/RigVMController.cpp
--- a/Engine/Source/Developer/RigVMDeveloper/Private/RigVMModel/RigVMController.cpp
+++ b/Engine/Source/Developer/RigVMDeveloper/Private/RigVMModel/RigVMController.cpp
@@ -14926,6 +14926,29 @@ bool URigVMController::ShouldRedirectPin(UScriptStruct* InOwningStruct, const FS
 
 	FProperty* Property = InOwningStruct->FindPropertyByName(*PinName);
 	if (Property == nullptr)
+	{
+		// If we were unable to find the property, it might be the case where the whole struct was also redirected. So we need to find redirectors on the old structs.
+		TArray<FCoreRedirectObjectName> PreviousStructNames;
+		const FCoreRedirectObjectName NewStructName(InOwningStruct->GetFName(), InOwningStruct->GetOuter()->GetFName(), *InOwningStruct->GetOutermost()->GetPathName());
+		FCoreRedirects::FindPreviousNames(ECoreRedirectFlags::Type_Struct, NewStructName, PreviousStructNames);
+		for (const FCoreRedirectObjectName& PreviousStructName : PreviousStructNames)
+		{
+			const FCoreRedirectObjectName OldPinName(*PinName, PreviousStructName.ObjectName, PreviousStructName.PackageName);
+			const FCoreRedirectObjectName NewPinName = FCoreRedirects::GetRedirectedName(ECoreRedirectFlags::Type_Property, OldPinName);
+
+			if (OldPinName != NewPinName)
+			{
+				PinName = NewPinName.ObjectName.ToString();
+				bShouldRedirect = true;
+				Property = InOwningStruct->FindPropertyByName(*PinName);
+				if (Property)
+				{
+					break;
+				}
+			}
+		}
+	}
+	if (Property == nullptr)
 	{
 		return false;
 	}
