Control Rig: Fix compilation error when SetControlVisibility is used

diff --git a/Engine/Source/Developer/RigVMDeveloper/Private/RigVMCompiler/RigVMCompiler.cpp b/Engine/Source/Developer/RigVMDeveloper/Private/RigVMCompiler/RigVMCompiler.cpp
--- a/Engine/Source/Developer/RigVMDeveloper/Private/RigVMCompiler/RigVMCompiler.cpp
+++ b/Engine/Source/Developer/RigVMDeveloper/Private/RigVMCompiler/RigVMCompiler.cpp
@@ -651,6 +651,32 @@ bool URigVMCompiler::Compile(TArray<URigVMGraph*> InGraphs, URigVMController* In
 						Settings.ASTSettings.Report(EMessageSeverity::Error, ModelNode, UnresolvedUnitNodeMessage);
 						bEncounteredGraphError = true;
 					}
+
+					// Make sure all the pins exist in the node
+					ScriptStruct = UnitNode->GetScriptStruct();
+					if (ScriptStruct)
+					{
+						for (TFieldIterator<FProperty> It(ScriptStruct, EFieldIterationFlags::None); It; ++It)
+						{
+							const FRigVMTemplateArgument ExpectedArgument(*It);
+							const TRigVMTypeIndex ExpectedTypeIndex = ExpectedArgument.GetSupportedTypeIndices()[0];
+							if (URigVMPin* Pin = UnitNode->FindPin(ExpectedArgument.Name.ToString()))
+							{
+								if (Pin->GetTypeIndex() != ExpectedArgument.GetTypeIndices()[0])
+								{
+									FString MissingPinMessage = FString::Printf(TEXT("Could not find pin %s of type %s in Node @@."), *ExpectedArgument.Name.ToString(), *FRigVMRegistry::Get().GetType(ExpectedArgument.TypeIndices[0]).CPPType.ToString());
+									Settings.ASTSettings.Report(EMessageSeverity::Error, ModelNode, MissingPinMessage);
+									bEncounteredGraphError = true;
+								}
+							}
+							else
+							{
+								FString MissingPinMessage = FString::Printf(TEXT("Could not find pin %s of type %s in Node @@."), *ExpectedArgument.Name.ToString(), *FRigVMRegistry::Get().GetType(ExpectedArgument.TypeIndices[0]).CPPType.ToString());
+								Settings.ASTSettings.Report(EMessageSeverity::Error, ModelNode, MissingPinMessage);
+								bEncounteredGraphError = true;
+							}
+						}
+					}
 				}
 			}
 
diff --git a/Engine/Source/Developer/RigVMDeveloper/Private/RigVMModel/RigVMController.cpp b/Engine/Source/Developer/RigVMDeveloper/Private/RigVMModel/RigVMController.cpp
--- a/Engine/Source/Developer/RigVMDeveloper/Private/RigVMModel/RigVMController.cpp
+++ b/Engine/Source/Developer/RigVMDeveloper/Private/RigVMModel/RigVMController.cpp
@@ -13034,7 +13034,17 @@ URigVMTemplateNode* URigVMController::AddTemplateNode(const FName& InNotation, c
 	FRigVMRegistry& Registry = FRigVMRegistry::Get();
 	AddPinsForTemplate(Template, Types, Node);
 
-	UpdateTemplateNodePinTypes(Node, false);
+	if (Node->HasWildCardPin())
+	{
+		UpdateTemplateNodePinTypes(Node, false);
+	}
+	else
+	{
+		if (!Node->IsA<URigVMFunctionEntryNode>() && !Node->IsA<URigVMFunctionReturnNode>())
+		{
+			FullyResolveTemplateNode(Node, INDEX_NONE, false);
+		}
+	}
 
 	Graph->Nodes.Add(Node);
 
