Control Rig: Fix property redirect in QuaternionRotateVector

diff --git a/Engine/Plugins/Animation/ControlRig/Config/BaseControlRig.ini b/Engine/Plugins/Animation/ControlRig/Config/BaseControlRig.ini
--- a/Engine/Plugins/Animation/ControlRig/Config/BaseControlRig.ini
+++ b/Engine/Plugins/Animation/ControlRig/Config/BaseControlRig.ini
@@ -18,8 +18,8 @@
 +PropertyRedirects=(OldName="RigUnit_ToRigSpace_Transform.Transform", NewName="RigUnit_ToRigSpace_Transform.Value")
 +PropertyRedirects=(OldName="RigUnit_ToRigSpace_Location.Location", NewName="RigUnit_ToRigSpace_Location.Value")
 +PropertyRedirects=(OldName="RigUnit_ToRigSpace_Rotation.Rotation", NewName="RigUnit_ToRigSpace_Rotation.Value")
-+PropertyRedirects=(OldName="RigUnit_MathQuaternionRotateVector.Quaternion", NewName="RigUnit_MathQuaternionRotateVector.Transform")
-+PropertyRedirects=(OldName="RigUnit_MathTransformRotateVector.Direction", NewName="RigUnit_MathTransformRotateVector.Vector")
++PropertyRedirects=(OldName="RigVMFunction_MathQuaternionRotateVector.Quaternion", NewName="RigVMFunction_MathQuaternionRotateVector.Transform")
++PropertyRedirects=(OldName="RigVMFunction_MathTransformRotateVector.Direction", NewName="RigVMFunction_MathTransformRotateVector.Vector")
 +PropertyRedirects=(OldName="RigUnit_SetRelativeTransformForItem.RelativeTransform", NewName="RigUnit_SetRelativeTransformForItem.Value")
 
 +ClassRedirects=(OldName="RigVMGraph",NewName="/Script/RigVMDeveloper.RigVMGraph",ValueChanges=(("Model","RigVMModel")))
diff --git a/Engine/Source/Developer/RigVMDeveloper/Private/RigVMModel/RigVMController.cpp b/Engine/Source/Developer/RigVMDeveloper/Private/RigVMModel/RigVMController.cpp
--- a/Engine/Source/Developer/RigVMDeveloper/Private/RigVMModel/RigVMController.cpp
+++ b/Engine/Source/Developer/RigVMDeveloper/Private/RigVMModel/RigVMController.cpp
@@ -14936,29 +14936,6 @@ bool URigVMController::ShouldRedirectPin(UScriptStruct* InOwningStruct, const FS
 
 	FProperty* Property = InOwningStruct->FindPropertyByName(*PinName);
 	if (Property == nullptr)
-	{
-		// If we were unable to find the property, it might be the case where the whole struct was also redirected. So we need to find redirectors on the old structs.
-		TArray<FCoreRedirectObjectName> PreviousStructNames;
-		const FCoreRedirectObjectName NewStructName(InOwningStruct->GetFName(), InOwningStruct->GetOuter()->GetFName(), *InOwningStruct->GetOutermost()->GetPathName());
-		FCoreRedirects::FindPreviousNames(ECoreRedirectFlags::Type_Struct, NewStructName, PreviousStructNames);
-		for (const FCoreRedirectObjectName& PreviousStructName : PreviousStructNames)
-		{
-			const FCoreRedirectObjectName OldPinName(*PinName, PreviousStructName.ObjectName, PreviousStructName.PackageName);
-			const FCoreRedirectObjectName NewPinName = FCoreRedirects::GetRedirectedName(ECoreRedirectFlags::Type_Property, OldPinName);
-
-			if (OldPinName != NewPinName)
-			{
-				PinName = NewPinName.ObjectName.ToString();
-				bShouldRedirect = true;
-				Property = InOwningStruct->FindPropertyByName(*PinName);
-				if (Property)
-				{
-					break;
-				}
-			}
-		}
-	}
-	if (Property == nullptr)
 	{
 		return false;
 	}
