Control Rig: Fix RigVMController redirect not finding struct in new package for CRFourPointBezier

diff --git a/Engine/Source/Developer/RigVMDeveloper/Private/RigVMModel/RigVMController.cpp b/Engine/Source/Developer/RigVMDeveloper/Private/RigVMModel/RigVMController.cpp
--- a/Engine/Source/Developer/RigVMDeveloper/Private/RigVMModel/RigVMController.cpp
+++ b/Engine/Source/Developer/RigVMDeveloper/Private/RigVMModel/RigVMController.cpp
@@ -19207,7 +19207,7 @@ static UObject* FindObjectGloballyWithRedirectors(const TCHAR* InObjectName)
 		FCoreRedirectObjectName(InObjectName),
 		NewObjectName,
 		nullptr,
-		ECoreRedirectMatchFlags::None);
+		ECoreRedirectMatchFlags::AllowPartialMatch); // AllowPartialMatch to allow redirects from one package to another (see /Script/ControlRig.CRFourPointBezier -> /Script/RigVM.RigVMFourPointBezier)
 
 	if (!bFoundRedirect)
 	{
