Control Rig: Update nativized rigs

diff --git a/Engine/Source/Developer/RigVMDeveloper/Private/RigVMCompiler/RigVMCodeGenerator.cpp b/Engine/Source/Developer/RigVMDeveloper/Private/RigVMCompiler/RigVMCodeGenerator.cpp
--- a/Engine/Source/Developer/RigVMDeveloper/Private/RigVMCompiler/RigVMCodeGenerator.cpp
+++ b/Engine/Source/Developer/RigVMDeveloper/Private/RigVMCompiler/RigVMCodeGenerator.cpp
@@ -115,8 +115,8 @@ static constexpr TCHAR RigVM_GetVMHashFormat[] = TEXT("\tvirtual uint32 GetVMHas
 static constexpr TCHAR RigVM_GetEntryNamesFormat[] = TEXT("\tvirtual const TArray<FName>& GetEntryNames() const override\r\n\t{\r\n\t\tstatic const TArray<FName> StaticEntryNames = { {0} };\r\n\t\treturn StaticEntryNames;\r\n\t}");
 static constexpr TCHAR RigVM_DeclareUpdateExternalVariablesFormat[] = TEXT("\tvirtual void UpdateExternalVariables() override;");
 static constexpr TCHAR RigVM_DeclareInvokeEntryByNameFormat[] = TEXT("\tERigVMExecuteResult InvokeEntryByName(const FName& InEntryName{0});");
-static constexpr TCHAR RigVM_DeclareInitializeFormat[] = TEXT("\tvirtual bool Initialize(TArrayView<URigVMMemoryStorage*> Memory, bool bInitializeMemory = true) override;");
-static constexpr TCHAR RigVM_DefineInitializeFormat[] = TEXT("bool U{0}::Initialize(TArrayView<URigVMMemoryStorage*> Memory, bool bInitializeMemory)\r\n{");
+static constexpr TCHAR RigVM_DeclareInitializeFormat[] = TEXT("\tvirtual bool Initialize(TArrayView<URigVMMemoryStorage*> Memory) override;");
+static constexpr TCHAR RigVM_DefineInitializeFormat[] = TEXT("bool U{0}::Initialize(TArrayView<URigVMMemoryStorage*> Memory)\r\n{");
 static constexpr TCHAR RigVM_DeclareExecuteFormat[] = TEXT("\tvirtual ERigVMExecuteResult Execute(TArrayView<URigVMMemoryStorage*> Memory, const FName& InEntryName) override;");
 static constexpr TCHAR RigVM_DefineUpdateExternalVariablesFormat[] = TEXT("void U{0}::UpdateExternalVariables()\r\n{");
 static constexpr TCHAR RigVM_DefineInvokeEntryByNameFormat[] = TEXT("ERigVMExecuteResult U{0}::InvokeEntryByName(const FName& InEntryName{1})\r\n{");
