RigVM: Remove support for constant folding

diff --git a/Engine/Source/Developer/RigVMDeveloper/Private/RigVMCompiler/RigVMAST.cpp b/Engine/Source/Developer/RigVMDeveloper/Private/RigVMCompiler/RigVMAST.cpp
--- a/Engine/Source/Developer/RigVMDeveloper/Private/RigVMCompiler/RigVMAST.cpp
+++ b/Engine/Source/Developer/RigVMDeveloper/Private/RigVMCompiler/RigVMAST.cpp
@@ -713,20 +713,6 @@ FRigVMParserAST::FRigVMParserAST(TArray<URigVMGraph*> InGraphs, URigVMController
 	InjectExitsToEntries();
 	FoldNoOps();
 
-	// keep folding constant branches and values while we can
-	bool bContinueToFoldConstantBranches = InSettings.bFoldConstantBranches;
-	while (bContinueToFoldConstantBranches)
-	{
-		bContinueToFoldConstantBranches = false;
-		if (LibraryNodeBeingCompiled == nullptr)
-		{
-			if (FoldConstantValuesToLiterals(InGraphs, InController, InExternalVariables))
-			{
-				bContinueToFoldConstantBranches = true;
-			}
-		}
-	}
-
 	BubbleUpExpressions();
 
 	if (InSettings.bFoldAssignments)
@@ -1743,387 +1729,6 @@ void FRigVMParserAST::FoldAssignments()
 
 	RemoveExpressions(ExpressionsToRemove);
 }
-bool FRigVMParserAST::FoldConstantValuesToLiterals(TArray<URigVMGraph*> InGraphs, URigVMController* InController, const TArray<FRigVMExternalVariable>& InExternalVariables)
-{
-	DECLARE_SCOPE_HIERARCHICAL_COUNTER_FUNC()
-
-	if (InController == nullptr)
-	{
-		return false;
-	}
-
-	// loop over all call externs and figure out if they are a non-const node
-	// with one or more const pins. we then build a temporary VM to run the part of the 
-	// graph, and pull out the required values - we then bake the value into a literal 
-	// and remove the tree that created the value.
-
-	TMap<FString, FString> ComputedDefaultValues;
-	TArray<FRigVMASTProxy> PinsToUpdate;
-	TArray<FRigVMASTProxy> RootPinsToUpdate;
-	TArray<FRigVMASTProxy> PinsToCompute;
-	TArray<FRigVMASTProxy> NodesToCompute;
-
-	for (const FRigVMASTProxy& NodeProxy : NodeProxies)
-	{
-		if (NodeProxy.IsA<URigVMVariableNode>() ||
-			NodeProxy.IsA<URigVMEnumNode>() ||
-			NodeProxy.IsA<URigVMFunctionReferenceNode>())
-		{
-			continue;
-		}
-
-		FRigVMExprAST** NodeExprPtr = SubjectToExpression.Find(NodeProxy);
-		if(NodeExprPtr == nullptr)
-		{
-			continue;
-		}
-
-		FRigVMExprAST* NodeExpr = *NodeExprPtr;
-		if (NodeExpr->IsConstant())
-		{
-			continue;
-		}
-
-		URigVMNode* Node = NodeProxy.GetSubjectChecked<URigVMNode>();
-		const TArray<URigVMPin*> Pins = Node->GetPins();
-		for (URigVMPin* Pin : Pins)
-		{
-			if (Pin->GetDirection() != ERigVMPinDirection::Input &&
-				Pin->GetDirection() != ERigVMPinDirection::IO)
-			{
-				continue;
-			}
-
-			FRigVMASTProxy PinProxy = NodeProxy.GetSibling(Pin);
-			FRigVMExprAST** PinExprPtr = SubjectToExpression.Find(PinProxy);
-			if (PinExprPtr == nullptr)
-			{
-				continue;
-			}
-			FRigVMExprAST* PinExpr = *PinExprPtr;
-			if (PinExpr->IsA(FRigVMExprAST::EType::Literal))
-			{
-				if (const FRigVMExprAST* VarPinExpr = PinExpr->GetFirstParentOfType(FRigVMExprAST::EType::Var))
-				{
-					if (VarPinExpr->GetName() == PinExpr->GetName())
-					{
-						PinExpr = (FRigVMExprAST*)VarPinExpr;
-					}
-				}
-
-				// if we are still a literal, carry on
-				if (PinExpr->IsA(FRigVMExprAST::EType::Literal))
-				{
-					continue;
-				}
-			}
-
-			TArray<int32> LinkIndices = GetSourceLinkIndices(PinProxy, true);
-			if (LinkIndices.Num() == 0)
-			{
-				continue;
-			}
-
-			if (!PinExpr->IsConstant())
-			{
-				continue;
-			}
-
-			bool bFoundValidSourcePin = false;
-			for (const int32 LinkIndex : LinkIndices)
-			{
-				const FRigVMASTProxy& SourcePinProxy = Links[LinkIndex].SourceProxy;
-
-				// don't compute values which are also wired to an sub pins
-				bool bHasLinkToSubPin = false;
-				const TArray<int32> CurrentTargetLinkIndices = GetTargetLinkIndices(SourcePinProxy, true);
-				for(int32 TargetLinkIndex : CurrentTargetLinkIndices)
-				{
-					if(!Links[TargetLinkIndex].SegmentPath.IsEmpty())
-					{
-						bHasLinkToSubPin = true;
-						break;
-					}
-				}
-
-				if(bHasLinkToSubPin)
-				{
-					bFoundValidSourcePin = false;
-					break;
-				}
-				
-				URigVMNode* SourceNode = SourcePinProxy.GetSubjectChecked<URigVMPin>()->GetNode();
-				FRigVMASTProxy SourceNodeProxy = SourcePinProxy.GetSibling(SourceNode);
-
-				check(SourceNodeProxy.IsValid());
-
-				if (SourceNodeProxy.IsA<URigVMVariableNode>() ||
-					SourceNodeProxy.IsA<URigVMRerouteNode>() ||
-					SourceNodeProxy.IsA<URigVMFunctionReferenceNode>() ||
-					SourceNodeProxy.IsA<URigVMEnumNode>())
-				{
-					continue;
-				}
-
-				PinsToCompute.AddUnique(SourcePinProxy);
-				NodesToCompute.AddUnique(SourceNodeProxy);
-				bFoundValidSourcePin = true;
-			}
-
-			if (bFoundValidSourcePin)
-			{
-				PinsToUpdate.Add(PinProxy);
-				RootPinsToUpdate.AddUnique(PinProxy.GetSibling(Pin->GetRootPin()));
-			}
-		}
-	}
-
-	if (NodesToCompute.Num() == 0)
-	{
-		return false;
-	}
-
-	// add all of the additional nodes driving these
-	for(int32 NodeToComputeIndex = 0; NodeToComputeIndex < NodesToCompute.Num(); NodeToComputeIndex++)
-	{
-		FRigVMASTProxy ProxyToCompute = NodesToCompute[NodeToComputeIndex];
-		if(URigVMNode* NodeToCompute = ProxyToCompute.GetSubject<URigVMNode>())
-		{
-			// first find all of the source node of this node directly
-			TArray<URigVMNode*> SourceNodes = NodeToCompute->GetLinkedSourceNodes();
-			for(URigVMNode* SourceNode : SourceNodes)
-			{
-				FRigVMASTProxy SourceProxy = ProxyToCompute.GetSibling(SourceNode);
-				NodesToCompute.AddUnique(SourceProxy);
-			}
-
-			// if the node is a collapse node - also add the return node
-			if(URigVMCollapseNode* CollapseToCompute = Cast<URigVMCollapseNode>(NodeToCompute))
-			{
-				if(URigVMNode* ReturnNodeToCompute = CollapseToCompute->GetReturnNode())
-				{
-					const FRigVMASTProxy CollapseProxy = ProxyToCompute.GetSibling(CollapseToCompute);
-					const FRigVMASTProxy ReturnProxy = CollapseProxy.GetChild(ReturnNodeToCompute);
-					NodesToCompute.AddUnique(ReturnProxy);
-				}
-			}
-      
-			// if the node is an entry node - also add host collapse node
-			if(NodeToCompute->IsA<URigVMFunctionEntryNode>() && NodeToCompute->GetGraph()->GetOuter()->IsA<URigVMCollapseNode>())
-			{
-				check(ProxyToCompute.GetCallstack().Num() > 1);
-				const FRigVMASTProxy LibraryProxy = ProxyToCompute.GetParent();
-				NodesToCompute.AddUnique(LibraryProxy);
-			}
-		}
-	}
-
-	// validate nodes to compute
-	for(int32 NodeToComputeIndex = 0; NodeToComputeIndex < NodesToCompute.Num(); NodeToComputeIndex++)
-	{
-		FRigVMASTProxy ProxyToCompute = NodesToCompute[NodeToComputeIndex];
-		if(URigVMNode* NodeToCompute = ProxyToCompute.GetSubject<URigVMNode>())
-		{
-			// All mutable nodes should be considered varying, so should never arrive to this point
-			check(!NodeToCompute->IsMutable());			
-		}
-	}
-	
-
-	// we now know the node we need to run.
-	// let's build an temporary AST which has only those nodes
-	TSharedPtr<FRigVMParserAST> TempAST = MakeShareable(new FRigVMParserAST(InGraphs, NodesToCompute));
-
-	// share the pin overrides with the constant folding AST to ensure
-	// the complete view of default values across function references is available
-	// in the subset AST.
-	TempAST->PinOverrides = PinOverrides;
-
-	// build the VM to run this AST
-	TMap<FString, FRigVMOperand> Operands;
-	check(!InGraphs.IsEmpty());
-	URigVM* TempVM = NewObject<URigVM>(InGraphs[0]);
-	
-	URigVMCompiler* TempCompiler = NewObject<URigVMCompiler>(GetTransientPackage());
-	TempCompiler->Settings.SetupNodeInstructionIndex = false;
-	TempCompiler->Settings.IsPreprocessorPhase = true;
-	TempCompiler->Settings.EnablePinWatches = false;
-	TempCompiler->Settings.ASTSettings = FRigVMParserASTSettings::Fast(Settings.ExecuteContextStruct);
-
-	TempCompiler->Compile(InGraphs, InController, TempVM, InExternalVariables, &Operands, TempAST);
-
-	TArray<URigVMMemoryStorage*> Memory;
-	Memory.Add(TempVM->GetWorkMemory());
-	Memory.Add(TempVM->GetLiteralMemory());
-	Memory.Add(TempVM->GetDebugMemory());
-	TempVM->Execute(Memory);
-
-	// copy the values out of the temp VM and set them on the cached value
-	for (const FRigVMASTProxy& PinToComputeProxy : PinsToCompute)
-	{
-		TGuardValue<bool> GuardControllerNotifs(InController->bSuspendNotifications, true);
-
-		URigVMPin* PinToCompute = PinToComputeProxy.GetSubjectChecked<URigVMPin>();
-		URigVMPin* RootPin = PinToCompute->GetRootPin();
-		FRigVMASTProxy RootPinProxy = PinToComputeProxy.GetSibling(RootPin);
-
-		FRigVMVarExprAST* RootVarExpr = nullptr;
-		FRigVMExprAST** RootPinExprPtr = SubjectToExpression.Find(RootPinProxy);
-		if (RootPinExprPtr != nullptr)
-		{
-			FRigVMExprAST* RootPinExpr = *RootPinExprPtr;
-			if (RootPinExpr->IsA(FRigVMExprAST::EType::Var))
-			{
-				RootVarExpr = RootPinExpr->To<FRigVMVarExprAST>();
-			}
-		}
-
-		FString PinHash = URigVMCompiler::GetPinHash(RootPin, RootVarExpr, false, LibraryNodeBeingCompiled);
-		const FRigVMOperand* OperandPtr = Operands.Find(PinHash);
-		if(OperandPtr == nullptr)
-		{
-			continue;
-		}
-		const FRigVMOperand& Operand = *OperandPtr;
-
-		FString DefaultValue;
-		if(Operand.GetMemoryType() == ERigVMMemoryType::Work)
-		{
-			DefaultValue = TempVM->GetWorkMemory()->GetDataAsString(Operand.GetRegisterIndex());
-		}
-		else if(Operand.GetMemoryType() == ERigVMMemoryType::Literal)
-		{
-			DefaultValue = TempVM->GetLiteralMemory()->GetDataAsString(Operand.GetRegisterIndex());
-		}
-		if(DefaultValue.IsEmpty())
-		{
-			continue;
-		}
-
-		// FString TempDefaultValue = FString::Printf(TEXT("(%s)"), *FString::Join(DefaultValues, TEXT(",")));
-		// UE_LOG(LogRigVMDeveloper, Display, TEXT("Computed constant value '%s' = '%s'"), *PinHash, *TempDefaultValue);
-
-		TArray<FString> SegmentNames;
-		if (!URigVMPin::SplitPinPath(PinToCompute->GetSegmentPath(), SegmentNames))
-		{
-			SegmentNames.Add(PinToCompute->GetName());
-		}
-
-		URigVMPin* PinForDefaultValue = RootPin;
-		while (PinForDefaultValue != PinToCompute && SegmentNames.Num() > 0)
-		{
-			TArray<FString> SplitDefaultValues = URigVMPin::SplitDefaultValue(DefaultValue);
-
-			if (PinForDefaultValue->IsArray())
-			{
-				int32 ElementIndex = FCString::Atoi(*SegmentNames[0]);
-				DefaultValue = SplitDefaultValues[ElementIndex];
-				PinForDefaultValue = PinForDefaultValue->GetSubPins()[ElementIndex];
-				URigVMController::PostProcessDefaultValue(PinForDefaultValue, DefaultValue);
-				SegmentNames.RemoveAt(0);
-			}
-			else if (PinForDefaultValue->IsStruct())
-			{
-				if(SplitDefaultValues.IsEmpty())
-				{
-					break;
-				}
-				
-				for (const FString& MemberNameValuePair : SplitDefaultValues)
-				{
-					FString MemberName, MemberValue;
-					if (MemberNameValuePair.Split(TEXT("="), &MemberName, &MemberValue))
-					{
-						if (MemberName == SegmentNames[0])
-						{
-							URigVMPin* SubPin = PinForDefaultValue->FindSubPin(MemberName);
-							if (SubPin == nullptr)
-							{
-								SegmentNames.Reset();
-								break;
-							}
-
-							DefaultValue = MemberValue;
-							PinForDefaultValue = SubPin;
-							URigVMController::PostProcessDefaultValue(PinForDefaultValue, DefaultValue);
-							SegmentNames.RemoveAt(0);
-							break;
-						}
-					}
-				}
-			}
-			else
-			{
-				checkNoEntry();
-			}
-		}
-
-		const TArray<int32> LinkIndices = GetTargetLinkIndices(PinToComputeProxy);
-		for (const int32 LinkIndex : LinkIndices)
-		{
-			URigVMPin::FPinOverrideValue OverrideValue;
-			OverrideValue.DefaultValue = DefaultValue;
-			PinOverrides.FindOrAdd(Links[LinkIndex].TargetProxy) = OverrideValue;
-		}
-	}
-
-	// now remove all of the expressions no longer needed
-	TArray<FRigVMExprAST*> ExpressionsToRemove;
-	for (const FRigVMASTProxy& RootPinToUpdateProxy : RootPinsToUpdate)
-	{
-		FRigVMExprAST** PreviousExprPtr = SubjectToExpression.Find(RootPinToUpdateProxy);
-		if (PreviousExprPtr)
-		{
-			FRigVMVarExprAST* PreviousVarExpr = (*PreviousExprPtr)->To<FRigVMVarExprAST>();
-
-			// if the previous var expression is a literal used to initialize a var
-			// (for example on an IO pin, or when we are driving sub pins)
-			if (PreviousVarExpr->IsA(FRigVMExprAST::EType::Literal))
-			{
-				bool bRedirectedVar = false;
-				for (int32 ParentIndex = 0; ParentIndex < PreviousVarExpr->NumParents(); ParentIndex++)
-				{
-					const FRigVMExprAST* ParentExpr = PreviousVarExpr->ParentAt(ParentIndex);
-					if (ParentExpr->IsA(FRigVMExprAST::EType::Assign))
-					{
-						for (int32 GrandParentIndex = 0; GrandParentIndex < ParentExpr->NumParents(); GrandParentIndex++)
-						{
-							const FRigVMExprAST* GrandParentExpr = ParentExpr->ParentAt(GrandParentIndex);
-							if (GrandParentExpr->IsA(FRigVMExprAST::EType::Block))
-							{
-								GrandParentExpr = GrandParentExpr->GetParent();
-							}
-							if (GrandParentExpr->IsA(FRigVMExprAST::EType::Var) && (GrandParentExpr->GetName() == PreviousVarExpr->GetName()))
-							{
-								PreviousVarExpr = (FRigVMVarExprAST*)GrandParentExpr->To<FRigVMVarExprAST>();
-								bRedirectedVar = true;
-								break;
-							}
-						}
-					}
-
-					if (bRedirectedVar)
-					{
-						break;
-					}
-				}
-			}
-
-			FRigVMLiteralExprAST* LiteralExpr = MakeExpr<FRigVMLiteralExprAST>(RootPinToUpdateProxy);
-			LiteralExpr->Name = PreviousVarExpr->Name;
-			SubjectToExpression[RootPinToUpdateProxy] = LiteralExpr;
-			PreviousVarExpr->ReplaceBy(LiteralExpr);
-			ExpressionsToRemove.Add(PreviousVarExpr);
-		}
-	}
-
-	TempVM->Rename(nullptr, GetTransientPackage(), REN_ForceNoResetLoaders | REN_DoNotDirty | REN_DontCreateRedirectors | REN_NonTransactional);
-	TempVM->MarkAsGarbage();
-
-	RemoveExpressions(ExpressionsToRemove);
-
-	return ExpressionsToRemove.Num() > 0;
-}
 
 void FRigVMParserAST::FoldLiterals()
 {
diff --git a/Engine/Source/Developer/RigVMDeveloper/Public/RigVMCompiler/RigVMAST.h b/Engine/Source/Developer/RigVMDeveloper/Public/RigVMCompiler/RigVMAST.h
--- a/Engine/Source/Developer/RigVMDeveloper/Public/RigVMCompiler/RigVMAST.h
+++ b/Engine/Source/Developer/RigVMDeveloper/Public/RigVMCompiler/RigVMAST.h
@@ -1079,10 +1079,6 @@ struct RIGVMDEVELOPER_API FRigVMParserASTSettings
 	UPROPERTY(EditAnywhere, Category = "AST")
 	bool bFoldLiterals = false;
 
-	// fold / remove unreachable / constant branches
-	UPROPERTY(EditAnywhere, Category = "AST")
-	bool bFoldConstantBranches = false;
-
 	// links to be ignored during the parse
 	UPROPERTY()
 	TArray<TObjectPtr<URigVMLink>> LinksToSkip;
@@ -1098,7 +1094,6 @@ struct RIGVMDEVELOPER_API FRigVMParserASTSettings
 		FRigVMParserASTSettings Settings;
 		Settings.bFoldAssignments = false;
 		Settings.bFoldLiterals = false;
-		Settings.bFoldConstantBranches = false;
 		Settings.ExecuteContextStruct = InExecuteContextStruct ? InExecuteContextStruct : FRigVMExecuteContext::StaticStruct();
 		return Settings;
 	}
@@ -1110,7 +1105,6 @@ struct RIGVMDEVELOPER_API FRigVMParserASTSettings
 		FRigVMParserASTSettings Settings;
 		Settings.bFoldAssignments = true;
 		Settings.bFoldLiterals = true;
-		Settings.bFoldConstantBranches = true;
 		Settings.ExecuteContextStruct = InExecuteContextStruct ? InExecuteContextStruct : FRigVMExecuteContext::StaticStruct();
 		return Settings;
 	}
@@ -1319,9 +1313,6 @@ private:
 	// helper function to fold / remove obsolete assignments and reduce assignment chains
 	void FoldAssignments();
 
-	// helper function to fold constant values into literals
-	bool FoldConstantValuesToLiterals(TArray<URigVMGraph*> InGraphs, URigVMController* InController, const TArray<FRigVMExternalVariable>& InExternalVariables);
-
 	// helper function to inline all contributing nodes of the graph
 	void Inline(TArray<URigVMGraph*> InGraphs);
 	void Inline(TArray<URigVMGraph*> InGraphs, const TArray<FRigVMASTProxy>& InNodeProxies);
