Control Rig: Fix slowness when working on functions

diff --git a/Engine/Plugins/Animation/ControlRig/Source/ControlRigDeveloper/Private/ControlRigBlueprint.cpp b/Engine/Plugins/Animation/ControlRig/Source/ControlRigDeveloper/Private/ControlRigBlueprint.cpp
--- a/Engine/Plugins/Animation/ControlRig/Source/ControlRigDeveloper/Private/ControlRigBlueprint.cpp
+++ b/Engine/Plugins/Animation/ControlRig/Source/ControlRigDeveloper/Private/ControlRigBlueprint.cpp
@@ -727,6 +727,8 @@ void UControlRigBlueprint::PreSave(FObjectPreSaveContext ObjectSaveContext)
 {
 	Super::PreSave(ObjectSaveContext);
 
+	RigVMClient.PreSave();
+
 	SupportedEventNames.Reset();
 	if (UControlRigBlueprintGeneratedClass* RigClass = GetControlRigBlueprintGeneratedClass())
 	{
diff --git a/Engine/Source/Developer/RigVMDeveloper/Private/RigVMModel/RigVMClient.cpp b/Engine/Source/Developer/RigVMDeveloper/Private/RigVMModel/RigVMClient.cpp
--- a/Engine/Source/Developer/RigVMDeveloper/Private/RigVMModel/RigVMClient.cpp
+++ b/Engine/Source/Developer/RigVMDeveloper/Private/RigVMModel/RigVMClient.cpp
@@ -685,6 +685,17 @@ void FRigVMClient::PostDuplicateHost(const FString& InOldPathName, const FString
 	}
 }
 
+void FRigVMClient::PreSave()
+{
+	for (URigVMNode* Node : FunctionLibrary->GetNodes())
+	{
+		if (URigVMLibraryNode* LibraryNode = Cast<URigVMLibraryNode>(Node))
+		{
+			UpdateGraphFunctionSerializedGraph(LibraryNode);
+		}
+	}
+}
+
 void FRigVMClient::HandleGraphModifiedEvent(ERigVMGraphNotifType InNotifType, URigVMGraph* InGraph, UObject* InSubject)
 {
 	if (bIgnoreModelNotifications)
@@ -710,14 +721,12 @@ void FRigVMClient::HandleGraphModifiedEvent(ERigVMGraphNotifType InNotifType, UR
 					if (GetOuter()->Implements<URigVMClientHost>())
 					{
 						FunctionStore->AddFunction(CollapseNode->GetFunctionHeader(FunctionHost), false);
-						UpdateGraphFunctionSerializedGraph(CollapseNode);
 					}
 				}
 			}
 			// A node was added into the contained graph of a function
 			else if(URigVMLibraryNode* LibraryNode = Cast<URigVMNode>(InSubject)->FindFunctionForNode())
 			{
-				UpdateGraphFunctionSerializedGraph(LibraryNode);
 				DirtyGraphFunctionCompilationData(LibraryNode);
 				if (URigVMFunctionReferenceNode* FunctionReference = Cast<URigVMFunctionReferenceNode>(InSubject))
 				{
@@ -742,7 +751,6 @@ void FRigVMClient::HandleGraphModifiedEvent(ERigVMGraphNotifType InNotifType, UR
 			// A node was added into the contained graph of a function
 			else if(URigVMLibraryNode* LibraryNode = Cast<URigVMNode>(InSubject)->FindFunctionForNode())
 			{
-				UpdateGraphFunctionSerializedGraph(LibraryNode);
 				DirtyGraphFunctionCompilationData(LibraryNode);
 				if (URigVMFunctionReferenceNode* FunctionReference = Cast<URigVMFunctionReferenceNode>(InSubject))
 				{
@@ -759,7 +767,6 @@ void FRigVMClient::HandleGraphModifiedEvent(ERigVMGraphNotifType InNotifType, UR
 		{
 			if(URigVMLibraryNode* LibraryNode = Cast<URigVMNode>(InSubject)->FindFunctionForNode())
 			{
-				UpdateGraphFunctionSerializedGraph(LibraryNode);
 				DirtyGraphFunctionCompilationData(LibraryNode);
 				UpdateExternalVariablesForFunction(LibraryNode);				
 			}
@@ -783,10 +790,6 @@ void FRigVMClient::HandleGraphModifiedEvent(ERigVMGraphNotifType InNotifType, UR
 					UpdateGraphFunctionData(CollapseNode);
 				}
 			}
-			else if(URigVMLibraryNode* LibraryNode = Cast<URigVMNode>(InSubject)->FindFunctionForNode())
-			{
-				UpdateGraphFunctionSerializedGraph(LibraryNode);
-			}
 			break;	
 		}
 		case ERigVMGraphNotifType::NodeColorChanged: // A node's color has changed (Subject == URigVMNode)
@@ -801,21 +804,8 @@ void FRigVMClient::HandleGraphModifiedEvent(ERigVMGraphNotifType InNotifType, UR
 					UpdateGraphFunctionData(CollapseNode);
 				}
 			}
-			else if(URigVMLibraryNode* LibraryNode = Cast<URigVMNode>(InSubject)->FindFunctionForNode())
-			{
-				UpdateGraphFunctionSerializedGraph(LibraryNode);
-			}
 			break;
 		}
-		case ERigVMGraphNotifType::NodeReferenceChanged: // A node has changed it's referenced function (Subject == URigVMFunctionReferenceNode)
-		case ERigVMGraphNotifType::LibraryTemplateChanged: // The definition of a library node's template has changed (Subject == URigVMLibraryNode)
-		{
-			if(URigVMLibraryNode* LibraryNode = Cast<URigVMNode>(InSubject)->FindFunctionForNode())
-			{
-				UpdateGraphFunctionSerializedGraph(LibraryNode);
-			}
-			break;	
-		}
 		
 		case ERigVMGraphNotifType::PinAdded: // A pin has been added to a given node (Subject == URigVMPin)
 		case ERigVMGraphNotifType::PinRemoved: // A pin has been removed from a given node (Subject == URigVMPin)
@@ -830,7 +820,6 @@ void FRigVMClient::HandleGraphModifiedEvent(ERigVMGraphNotifType InNotifType, UR
 			{
 				if (URigVMLibraryNode* LibraryNode = Node->FindFunctionForNode())
 				{
-					UpdateGraphFunctionSerializedGraph(LibraryNode);
 					DirtyGraphFunctionCompilationData(LibraryNode);
 				}
 				if(Node->GetOuter()->IsA<URigVMFunctionLibrary>())
@@ -853,7 +842,6 @@ void FRigVMClient::HandleGraphModifiedEvent(ERigVMGraphNotifType InNotifType, UR
 				{
 					if (URigVMLibraryNode* LibraryNode = OuterNode->FindFunctionForNode())
 					{
-						UpdateGraphFunctionSerializedGraph(LibraryNode);
 						DirtyGraphFunctionCompilationData(LibraryNode);
 					}
 				}
diff --git a/Engine/Source/Developer/RigVMDeveloper/Private/RigVMModel/RigVMController.cpp b/Engine/Source/Developer/RigVMDeveloper/Private/RigVMModel/RigVMController.cpp
--- a/Engine/Source/Developer/RigVMDeveloper/Private/RigVMModel/RigVMController.cpp
+++ b/Engine/Source/Developer/RigVMDeveloper/Private/RigVMModel/RigVMController.cpp
@@ -3939,15 +3939,12 @@ TMap<FRigVMGraphFunctionIdentifier, URigVMLibraryNode*> URigVMController::Locali
 		FRigVMControllerGraphGuard GraphGuard(this, ThisLibrary, bSetupUndoRedo);
 		for(FRigVMGraphFunctionData* FunctionToLocalize : FunctionsToLocalize)
 		{
-			if (FunctionToLocalize->SerializedCollapsedNode.IsEmpty())
+			if (URigVMLibraryNode* ReferencedFunction = Cast<URigVMLibraryNode>(FunctionToLocalize->Header.LibraryPointer.LibraryNode.TryLoad()))
 			{
-				if (URigVMLibraryNode* ReferencedFunction = Cast<URigVMLibraryNode>(FunctionToLocalize->Header.LibraryPointer.LibraryNode.TryLoad()))
+				if (IRigVMClientHost* ClientHost = ReferencedFunction->GetImplementingOuter<IRigVMClientHost>())
 				{
-					if (IRigVMClientHost* ClientHost = ReferencedFunction->GetImplementingOuter<IRigVMClientHost>())
-					{
-						ClientHost->GetRigVMClient()->UpdateGraphFunctionSerializedGraph(ReferencedFunction);
-					}
-				}				
+					ClientHost->GetRigVMClient()->UpdateGraphFunctionSerializedGraph(ReferencedFunction);
+				}
 			}
 			
 			TArray<FName> NodeNames = ImportNodesFromText(FunctionToLocalize->SerializedCollapsedNode, false);
diff --git a/Engine/Source/Developer/RigVMDeveloper/Public/RigVMModel/RigVMClient.h b/Engine/Source/Developer/RigVMDeveloper/Public/RigVMModel/RigVMClient.h
--- a/Engine/Source/Developer/RigVMDeveloper/Public/RigVMModel/RigVMClient.h
+++ b/Engine/Source/Developer/RigVMDeveloper/Public/RigVMModel/RigVMClient.h
@@ -145,6 +145,9 @@ public:
 	// work to be done after a duplication of the source asset
 	void PostDuplicateHost(const FString& InOldPathName, const FString& InNewPathName);
 
+	// work to be done before saving
+	void PreSave();
+
 	void HandleGraphModifiedEvent(ERigVMGraphNotifType InNotifType, URigVMGraph* InGraph, UObject* InSubject);
 
 	FRigVMGraphFunctionStore* FindFunctionStore(const URigVMLibraryNode* InLibraryNode);
