Control Rig: Fix paste function not available in MyBlueprint

diff --git a/Engine/Source/Developer/RigVMDeveloper/Private/RigVMModel/RigVMController.cpp b/Engine/Source/Developer/RigVMDeveloper/Private/RigVMModel/RigVMController.cpp
--- a/Engine/Source/Developer/RigVMDeveloper/Private/RigVMModel/RigVMController.cpp
+++ b/Engine/Source/Developer/RigVMDeveloper/Private/RigVMModel/RigVMController.cpp
@@ -3385,13 +3385,39 @@ bool URigVMController::CanImportNodesFromText(const FString& InText)
 		return false;
 	}
 
-	if (GetGraph()->IsA<URigVMFunctionLibrary>())
+	FRigVMControllerObjectFactory Factory(this);
+	if (!Factory.CanCreateObjectsFromText(InText))
 	{
 		return false;
 	}
 
-	FRigVMControllerObjectFactory Factory(nullptr);
-	return Factory.CanCreateObjectsFromText(InText);
+	URigVMGraph* Graph = GetGraph();
+	check(Graph);
+	Factory.ProcessBuffer(Graph, RF_Transactional, InText);
+
+	if (Factory.CreatedNodes.Num() == 0)
+	{
+		return false;
+	}
+
+	// If the graph is a function library, make sure we are pasting unconnected collapse nodes
+	if (Graph->IsA<URigVMFunctionLibrary>())
+	{
+		if (!Factory.CreatedLinks.IsEmpty())
+		{
+			return false;
+		}
+		
+		for (URigVMNode* Node : Factory.CreatedNodes)
+		{
+			if (!Node->IsA<URigVMCollapseNode>())
+			{
+				return false;
+			}
+		}
+	}
+
+	return true;
 }
 
 TArray<FName> URigVMController::ImportNodesFromText(const FString& InText, bool bSetupUndoRedo, bool bPrintPythonCommands)
