Control Rig: Move type postprocess and redirectors to RigVMTypeUtils

diff --git a/Engine/Plugins/Animation/ControlRig/Source/ControlRigDeveloper/Private/ControlRigBlueprint.cpp b/Engine/Plugins/Animation/ControlRig/Source/ControlRigDeveloper/Private/ControlRigBlueprint.cpp
--- a/Engine/Plugins/Animation/ControlRig/Source/ControlRigDeveloper/Private/ControlRigBlueprint.cpp
+++ b/Engine/Plugins/Animation/ControlRig/Source/ControlRigDeveloper/Private/ControlRigBlueprint.cpp
@@ -81,7 +81,7 @@ FEdGraphPinType FControlRigPublicFunctionArg::GetPinType() const
 	
 	if(CPPTypeObjectPath.IsValid())
 	{
-		Variable.TypeObject = URigVMPin::FindObjectFromCPPTypeObjectPath(CPPTypeObjectPath.ToString());
+		Variable.TypeObject = RigVMTypeUtils::FindObjectFromCPPTypeObjectPath(CPPTypeObjectPath.ToString());
 	}
 
 	return RigVMTypeUtils::PinTypeFromExternalVariable(Variable);
@@ -94,7 +94,7 @@ bool FControlRigPublicFunctionData::IsMutable() const
 		if(!Arg.CPPTypeObjectPath.IsNone())
 		{
 			if(UScriptStruct* Struct = Cast<UScriptStruct>(
-				URigVMPin::FindObjectFromCPPTypeObjectPath(Arg.CPPTypeObjectPath.ToString())))
+				RigVMTypeUtils::FindObjectFromCPPTypeObjectPath(Arg.CPPTypeObjectPath.ToString())))
 			{
 				if(Struct->IsChildOf(FRigVMExecuteContext::StaticStruct()))
 				{
@@ -2861,13 +2861,13 @@ bool UControlRigBlueprint::ChangeMemberVariableType(const FName& InName, const F
 		Variable.TypeName = *CPPType;
 		Variable.Size = sizeof(FName);
 	}
-	else if(UScriptStruct* ScriptStruct = URigVMPin::FindObjectFromCPPTypeObjectPath<UScriptStruct>(CPPType))
+	else if(UScriptStruct* ScriptStruct = RigVMTypeUtils::FindObjectFromCPPTypeObjectPath<UScriptStruct>(CPPType))
 	{
 		Variable.TypeName = *RigVMTypeUtils::GetUniqueStructTypeName(ScriptStruct);
 		Variable.TypeObject = ScriptStruct;
 		Variable.Size = ScriptStruct->GetStructureSize();
 	}
-	else if (UEnum* Enum= URigVMPin::FindObjectFromCPPTypeObjectPath<UEnum>(CPPType))
+	else if (UEnum* Enum= RigVMTypeUtils::FindObjectFromCPPTypeObjectPath<UEnum>(CPPType))
 	{
 		Variable.TypeName = *RigVMTypeUtils::CPPTypeFromEnum(Enum);
 		Variable.TypeObject = Enum;
@@ -3176,7 +3176,7 @@ void UControlRigBlueprint::PopulateModelFromGraphForBackwardsCompatibility(UCont
 
 				URigVMNode* ModelNode = nullptr;
 
-				UScriptStruct* UnitStruct = URigVMPin::FindObjectFromCPPTypeObjectPath<UScriptStruct>(StructPath);
+				UScriptStruct* UnitStruct = RigVMTypeUtils::FindObjectFromCPPTypeObjectPath<UScriptStruct>(StructPath);
 				if (UnitStruct && UnitStruct->IsChildOf(FRigVMStruct::StaticStruct()))
 				{ 
 					ModelNode = GetOrCreateController()->AddUnitNode(UnitStruct, FRigUnit::GetMethodName(), NodePosition, PropertyName.ToString(), false);
diff --git a/Engine/Source/Developer/RigVMDeveloper/Private/RigVMCompiler/RigVMCompiler.cpp b/Engine/Source/Developer/RigVMDeveloper/Private/RigVMCompiler/RigVMCompiler.cpp
--- a/Engine/Source/Developer/RigVMDeveloper/Private/RigVMCompiler/RigVMCompiler.cpp
+++ b/Engine/Source/Developer/RigVMDeveloper/Private/RigVMCompiler/RigVMCompiler.cpp
@@ -1162,10 +1162,12 @@ bool URigVMCompiler::CompileFunction(const URigVMLibraryNode* InLibraryNode, URi
 				{
 					for(TPair<FRigVMGraphFunctionIdentifier, uint32>& Pair : Data->Header.Dependencies)
 					{
-						IRigVMGraphFunctionHost* ReferencedFunctionHost = Cast<IRigVMGraphFunctionHost>(Pair.Key.HostObject.ResolveObject());
-						if (FRigVMGraphFunctionData* ReferencedData = ReferencedFunctionHost->GetRigVMGraphFunctionStore()->FindFunction(Pair.Key))
+						if (IRigVMGraphFunctionHost* ReferencedFunctionHost = Cast<IRigVMGraphFunctionHost>(Pair.Key.HostObject.ResolveObject()))
 						{
-							Pair.Value = ReferencedData->CompilationData.Hash;
+							if (FRigVMGraphFunctionData* ReferencedData = ReferencedFunctionHost->GetRigVMGraphFunctionStore()->FindFunction(Pair.Key))
+							{
+								Pair.Value = ReferencedData->CompilationData.Hash;
+							}
 						}
 					}
 				}
diff --git a/Engine/Source/Developer/RigVMDeveloper/Private/RigVMDeveloperTypeUtils.cpp b/Engine/Source/Developer/RigVMDeveloper/Private/RigVMDeveloperTypeUtils.cpp
--- a/Engine/Source/Developer/RigVMDeveloper/Private/RigVMDeveloperTypeUtils.cpp
+++ b/Engine/Source/Developer/RigVMDeveloper/Private/RigVMDeveloperTypeUtils.cpp
@@ -182,13 +182,13 @@ FRigVMExternalVariable RigVMTypeUtils::ExternalVariableFromCPPTypePath(const FNa
 		Variable.TypeName = *CPPTypePath;
 		Variable.Size = sizeof(FName);
 	}
-	else if(UScriptStruct* ScriptStruct = URigVMPin::FindObjectFromCPPTypeObjectPath<UScriptStruct>(CPPTypePath))
+	else if(UScriptStruct* ScriptStruct = RigVMTypeUtils::FindObjectFromCPPTypeObjectPath<UScriptStruct>(CPPTypePath))
 	{
 		Variable.TypeName = *RigVMTypeUtils::GetUniqueStructTypeName(ScriptStruct);
 		Variable.TypeObject = ScriptStruct;
 		Variable.Size = ScriptStruct->GetStructureSize();
 	}
-	else if (UEnum* Enum= URigVMPin::FindObjectFromCPPTypeObjectPath<UEnum>(CPPTypePath))
+	else if (UEnum* Enum= RigVMTypeUtils::FindObjectFromCPPTypeObjectPath<UEnum>(CPPTypePath))
 	{
 		Variable.TypeName = *CPPTypeFromEnum(Enum);
 		Variable.TypeObject = Enum;
@@ -251,13 +251,13 @@ FRigVMExternalVariable RigVMTypeUtils::ExternalVariableFromCPPType(const FName&
 		Variable.TypeName = *CPPType;
 		Variable.Size = sizeof(FName);
 	}
-	else if(UScriptStruct* ScriptStruct = URigVMPin::FindObjectFromCPPTypeObjectPath<UScriptStruct>(CPPType))
+	else if(UScriptStruct* ScriptStruct = RigVMTypeUtils::FindObjectFromCPPTypeObjectPath<UScriptStruct>(CPPType))
 	{
 		Variable.TypeName = *RigVMTypeUtils::GetUniqueStructTypeName(ScriptStruct);
 		Variable.TypeObject = ScriptStruct;
 		Variable.Size = ScriptStruct->GetStructureSize();
 	}
-	else if (UEnum* Enum= URigVMPin::FindObjectFromCPPTypeObjectPath<UEnum>(CPPType))
+	else if (UEnum* Enum= RigVMTypeUtils::FindObjectFromCPPTypeObjectPath<UEnum>(CPPType))
 	{
 		Variable.TypeName = *CPPTypeFromEnum(Enum);
 		Variable.TypeObject = Enum;
diff --git a/Engine/Source/Developer/RigVMDeveloper/Private/RigVMModel/Nodes/RigVMTemplateNode.cpp b/Engine/Source/Developer/RigVMDeveloper/Private/RigVMModel/Nodes/RigVMTemplateNode.cpp
--- a/Engine/Source/Developer/RigVMDeveloper/Private/RigVMModel/Nodes/RigVMTemplateNode.cpp
+++ b/Engine/Source/Developer/RigVMDeveloper/Private/RigVMModel/Nodes/RigVMTemplateNode.cpp
@@ -25,19 +25,32 @@ void FRigVMTemplatePreferredType::UpdateIndexFromString()
 		return;
 	}
 
-	FString Left, Right;
-	verify(TypeString.Split(TEXT(","), &Left, &Right));
+	FString OriginalCPPType, CPPTypeObjectPath;
+	verify(TypeString.Split(TEXT(","), &OriginalCPPType, &CPPTypeObjectPath));
+	FString CPPType = OriginalCPPType;
 
 	UObject* CPPTypeObject = nullptr;
 
 	static const FString NoneString = FName(NAME_None).ToString(); 
-	if(Right != NoneString)
+	if(CPPTypeObjectPath != NoneString)
 	{
-		CPPTypeObject = URigVMPin::FindObjectFromCPPTypeObjectPath(Right);
+		CPPTypeObject = RigVMTypeUtils::FindObjectFromCPPTypeObjectPath(CPPTypeObjectPath);
 	}
 
-	const FRigVMTemplateArgumentType Type(*Left, CPPTypeObject);
+	// If we still haven't found the object, try to find it with the CPPType
+	if (!CPPTypeObject && RigVMTypeUtils::RequiresCPPTypeObject(CPPType))
+	{
+		CPPTypeObject = RigVMTypeUtils::ObjectFromCPPType(CPPType, true);
+	}
+
+	const FRigVMTemplateArgumentType Type(*CPPType, CPPTypeObject);
 	TypeIndex = FRigVMRegistry::Get().FindOrAddType(Type);
+
+	// We might have used redirectors. Update the TypeString with the new information.
+	if (TypeIndex != INDEX_NONE && OriginalCPPType != Type.CPPType.ToString())
+	{
+		UpdateStringFromIndex();
+	}
 }
 
 ////////////////////////////////////////////////////////////////////////////////////////////////
@@ -94,6 +107,7 @@ void URigVMTemplateNode::PostLoad()
 			FString ArgName, CPPType;
 			PreferredPermutation.Split(TEXT(":"), &ArgName, &CPPType);
 
+			CPPType = RigVMTypeUtils::PostProcessCPPType(CPPType);
 			const FRigVMTemplateArgumentType Type = Registry.FindTypeFromCPPType(CPPType);
 			PreferredPermutationPairs.Emplace(*ArgName, Registry.GetTypeIndex(Type));
 		}
diff --git a/Engine/Source/Developer/RigVMDeveloper/Private/RigVMModel/RigVMController.cpp b/Engine/Source/Developer/RigVMDeveloper/Private/RigVMModel/RigVMController.cpp
--- a/Engine/Source/Developer/RigVMDeveloper/Private/RigVMModel/RigVMController.cpp
+++ b/Engine/Source/Developer/RigVMDeveloper/Private/RigVMModel/RigVMController.cpp
@@ -947,18 +947,7 @@ URigVMUnitNode* URigVMController::AddUnitNodeFromStructPath(const FString& InScr
 		return nullptr;
 	}
 
-	UScriptStruct* ScriptStruct = URigVMPin::FindObjectFromCPPTypeObjectPath<UScriptStruct>(InScriptStructPath);
-	if (ScriptStruct == nullptr)
-	{
-		const FCoreRedirectObjectName OldObjectName(InScriptStructPath);
-		const FCoreRedirectObjectName NewObjectName = FCoreRedirects::GetRedirectedName(ECoreRedirectFlags::Type_Struct, OldObjectName);
-		if (OldObjectName != NewObjectName)
-		{
-			ScriptStruct = URigVMPin::FindObjectFromCPPTypeObjectPath<UScriptStruct>(NewObjectName.ToString());
-		}
-		
-	}
-	
+	UScriptStruct* ScriptStruct = RigVMTypeUtils::FindObjectFromCPPTypeObjectPath<UScriptStruct>(InScriptStructPath);
 	if (ScriptStruct == nullptr)
 	{
 		ReportErrorf(TEXT("Cannot find struct for path '%s'."), *InScriptStructPath);
@@ -1176,7 +1165,7 @@ URigVMVariableNode* URigVMController::AddVariableNode(const FName& InVariableNam
 	}
 	if (InCPPTypeObject == nullptr)
 	{
-		InCPPTypeObject = URigVMPin::FindObjectFromCPPTypeObjectPath<UObject>(InCPPType);
+		InCPPTypeObject = RigVMTypeUtils::FindObjectFromCPPTypeObjectPath<UObject>(InCPPType);
 	}
 
 	FString CPPType = RigVMTypeUtils::PostProcessCPPType(InCPPType, InCPPTypeObject);
@@ -1311,7 +1300,7 @@ URigVMVariableNode* URigVMController::AddVariableNodeFromObjectPath(const FName&
 	UObject* CPPTypeObject = nullptr;
 	if (!InCPPTypeObjectPath.IsEmpty())
 	{
-		CPPTypeObject = URigVMPin::FindObjectFromCPPTypeObjectPath<UObject>(InCPPTypeObjectPath);
+		CPPTypeObject = RigVMTypeUtils::FindObjectFromCPPTypeObjectPath<UObject>(InCPPTypeObjectPath);
 		if (CPPTypeObject == nullptr)
 		{
 			ReportErrorf(TEXT("Cannot find cpp type object for path '%s'."), *InCPPTypeObjectPath);
@@ -1702,7 +1691,7 @@ void URigVMController::OnExternalVariableTypeChangedFromObjectPath(const FName&
 	UObject* CPPTypeObject = nullptr;
 	if (!InCPPTypeObjectPath.IsEmpty())
 	{
-		CPPTypeObject = URigVMPin::FindObjectFromCPPTypeObjectPath<UObject>(InCPPTypeObjectPath);
+		CPPTypeObject = RigVMTypeUtils::FindObjectFromCPPTypeObjectPath<UObject>(InCPPTypeObjectPath);
 		if (CPPTypeObject == nullptr)
 		{
 			ReportErrorf(TEXT("Cannot find cpp type object for path '%s'."), *InCPPTypeObjectPath);
@@ -2220,7 +2209,7 @@ URigVMParameterNode* URigVMController::AddParameterNodeFromObjectPath(const FNam
 	UObject* CPPTypeObject = nullptr;
 	if (!InCPPTypeObjectPath.IsEmpty())
 	{
-		CPPTypeObject = URigVMPin::FindObjectFromCPPTypeObjectPath<UObject>(InCPPTypeObjectPath);
+		CPPTypeObject = RigVMTypeUtils::FindObjectFromCPPTypeObjectPath<UObject>(InCPPTypeObjectPath);
 		if (CPPTypeObject == nullptr)
 		{
 			ReportErrorf(TEXT("Cannot find cpp type object for path '%s'."), *InCPPTypeObjectPath);
@@ -2726,7 +2715,7 @@ URigVMInjectionInfo* URigVMController::AddInjectedNodeFromStructPath(const FStri
 		return nullptr;
 	}
 
-	UScriptStruct* ScriptStruct = URigVMPin::FindObjectFromCPPTypeObjectPath<UScriptStruct>(InScriptStructPath);
+	UScriptStruct* ScriptStruct = RigVMTypeUtils::FindObjectFromCPPTypeObjectPath<UScriptStruct>(InScriptStructPath);
 	if (ScriptStruct == nullptr)
 	{
 		ReportErrorf(TEXT("Cannot find struct for path '%s'."), *InScriptStructPath);
@@ -10089,7 +10078,7 @@ FName URigVMController::AddExposedPin(const FName& InPinName, ERigVMPinDirection
 		}
 		if (CPPTypeObject == nullptr)
 		{
-			CPPTypeObject = URigVMPin::FindObjectFromCPPTypeObjectPath<UObject>(InCPPTypeObjectPath.ToString());
+			CPPTypeObject = RigVMTypeUtils::FindObjectFromCPPTypeObjectPath<UObject>(InCPPTypeObjectPath.ToString());
 		}
 	}
 
@@ -11045,7 +11034,7 @@ bool URigVMController::ChangeExposedPinType(const FName& InPinName, const FStrin
 	UObject* CPPTypeObject = nullptr;
 	if (!InCPPTypeObjectPath.IsNone())
 	{
-		CPPTypeObject = URigVMPin::FindObjectFromCPPTypeObjectPath<UObject>(InCPPTypeObjectPath.ToString());
+		CPPTypeObject = RigVMTypeUtils::FindObjectFromCPPTypeObjectPath<UObject>(InCPPTypeObjectPath.ToString());
 		if(CPPTypeObject)
 		{
 			if(const UScriptStruct* CPPTypeStruct = Cast<UScriptStruct>(CPPTypeObject))
@@ -12008,7 +11997,7 @@ FRigVMGraphVariableDescription URigVMController::AddLocalVariableFromObjectPath(
 	UObject* CPPTypeObject = nullptr;
 	if (!InCPPTypeObjectPath.IsEmpty())
 	{
-		CPPTypeObject = URigVMPin::FindObjectFromCPPTypeObjectPath<UObject>(InCPPTypeObjectPath);
+		CPPTypeObject = RigVMTypeUtils::FindObjectFromCPPTypeObjectPath<UObject>(InCPPTypeObjectPath);
 		if (CPPTypeObject == nullptr)
 		{
 			ReportErrorf(TEXT("Cannot find cpp type object for path '%s'."), *InCPPTypeObjectPath);
@@ -12364,7 +12353,7 @@ bool URigVMController::SetLocalVariableTypeFromObjectPath(const FName& InVariabl
 	UObject* CPPTypeObject = nullptr;
 	if (!InCPPTypeObjectPath.IsEmpty())
 	{
-		CPPTypeObject = URigVMPin::FindObjectFromCPPTypeObjectPath<UObject>(InCPPTypeObjectPath);
+		CPPTypeObject = RigVMTypeUtils::FindObjectFromCPPTypeObjectPath<UObject>(InCPPTypeObjectPath);
 		if (CPPTypeObject == nullptr)
 		{
 			ReportErrorf(TEXT("Cannot find cpp type object for path '%s'."), *InCPPTypeObjectPath);
@@ -12875,7 +12864,7 @@ URigVMNode* URigVMController::AddIfNode(const FString& InCPPType, const FName& I
 	UObject* CPPTypeObject = nullptr;
 	if(!InCPPTypeObjectPath.IsNone())
 	{
-		CPPTypeObject = URigVMPin::FindObjectFromCPPTypeObjectPath<UObject>(InCPPTypeObjectPath.ToString());
+		CPPTypeObject = RigVMTypeUtils::FindObjectFromCPPTypeObjectPath<UObject>(InCPPTypeObjectPath.ToString());
 		if (CPPTypeObject == nullptr)
 		{
 			ReportErrorf(TEXT("Cannot find cpp type object for path '%s'."), *InCPPTypeObjectPath.ToString());
@@ -12937,7 +12926,7 @@ URigVMNode* URigVMController::AddSelectNode(const FString& InCPPType, const FNam
 	UObject* CPPTypeObject = nullptr;
 	if (!InCPPTypeObjectPath.IsNone())
 	{
-		CPPTypeObject = URigVMPin::FindObjectFromCPPTypeObjectPath<UObject>(InCPPTypeObjectPath.ToString());
+		CPPTypeObject = RigVMTypeUtils::FindObjectFromCPPTypeObjectPath<UObject>(InCPPTypeObjectPath.ToString());
 		if (CPPTypeObject == nullptr)
 		{
 			ReportErrorf(TEXT("Cannot find cpp type object for path '%s'."), *InCPPTypeObjectPath.ToString());
@@ -13190,7 +13179,7 @@ URigVMEnumNode* URigVMController::AddEnumNode(const FName& InCPPTypeObjectPath,
 	URigVMGraph* Graph = GetGraph();
 	check(Graph);
 
-    UObject* CPPTypeObject = URigVMPin::FindObjectFromCPPTypeObjectPath<UObject>(InCPPTypeObjectPath.ToString());
+    UObject* CPPTypeObject = RigVMTypeUtils::FindObjectFromCPPTypeObjectPath<UObject>(InCPPTypeObjectPath.ToString());
 	if (CPPTypeObject == nullptr)
 	{
 		ReportErrorf(TEXT("Cannot find cpp type object for path '%s'."), *InCPPTypeObjectPath.ToString());
@@ -13353,7 +13342,7 @@ URigVMNode* URigVMController::AddArrayNodeFromObjectPath(ERigVMOpCode InOpCode,
 	UObject* CPPTypeObject = nullptr;
 	if (!InCPPTypeObjectPath.IsEmpty())
 	{
-		CPPTypeObject = URigVMPin::FindObjectFromCPPTypeObjectPath<UObject>(InCPPTypeObjectPath);
+		CPPTypeObject = RigVMTypeUtils::FindObjectFromCPPTypeObjectPath<UObject>(InCPPTypeObjectPath);
 		if (CPPTypeObject == nullptr)
 		{
 			ReportErrorf(TEXT("Cannot find cpp type object for path '%s'."), *InCPPTypeObjectPath);
@@ -16602,7 +16591,7 @@ bool URigVMController::ResolveWildCardPin(const FString& InPinPath, const FStrin
 	UObject* CPPTypeObject = nullptr;
 	if (!InCPPTypeObjectPath.IsNone())
 	{
-		CPPTypeObject = URigVMPin::FindObjectFromCPPTypeObjectPath<UObject>(InCPPTypeObjectPath.ToString());
+		CPPTypeObject = RigVMTypeUtils::FindObjectFromCPPTypeObjectPath<UObject>(InCPPTypeObjectPath.ToString());
 		if (CPPTypeObject == nullptr)
 		{
 			ReportErrorf(TEXT("Cannot find cpp type object for path '%s'."), *InCPPTypeObjectPath.ToString());
@@ -18474,7 +18463,7 @@ bool URigVMController::ChangePinType(URigVMPin* InPin, const FString& InCPPType,
 		return false;
 	}
 
-	UObject* CPPTypeObject = URigVMPin::FindObjectFromCPPTypeObjectPath(InCPPTypeObjectPath.ToString());
+	UObject* CPPTypeObject = RigVMTypeUtils::FindObjectFromCPPTypeObjectPath(InCPPTypeObjectPath.ToString());
 
 	// always refresh pin type if it is a user defined struct, whose internal layout can change at anytime
 	bool bForceRefresh = false;
@@ -18506,7 +18495,7 @@ bool URigVMController::ChangePinType(URigVMPin* InPin, const FString& InCPPType,
 		return false;
 	}
 
-	if (FRigVMPropertyDescription::RequiresCPPTypeObject(InCPPType) && !InCPPTypeObject)
+	if (RigVMTypeUtils::RequiresCPPTypeObject(InCPPType) && !InCPPTypeObject)
 	{
 		return false;
 	}
@@ -18968,90 +18957,27 @@ void URigVMController::AddSubPin(URigVMPin* InParentPin, URigVMPin* InPin)
 	InParentPin->SubPins.Add(InPin);
 }
 
-
-static UObject* FindObjectGloballyWithRedirectors(const TCHAR* InObjectName)
-{
-	// Do a global search for the CPP type. Note that searching with ANY_PACKAGE _does not_
-	// apply redirectors. So only if this fails do we apply them manually below.
-	UObject* Object = FindFirstObject<UField>(InObjectName, EFindFirstObjectOptions::EnsureIfAmbiguous);
-	if(Object != nullptr)
-	{
-		return Object;
-	}
-
-	FCoreRedirectObjectName NewObjectName;
-	const bool bFoundRedirect = FCoreRedirects::RedirectNameAndValues(
-		ECoreRedirectFlags::Type_Class | ECoreRedirectFlags::Type_Struct | ECoreRedirectFlags::Type_Enum,
-		FCoreRedirectObjectName(InObjectName),
-		NewObjectName,
-		nullptr,
-		ECoreRedirectMatchFlags::AllowPartialMatch); // AllowPartialMatch to allow redirects from one package to another (see /Script/ControlRig.CRFourPointBezier -> /Script/RigVM.RigVMFourPointBezier)
-
-	if (!bFoundRedirect)
-	{
-		return nullptr;
-	}
-
-	const FString RedirectedObjectName = NewObjectName.ObjectName.ToString();
-	UPackage *Package = nullptr;
-	if (!NewObjectName.PackageName.IsNone())
-	{
-		Package = FindPackage(nullptr, *NewObjectName.PackageName.ToString());
-	}
-	if (Package != nullptr)
-	{
-		Object = FindObject<UField>(Package, *RedirectedObjectName);
-	}
-	if (Package == nullptr || Object == nullptr)
-	{
-		// Hail Mary pass.
-		Object = FindFirstObject<UField>(*RedirectedObjectName, EFindFirstObjectOptions::EnsureIfAmbiguous);
-	}
-	return Object;
-}
-
 bool URigVMController::EnsurePinValidity(URigVMPin* InPin, bool bRecursive)
 {
 	check(InPin);
 	
 	// check if the CPPTypeObject is set up correctly.
-	if(FRigVMPropertyDescription::RequiresCPPTypeObject(InPin->GetCPPType()))
+	if(RigVMTypeUtils::RequiresCPPTypeObject(InPin->GetCPPType()))
 	{
 		// GetCPPTypeObject attempts to update pin type information to the latest
 		// without testing for redirector
 		if(InPin->GetCPPTypeObject() == nullptr)
 		{
-			// try to find the CPPTypeObject by name
-			FString CPPType = InPin->IsArray() ? InPin->GetArrayElementCppType() : InPin->GetCPPType();
-
-			UObject* CPPTypeObject = FindObjectGloballyWithRedirectors(*CPPType);
-
-			if (CPPTypeObject == nullptr)
-			{
-				// If we've mistakenly stored the struct type with the 'F', 'U', or 'A' prefixes, we need to strip them
-				// off first. Enums are always named with their prefix intact.
-				if (!CPPType.IsEmpty() && (CPPType[0] == TEXT('F') || CPPType[0] == TEXT('U') || CPPType[0] == TEXT('A')))
-				{
-					CPPType = CPPType.Mid(1);
-				}
-				CPPTypeObject = FindObjectGloballyWithRedirectors(*CPPType);
-			}
-
-			if(CPPTypeObject == nullptr)
-			{
-				const FString Message = FString::Printf(
-					TEXT("%s: Pin '%s' is missing the CPPTypeObject for CPPType '%s'."),
-					*InPin->GetPathName(), *InPin->GetPinPath(), *InPin->GetCPPType());
-				FScriptExceptionHandler::Get().HandleException(ELogVerbosity::Error, *Message, *FString());
-				return false;
-			}
-			
-			InPin->CPPTypeObject = CPPTypeObject;
+			FString CPPType = InPin->GetCPPType();
+			InPin->CPPTypeObject = RigVMTypeUtils::ObjectFromCPPType(CPPType);
+			InPin->CPPType = CPPType;
+		}
+		else
+		{
+			InPin->CPPType = RigVMTypeUtils::PostProcessCPPType(InPin->CPPType, InPin->GetCPPTypeObject());
 		}
 	}
-
-	InPin->CPPType = RigVMTypeUtils::PostProcessCPPType(InPin->CPPType, InPin->GetCPPTypeObject());
-
+	
 	if(bRecursive)
 	{
 		for(URigVMPin* SubPin : InPin->SubPins)
diff --git a/Engine/Source/Developer/RigVMDeveloper/Private/RigVMModel/RigVMPin.cpp b/Engine/Source/Developer/RigVMDeveloper/Private/RigVMModel/RigVMPin.cpp
--- a/Engine/Source/Developer/RigVMDeveloper/Private/RigVMModel/RigVMPin.cpp
+++ b/Engine/Source/Developer/RigVMDeveloper/Private/RigVMModel/RigVMPin.cpp
@@ -838,7 +838,7 @@ bool URigVMPin::IsValidDefaultValue(const FString& InDefaultValue) const
 				return true;
 			}
 			
-			UObject* Object = FindObjectFromCPPTypeObjectPath(Value);
+			UObject* Object = RigVMTypeUtils::FindObjectFromCPPTypeObjectPath(Value);
 			if(Object == nullptr)
 			{
 				return false;
@@ -1053,36 +1053,6 @@ FText URigVMPin::GetToolTipText() const
 	return FText();
 }
 
-UObject* URigVMPin::FindObjectFromCPPTypeObjectPath(const FString& InObjectPath)
-{
-	if (InObjectPath.IsEmpty())
-	{
-		return nullptr;
-	}
-
-	if (InObjectPath == FName(NAME_None).ToString())
-	{
-		return nullptr;
-	}
-
-	// we do this to avoid ambiguous searches for 
-	// common names such as "transform" or "vector"
-	UPackage* Package = nullptr;
-	FString PackageName;
-	FString CPPTypeObjectName = InObjectPath;
-	if (InObjectPath.Split(TEXT("."), &PackageName, &CPPTypeObjectName))
-	{
-		Package = FindPackage(nullptr, *PackageName);
-	}
-	
-	if (UObject* ObjectWithinPackage = FindObject<UObject>(Package, *CPPTypeObjectName))
-	{
-		return ObjectWithinPackage;
-	}
-
-	return FindFirstObject<UObject>(*InObjectPath, EFindFirstObjectOptions::NativeFirst | EFindFirstObjectOptions::EnsureIfAmbiguous);
-}
-
 URigVMVariableNode* URigVMPin::GetBoundVariableNode() const
 {
 	for (TObjectPtr<URigVMInjectionInfo> InjectionInfo : InjectionInfos)
@@ -1331,10 +1301,13 @@ void URigVMPin::UpdateTypeInformationIfRequired() const
 		if (CPPTypeObjectPath != NAME_None)
 		{
 			URigVMPin* MutableThis = (URigVMPin*)this;
-			MutableThis->CPPTypeObject = FindObjectFromCPPTypeObjectPath(CPPTypeObjectPath.ToString());
+			MutableThis->CPPTypeObject = RigVMTypeUtils::FindObjectFromCPPTypeObjectPath(CPPTypeObjectPath.ToString());
 			MutableThis->CPPType = RigVMTypeUtils::PostProcessCPPType(CPPType, CPPTypeObject);
-			MutableThis->LastKnownTypeIndex = FRigVMRegistry::Get().FindOrAddType(FRigVMTemplateArgumentType(*CPPType, CPPTypeObject));
-			MutableThis->LastKnownCPPType = MutableThis->CPPType; 
+			if (!MutableThis->CPPType.IsEmpty())
+			{
+				MutableThis->LastKnownTypeIndex = FRigVMRegistry::Get().FindOrAddType(FRigVMTemplateArgumentType(*CPPType, CPPTypeObject));
+				MutableThis->LastKnownCPPType = MutableThis->CPPType;
+			} 
 		}
 	}
 
diff --git a/Engine/Source/Developer/RigVMDeveloper/Public/RigVMModel/RigVMPin.h b/Engine/Source/Developer/RigVMDeveloper/Public/RigVMModel/RigVMPin.h
--- a/Engine/Source/Developer/RigVMDeveloper/Public/RigVMModel/RigVMPin.h
+++ b/Engine/Source/Developer/RigVMDeveloper/Public/RigVMModel/RigVMPin.h
@@ -438,14 +438,6 @@ public:
 	// Returns true if the pin can be bound to a given variable
 	bool CanBeBoundToVariable(const FRigVMExternalVariable& InExternalVariable, const FString& InSegmentPath = FString()) const;
 
-	// helper function to retrieve an object from a path
-	static UObject* FindObjectFromCPPTypeObjectPath(const FString& InObjectPath);
-	template<class T>
-	static T* FindObjectFromCPPTypeObjectPath(const FString& InObjectPath)
-	{
-		return Cast<T>(FindObjectFromCPPTypeObjectPath(InObjectPath));
-	}
-
 	// Returns true if the pin should not show up on a node, but in the details panel
 	bool ShowInDetailsPanelOnly() const;
 
diff --git a/Engine/Source/Runtime/RigVM/Private/RigVMCore/RigVMMemoryDeprecated.cpp b/Engine/Source/Runtime/RigVM/Private/RigVMCore/RigVMMemoryDeprecated.cpp
--- a/Engine/Source/Runtime/RigVM/Private/RigVMCore/RigVMMemoryDeprecated.cpp
+++ b/Engine/Source/Runtime/RigVM/Private/RigVMCore/RigVMMemoryDeprecated.cpp
@@ -7,6 +7,7 @@
 #include "UObject/Package.h"
 #include "RigVMModule.h"
 #include "UObject/CoreRedirects.h"
+#include "RigVMTypeUtils.h"
 
 #include UE_INLINE_GENERATED_CPP_BY_NAME(RigVMMemoryDeprecated)
 
@@ -596,7 +597,7 @@ void FRigVMMemoryContainer::Load(FArchive& Ar)
 
 	for (const FString& ScriptStructPath : ScriptStructPaths)
 	{
-		UScriptStruct* ScriptStruct = FindObject<UScriptStruct>(nullptr, *ScriptStructPath);
+		UScriptStruct* ScriptStruct = RigVMTypeUtils::FindObjectFromCPPTypeObjectPath<UScriptStruct>(*ScriptStructPath);
 
 		// try to find a redirect
 		if (ScriptStruct == nullptr)
diff --git a/Engine/Source/Runtime/RigVM/Private/RigVMCore/RigVMMemoryStorage.cpp b/Engine/Source/Runtime/RigVM/Private/RigVMCore/RigVMMemoryStorage.cpp
--- a/Engine/Source/Runtime/RigVM/Private/RigVMCore/RigVMMemoryStorage.cpp
+++ b/Engine/Source/Runtime/RigVM/Private/RigVMCore/RigVMMemoryStorage.cpp
@@ -258,7 +258,7 @@ FRigVMPropertyDescription::FRigVMPropertyDescription(const FName& InName, const
 	, DefaultValue(InDefaultValue)
 {
 	// sanity check that we have a CPP Type Object for enums, structs and uobjects
-	if(RequiresCPPTypeObject(CPPType))
+	if(RigVMTypeUtils::RequiresCPPTypeObject(CPPType))
 	{
 		checkf(CPPTypeObject != nullptr, TEXT("CPPType '%s' requires the CPPTypeObject to be provided."), *CPPType);
 	}
@@ -373,40 +373,6 @@ FString FRigVMPropertyDescription::GetTailCPPType() const
 	return TailCPPType;
 }
 
-bool FRigVMPropertyDescription::RequiresCPPTypeObject(const FString& InCPPType)
-{
-	static const TArray<FString> PrefixesRequiringCPPTypeObject = {
-		TEXT("F"), 
-		TEXT("TArray<F"), 
-		TEXT("E"), 
-		TEXT("TArray<E"), 
-		TEXT("U"), 
-		TEXT("TArray<U"),
-		TEXT("TObjectPtr<"), 
-		TEXT("TArray<TObjectPtr<"),
-		TEXT("TScriptInterface<")
-	};
-	static const TArray<FString> CPPTypesNotRequiringCPPTypeObject = {
-		TEXT("FString"), 
-		TEXT("TArray<FString>"), 
-		TEXT("FName"), 
-		TEXT("TArray<FName>") 
-	};
-
-	if(!CPPTypesNotRequiringCPPTypeObject.Contains(InCPPType))
-	{
-		for(const FString& Prefix : PrefixesRequiringCPPTypeObject)
-		{
-			if(InCPPType.StartsWith(Prefix, ESearchCase::CaseSensitive))
-			{
-				return true;
-			}
-		}
-	}
-
-	return false;
-}
-
 ////////////////////////////////////////////////////////////////////////////////////////////////////////////////////////
 
 class FRigVMMemoryStorageImportErrorContext : public FOutputDevice
diff --git a/Engine/Source/Runtime/RigVM/Private/RigVMCore/RigVMRegistry.cpp b/Engine/Source/Runtime/RigVM/Private/RigVMCore/RigVMRegistry.cpp
--- a/Engine/Source/Runtime/RigVM/Private/RigVMCore/RigVMRegistry.cpp
+++ b/Engine/Source/Runtime/RigVM/Private/RigVMCore/RigVMRegistry.cpp
@@ -12,11 +12,7 @@
 #include "UObject/UObjectIterator.h"
 #include "Misc/CoreDelegates.h"
 #include "Misc/DelayedAutoRegister.h"
-#include "UObject/CoreRedirects.h"
 
-#if WITH_EDITOR
-#include "UObject/CoreRedirects.h"
-#endif
 
 const FName FRigVMRegistry::TemplateNameMetaName = TEXT("TemplateName");
 
@@ -723,22 +719,10 @@ TRigVMTypeIndex FRigVMRegistry::GetTypeIndexFromCPPType(const FString& InCPPType
 		// If not found, try to find a redirect
 		if (Result == INDEX_NONE)
 		{
-			FCoreRedirectObjectName NewObjectName;
-			const bool bFoundRedirect = FCoreRedirects::RedirectNameAndValues(
-				ECoreRedirectFlags::Type_Class | ECoreRedirectFlags::Type_Struct | ECoreRedirectFlags::Type_Enum,
-				FCoreRedirectObjectName(*InCPPType),
-				NewObjectName,
-				nullptr,
-				ECoreRedirectMatchFlags::AllowPartialMatch); // AllowPartialMatch to allow redirects from one package to another (see /Script/ControlRig.CRFourPointBezier -> /Script/RigVM.RigVMFourPointBezier)
-
-			if (!bFoundRedirect)
-			{
-				return INDEX_NONE;
-			}
-
-			Result = Types.IndexOfByPredicate([NewObjectName](const FTypeInfo& Info) -> bool
+			const FString NewCPPType = RigVMTypeUtils::PostProcessCPPType(InCPPType);
+			Result = Types.IndexOfByPredicate([NewCPPType](const FTypeInfo& Info) -> bool
 			{
-				return Info.Type.CPPType == NewObjectName.ObjectName;
+				return Info.Type.CPPType == *NewCPPType;
 			});
 		}
 	}
diff --git a/Engine/Source/Runtime/RigVM/Public/RigVMCore/RigVMMemoryStorage.h b/Engine/Source/Runtime/RigVM/Public/RigVMCore/RigVMMemoryStorage.h
--- a/Engine/Source/Runtime/RigVM/Public/RigVMCore/RigVMMemoryStorage.h
+++ b/Engine/Source/Runtime/RigVM/Public/RigVMCore/RigVMMemoryStorage.h
@@ -1028,8 +1028,6 @@ public:
 	// Returns the CPP type of the tail property, for ex: '[2].Translation' it is 'FVector'
 	FString GetTailCPPType() const;
 
-	static bool RequiresCPPTypeObject(const FString& InCPPType);
-
 private:
 	
 	static const FString ArrayPrefix;
diff --git a/Engine/Source/Runtime/RigVM/Public/RigVMTypeUtils.h b/Engine/Source/Runtime/RigVM/Public/RigVMTypeUtils.h
--- a/Engine/Source/Runtime/RigVM/Public/RigVMTypeUtils.h
+++ b/Engine/Source/Runtime/RigVM/Public/RigVMTypeUtils.h
@@ -7,6 +7,8 @@
 #include "RigVMCore/RigVMUnknownType.h"
 #include "UObject/Interface.h"
 #include "Engine/UserDefinedStruct.h"
+#include "UObject/CoreRedirects.h"
+#include "UObject/Package.h"
 
 namespace RigVMTypeUtils
 {
@@ -148,42 +150,253 @@ namespace RigVMTypeUtils
 		return WildCardArrayCPPTypeName;
 	}
 
+	static bool RequiresCPPTypeObject(const FString& InCPPType)
+	{
+		static const TArray<FString> PrefixesRequiringCPPTypeObject = {
+			TEXT("F"), 
+			TEXT("TArray<F"), 
+			TEXT("TArray<TArray<F"), 
+			TEXT("E"), 
+			TEXT("TArray<E"), 
+			TEXT("TArray<TArray<E"), 
+			TEXT("U"), 
+			TEXT("TArray<U"),
+			TEXT("TArray<TArray<U"),
+			TEXT("TObjectPtr<"), 
+			TEXT("TArray<TObjectPtr<"),
+			TEXT("TArray<TArray<TObjectPtr<"),
+			TEXT("TScriptInterface<")
+		};
+		static const TArray<FString> CPPTypesNotRequiringCPPTypeObject = {
+			TEXT("FString"), 
+			TEXT("TArray<FString>"), 
+			TEXT("TArray<TArray<FString>>"), 
+			TEXT("FName"), 
+			TEXT("TArray<FName>"),
+			TEXT("TArray<TArray<FName>>") 
+		};
 
+		if(!CPPTypesNotRequiringCPPTypeObject.Contains(InCPPType))
+		{
+			for(const FString& Prefix : PrefixesRequiringCPPTypeObject)
+			{
+				if(InCPPType.StartsWith(Prefix, ESearchCase::CaseSensitive))
+				{
+					return true;
+				}
+			}
+		}
 
-	inline FString PostProcessCPPType(const FString& InCPPType, UObject* InCPPTypeObject)
+		return false;
+	}
+
+	static UObject* FindObjectGlobally(const TCHAR* InObjectName, bool bUseRedirector)
+	{
+		// Do a global search for the CPP type. Note that searching with ANY_PACKAGE _does not_
+		// apply redirectors. So only if this fails do we apply them manually below.
+		UObject* Object = FindFirstObject<UField>(InObjectName, EFindFirstObjectOptions::NativeFirst);
+		if(Object != nullptr)
+		{
+			return Object;
+		}
+
+		// If its an enum, it might be defined as a namespace with the actual enum inside (see ERigVMClampSpatialMode). 
+		FString ObjectNameStr(InObjectName);
+		if (!ObjectNameStr.IsEmpty() && ObjectNameStr[0] == 'E')
+		{
+			FString Left, Right;
+			if (ObjectNameStr.Split(TEXT("::"), &Left, &Right))
+			{
+				ObjectNameStr = Left;
+				Object = FindFirstObject<UField>(*ObjectNameStr, EFindFirstObjectOptions::NativeFirst);
+				if(Object != nullptr)
+				{
+					return Object;
+				}
+			}
+		}
+
+		if (!bUseRedirector)
+		{
+			return nullptr;
+		}
+
+		FCoreRedirectObjectName OldObjectName (ObjectNameStr);
+		FCoreRedirectObjectName NewObjectName;
+		const bool bFoundRedirect = FCoreRedirects::RedirectNameAndValues(
+			ECoreRedirectFlags::Type_Class | ECoreRedirectFlags::Type_Struct | ECoreRedirectFlags::Type_Enum,
+			OldObjectName,
+			NewObjectName,
+			nullptr,
+			ECoreRedirectMatchFlags::AllowPartialMatch); // AllowPartialMatch to allow redirects from one package to another (see /Script/ControlRig.CRFourPointBezier -> /Script/RigVM.RigVMFourPointBezier)
+
+		if (!bFoundRedirect)
+		{
+			return nullptr;
+		}
+
+		const FString RedirectedObjectName = NewObjectName.ObjectName.ToString();
+		UPackage *Package = nullptr;
+		if (!NewObjectName.PackageName.IsNone())
+		{
+			Package = FindPackage(nullptr, *NewObjectName.PackageName.ToString());
+		}
+		if (Package != nullptr)
+		{
+			Object = FindObject<UField>(Package, *RedirectedObjectName);
+		}
+		if (Package == nullptr || Object == nullptr)
+		{
+			// Hail Mary pass.
+			Object = FindFirstObject<UField>(*RedirectedObjectName, EFindFirstObjectOptions::NativeFirst);
+		}
+		return Object;
+	}
+
+	static FString CPPTypeFromObject(UObject* InCPPTypeObject)
 	{
-		FString CPPType = InCPPType;
-	
 		if (const UClass* Class = Cast<UClass>(InCPPTypeObject))
 		{
 			if (Class->IsChildOf(UInterface::StaticClass()))
 			{
-				CPPType = FString::Printf(RigVMTypeUtils::TScriptInterfaceTemplate, TEXT("I"), *Class->GetName());
+				return FString::Printf(RigVMTypeUtils::TScriptInterfaceTemplate, TEXT("I"), *Class->GetName());
 			}
 			else
 			{
-				CPPType = FString::Printf(RigVMTypeUtils::TObjectPtrTemplate, Class->GetPrefixCPP(), *Class->GetName());
+				return FString::Printf(RigVMTypeUtils::TObjectPtrTemplate, Class->GetPrefixCPP(), *Class->GetName());
 			}
 		}
 		else if (const UScriptStruct* ScriptStruct = Cast<UScriptStruct>(InCPPTypeObject))
 		{
-			CPPType = GetUniqueStructTypeName(ScriptStruct);
+			return GetUniqueStructTypeName(ScriptStruct);
 		}
 		else if (UEnum* Enum = Cast<UEnum>(InCPPTypeObject))
 		{
-			CPPType = RigVMTypeUtils::CPPTypeFromEnum(Enum);
+			return RigVMTypeUtils::CPPTypeFromEnum(Enum);
+		}
+
+		return FString();
+	}
+
+	// Finds the CPPTypeObject from the CPPType. If not found, tries to use redirectors and modifies the InOutCPPType.
+	static UObject* ObjectFromCPPType(FString& InOutCPPType, bool bUseRedirector = true)
+	{
+		if (!RequiresCPPTypeObject(InOutCPPType))
+		{
+			return nullptr;
+		}
+
+		// try to find the CPPTypeObject by name
+		FString BaseCPPType = InOutCPPType;
+		while (IsArrayType(BaseCPPType))
+		{
+			BaseCPPType = BaseTypeFromArrayType(BaseCPPType);
 		}
+		FString CPPType = BaseCPPType;
 
-		if(CPPType != InCPPType)
+		static const FString PrefixScriptInterface = TScriptInterfacePrefix;
+		if (CPPType.StartsWith(TScriptInterfacePrefix))
 		{
-			FString TemplateType = InCPPType;
-			while (RigVMTypeUtils::IsArrayType(TemplateType))
+			// Chop the prefix + the I indicating interface class
+			CPPType = CPPType.RightChop(PrefixScriptInterface.Len() + 1).LeftChop(1);
+		}
+
+		UObject* CPPTypeObject = FindObjectGlobally(*CPPType, bUseRedirector);
+		if (CPPTypeObject == nullptr)
+		{
+			// If we've mistakenly stored the struct type with the 'F', 'U', or 'A' prefixes, we need to strip them
+			// off first. Enums are always named with their prefix intact.
+			if (!CPPType.IsEmpty() && (CPPType[0] == TEXT('F') || CPPType[0] == TEXT('U') || CPPType[0] == TEXT('A')))
 			{
-				CPPType = RigVMTypeUtils::ArrayTypeFromBaseType(CPPType);
-				TemplateType = RigVMTypeUtils::BaseTypeFromArrayType(TemplateType);
-			}		
+				CPPType = CPPType.Mid(1);
+			}
+			CPPTypeObject = RigVMTypeUtils::FindObjectGlobally(*CPPType, bUseRedirector);
+		}
+
+		if(CPPTypeObject == nullptr)
+		{
+			InOutCPPType.Reset();
+			return nullptr;
+		}
+
+		CPPType = CPPTypeFromObject(CPPTypeObject);
+		InOutCPPType.ReplaceInline(*BaseCPPType, *CPPType);
+		return CPPTypeObject;
+	}
+
+	static FString PostProcessCPPType(const FString& InCPPType, UObject* InCPPTypeObject = nullptr)
+	{
+		FString CPPType = InCPPType;
+		if (InCPPTypeObject)
+		{
+			CPPType = CPPTypeFromObject(InCPPTypeObject);	
+			if(CPPType != InCPPType)
+			{
+				FString TemplateType = InCPPType;
+				while (RigVMTypeUtils::IsArrayType(TemplateType))
+				{
+					CPPType = RigVMTypeUtils::ArrayTypeFromBaseType(CPPType);
+					TemplateType = RigVMTypeUtils::BaseTypeFromArrayType(TemplateType);
+				}		
+			}
+		}
+		else if (RequiresCPPTypeObject(CPPType))
+		{
+			// Uses redirectors and updates the CPPType if necessary
+			ObjectFromCPPType(CPPType, true);
 		}
 	
 		return CPPType;
 	}
+
+	// helper function to retrieve an object from a path
+	static UObject* FindObjectFromCPPTypeObjectPath(const FString& InObjectPath)
+	{
+		UObject* Result = nullptr;
+		if (InObjectPath.IsEmpty())
+		{
+			return Result;
+		}
+
+		if (InObjectPath == FName(NAME_None).ToString())
+		{
+			return Result;
+		}
+
+		// we do this to avoid ambiguous searches for 
+		// common names such as "transform" or "vector"
+		UPackage* Package = nullptr;
+		FString PackageName;
+		FString CPPTypeObjectName = InObjectPath;
+		if (InObjectPath.Split(TEXT("."), &PackageName, &CPPTypeObjectName))
+		{
+			Package = FindPackage(nullptr, *PackageName);
+		}
+	
+		if (UObject* ObjectWithinPackage = FindObject<UObject>(Package, *CPPTypeObjectName))
+		{
+			Result = ObjectWithinPackage;
+		}
+
+		if (!Result)
+		{
+			Result = FindFirstObject<UObject>(*InObjectPath, EFindFirstObjectOptions::NativeFirst | EFindFirstObjectOptions::EnsureIfAmbiguous);
+		}
+		if (!Result)
+		{
+			const FCoreRedirectObjectName OldObjectName(InObjectPath);
+			const FCoreRedirectObjectName NewObjectName = FCoreRedirects::GetRedirectedName(ECoreRedirectFlags::Type_Struct, OldObjectName);
+			if (OldObjectName != NewObjectName)
+			{
+				Result = FindObjectFromCPPTypeObjectPath(NewObjectName.ToString());
+			}
+		}
+		return Result;
+	}
+	
+	template<class T>
+	static T* FindObjectFromCPPTypeObjectPath(const FString& InObjectPath)
+	{
+		return Cast<T>(FindObjectFromCPPTypeObjectPath(InObjectPath));
+	}
 }
\ No newline at end of file
