RigVM: When adding array nodes - don't rely on ResolveWildCard Pin due to side effects

diff --git a/Engine/Source/Developer/RigVMDeveloper/Private/RigVMModel/RigVMController.cpp b/Engine/Source/Developer/RigVMDeveloper/Private/RigVMModel/RigVMController.cpp
--- a/Engine/Source/Developer/RigVMDeveloper/Private/RigVMModel/RigVMController.cpp
+++ b/Engine/Source/Developer/RigVMDeveloper/Private/RigVMModel/RigVMController.cpp
@@ -13236,7 +13236,7 @@ URigVMEnumNode* URigVMController::AddEnumNode(const FName& InCPPTypeObjectPath,
 
 URigVMNode* URigVMController::AddArrayNode(ERigVMOpCode InOpCode, const FString& InCPPType,
 	UObject* InCPPTypeObject, const FVector2D& InPosition, const FString& InNodeName, bool bSetupUndoRedo,
-	bool bPrintPythonCommand)
+	bool bPrintPythonCommand, bool bIsPatching)
 {
 	if (!IsValidGraph())
 	{
@@ -13310,7 +13310,30 @@ URigVMNode* URigVMController::AddArrayNode(ERigVMOpCode InOpCode, const FString&
 
 		if(!ArgumentNameToResolve.IsNone() && TypeIndex != INDEX_NONE)
 		{
-			if(const URigVMPin* Pin = Node->FindPin(ArgumentNameToResolve.ToString()))
+			if(bIsPatching)
+			{
+				FRigVMTemplate::FTypeMap TypeMap;
+				TypeMap.Add(ArgumentNameToResolve, TypeIndex);
+
+				TArray<int32> Permutations;
+				Template->Resolve(TypeMap, Permutations, true);
+				check(Permutations.Num() == 1);
+
+				for(const FRigVMTemplate::FTypePair& Pair : TypeMap)
+				{
+					if(!FRigVMRegistry::Get().IsWildCardType(Pair.Value))
+					{
+						if(URigVMPin* Pin = Node->FindPin(Pair.Key.ToString()))
+						{
+							ChangePinType(Pin, Pair.Value, false, false);
+							Node->UpdateFilteredPermutations(Pin, {Pair.Value});
+						}
+					}
+				}
+
+				FullyResolveTemplateNode(Node, Permutations[0], false);
+			}
+			else if(const URigVMPin* Pin = Node->FindPin(ArgumentNameToResolve.ToString()))
 			{
 				ResolveWildCardPin(Pin->GetPinPath(), TypeIndex, bSetupUndoRedo, bPrintPythonCommand);
 			}
@@ -13327,7 +13350,7 @@ URigVMNode* URigVMController::AddArrayNode(ERigVMOpCode InOpCode, const FString&
 
 URigVMNode* URigVMController::AddArrayNodeFromObjectPath(ERigVMOpCode InOpCode, const FString& InCPPType,
 	const FString& InCPPTypeObjectPath, const FVector2D& InPosition, const FString& InNodeName, bool bSetupUndoRedo,
-	bool bPrintPythonCommand)
+	bool bPrintPythonCommand, bool bIsPatching)
 {
 	if (!IsValidGraph())
 	{
@@ -13350,7 +13373,7 @@ URigVMNode* URigVMController::AddArrayNodeFromObjectPath(ERigVMOpCode InOpCode,
 		}
 	}
 
-	return AddArrayNode(InOpCode, InCPPType, CPPTypeObject, InPosition, InNodeName, bSetupUndoRedo, bPrintPythonCommand);
+	return AddArrayNode(InOpCode, InCPPType, CPPTypeObject, InPosition, InNodeName, bSetupUndoRedo, bPrintPythonCommand, bIsPatching);
 }
 
 URigVMInvokeEntryNode* URigVMController::AddInvokeEntryNode(const FName& InEntryName, const FVector2D& InPosition,
@@ -19496,7 +19519,7 @@ void URigVMController::PatchArrayNodesOnLoad()
 			
 			RemoveNode(ArrayNode, false, true, false, false);
 
-			URigVMNode* NewNode = AddArrayNode(OpCode, CPPType, CPPTypeObject, NodePosition, NodeName, false, false);
+			URigVMNode* NewNode = AddArrayNode(OpCode, CPPType, CPPTypeObject, NodePosition, NodeName, false, false, true);
 			ApplyPinStates(NewNode, PinStates, {}, false);
 			RestoreLinkedPaths(LinkedPaths, {}, {}, false);
 		}
diff --git a/Engine/Source/Developer/RigVMDeveloper/Public/RigVMModel/RigVMController.h b/Engine/Source/Developer/RigVMDeveloper/Public/RigVMModel/RigVMController.h
--- a/Engine/Source/Developer/RigVMDeveloper/Public/RigVMModel/RigVMController.h
+++ b/Engine/Source/Developer/RigVMDeveloper/Public/RigVMModel/RigVMController.h
@@ -410,12 +410,12 @@ public:
 	// Adds a Array Node to the edited Graph.
 	// This causes a NodeAdded modified event.
 	UFUNCTION(BlueprintCallable, Category = RigVMController)
-	URigVMNode* AddArrayNode(ERigVMOpCode InOpCode, const FString& InCPPType, UObject* InCPPTypeObject, const FVector2D& InPosition = FVector2D::ZeroVector, const FString& InNodeName = TEXT(""), bool bSetupUndoRedo = true, bool bPrintPythonCommand = false);
+	URigVMNode* AddArrayNode(ERigVMOpCode InOpCode, const FString& InCPPType, UObject* InCPPTypeObject, const FVector2D& InPosition = FVector2D::ZeroVector, const FString& InNodeName = TEXT(""), bool bSetupUndoRedo = true, bool bPrintPythonCommand = false, bool bIsPatching = false);
 
 	// Adds a Array Node to the edited Graph given a struct object path name.
 	// This causes a NodeAdded modified event.
 	UFUNCTION(BlueprintCallable, Category = RigVMController)
-	URigVMNode* AddArrayNodeFromObjectPath(ERigVMOpCode InOpCode, const FString& InCPPType, const FString& InCPPTypeObjectPath, const FVector2D& InPosition = FVector2D::ZeroVector, const FString& InNodeName = TEXT(""), bool bSetupUndoRedo = true, bool bPrintPythonCommand = false);
+	URigVMNode* AddArrayNodeFromObjectPath(ERigVMOpCode InOpCode, const FString& InCPPType, const FString& InCPPTypeObjectPath, const FVector2D& InPosition = FVector2D::ZeroVector, const FString& InNodeName = TEXT(""), bool bSetupUndoRedo = true, bool bPrintPythonCommand = false, bool bIsPatching = false);
 
 	// Adds an entry invocation node
 	// This causes a NodeAdded modified event.
