RigVM: Move RigVMTypeUtils functionality

diff --git a/Engine/Source/Developer/RigVMDeveloper/Private/RigVMDeveloperTypeUtils.cpp b/Engine/Source/Developer/RigVMDeveloper/Private/RigVMDeveloperTypeUtils.cpp
--- a/Engine/Source/Developer/RigVMDeveloper/Private/RigVMDeveloperTypeUtils.cpp
+++ b/Engine/Source/Developer/RigVMDeveloper/Private/RigVMDeveloperTypeUtils.cpp
@@ -9,15 +9,6 @@
 #if WITH_EDITOR
 
 #include "EdGraphSchema_K2.h"
-#include "Engine/Blueprint.h"
-
-FRigVMExternalVariable RigVMTypeUtils::ExternalVariableFromBPVariableDescription(
-	const FBPVariableDescription& InVariableDescription)
-{
-	const bool bIsPublic = !((InVariableDescription.PropertyFlags & CPF_DisableEditOnInstance) == CPF_DisableEditOnInstance);
-	const bool bIsReadOnly = ((InVariableDescription.PropertyFlags & CPF_BlueprintReadOnly) == CPF_BlueprintReadOnly);
-	return ExternalVariableFromPinType(InVariableDescription.VarName, InVariableDescription.VarType, bIsPublic, bIsReadOnly);
-}
 
 FRigVMExternalVariable RigVMTypeUtils::ExternalVariableFromRigVMVariableDescription(const FRigVMGraphVariableDescription& InVariableDescription)
 {	
@@ -49,230 +40,6 @@ FRigVMExternalVariable RigVMTypeUtils::ExternalVariableFromRigVMVariableDescript
 
 }
 
-FRigVMExternalVariable RigVMTypeUtils::ExternalVariableFromPinType(const FName& InName, const FEdGraphPinType& InPinType, bool bInPublic, bool bInReadonly)
-{
-	FRigVMExternalVariable ExternalVariable;
-	ExternalVariable.Name = InName;
-	ExternalVariable.bIsPublic = bInPublic;
-	ExternalVariable.bIsReadOnly = bInReadonly;
-
-	if (InPinType.ContainerType == EPinContainerType::None)
-	{
-		ExternalVariable.bIsArray = false;
-	}
-	else if (InPinType.ContainerType == EPinContainerType::Array)
-	{
-		ExternalVariable.bIsArray = true;
-	}
-	else
-	{
-		return FRigVMExternalVariable();
-	}
-
-	if (InPinType.PinCategory == UEdGraphSchema_K2::PC_Boolean)
-	{
-		ExternalVariable.TypeName = BoolTypeName;
-		ExternalVariable.Size = sizeof(bool);
-	}
-	else if (InPinType.PinCategory == UEdGraphSchema_K2::PC_Int)
-	{
-		ExternalVariable.TypeName = Int32TypeName;
-		ExternalVariable.Size = sizeof(int32);
-	}
-	else if (InPinType.PinCategory == UEdGraphSchema_K2::PC_Enum ||
-		InPinType.PinCategory == UEdGraphSchema_K2::PC_Byte)
-	{
-		if (UEnum* Enum = Cast<UEnum>(InPinType.PinSubCategoryObject))
-		{
-			ExternalVariable.TypeName = Enum->GetFName();
-			ExternalVariable.TypeObject = Enum;
-		}
-		else
-		{
-			ExternalVariable.TypeName = UInt8TypeName;
-		}
-		ExternalVariable.Size = sizeof(uint8);
-	}
-	else if (InPinType.PinCategory == UEdGraphSchema_K2::PC_Real)
-	{
-		if (InPinType.PinSubCategory == UEdGraphSchema_K2::PC_Float)
-		{
-			ExternalVariable.TypeName = FloatTypeName;
-			ExternalVariable.Size = sizeof(float);
-		}
-		else if (InPinType.PinSubCategory == UEdGraphSchema_K2::PC_Double)
-		{
-			ExternalVariable.TypeName = DoubleTypeName;
-			ExternalVariable.Size = sizeof(double);
-		}
-		else
-		{
-			checkf(false, TEXT("Unexpected subcategory for PC_Real pin type."));
-		}
-	}
-	else if (InPinType.PinCategory == UEdGraphSchema_K2::PC_Name)
-	{
-		ExternalVariable.TypeName = FNameTypeName;
-		ExternalVariable.Size = sizeof(FName);
-	}
-	else if (InPinType.PinCategory == UEdGraphSchema_K2::PC_String)
-	{
-		ExternalVariable.TypeName = FStringTypeName;
-		ExternalVariable.Size = sizeof(FString);
-	}
-	else if (InPinType.PinCategory == UEdGraphSchema_K2::PC_Struct)
-	{
-		if (UScriptStruct* Struct = Cast<UScriptStruct>(InPinType.PinSubCategoryObject))
-		{
-			ExternalVariable.TypeName = *RigVMTypeUtils::GetUniqueStructTypeName(Struct);
-			ExternalVariable.TypeObject = Struct;
-			ExternalVariable.Size = Struct->GetStructureSize();
-		}
-	}
-
-	return ExternalVariable;
-}
-
-FRigVMExternalVariable RigVMTypeUtils::ExternalVariableFromCPPTypePath(const FName& InName, const FString& InCPPTypePath, bool bInPublic, bool bInReadonly)
-{
-	FRigVMExternalVariable Variable;
-	if (InCPPTypePath.StartsWith(TEXT("TMap<")))
-	{
-		return Variable;
-	}
-	
-	Variable.Name = InName;
-	Variable.bIsPublic = bInPublic;
-	Variable.bIsReadOnly = bInReadonly;
-
-	FString CPPTypePath = InCPPTypePath;
-	Variable.bIsArray = RigVMTypeUtils::IsArrayType(CPPTypePath);
-	if (Variable.bIsArray)
-	{
-		CPPTypePath = BaseTypeFromArrayType(CPPTypePath);
-	}
-
-	if (CPPTypePath == BoolType)
-	{
-		Variable.TypeName = *CPPTypePath;
-		Variable.Size = sizeof(bool);
-	}
-	else if (CPPTypePath == FloatType)
-	{
-		Variable.TypeName = *CPPTypePath;
-		Variable.Size = sizeof(float);
-	}
-	else if (CPPTypePath == DoubleType)
-	{
-		Variable.TypeName = *CPPTypePath;
-		Variable.Size = sizeof(double);
-	}
-	else if (CPPTypePath == Int32Type)
-	{
-		Variable.TypeName = *CPPTypePath;
-		Variable.Size = sizeof(int32);
-	}
-	else if (CPPTypePath == FStringType)
-	{
-		Variable.TypeName = *CPPTypePath;
-		Variable.Size = sizeof(FString);
-	}
-	else if (CPPTypePath == FNameType)
-	{
-		Variable.TypeName = *CPPTypePath;
-		Variable.Size = sizeof(FName);
-	}
-	else if(UScriptStruct* ScriptStruct = RigVMTypeUtils::FindObjectFromCPPTypeObjectPath<UScriptStruct>(CPPTypePath))
-	{
-		Variable.TypeName = *RigVMTypeUtils::GetUniqueStructTypeName(ScriptStruct);
-		Variable.TypeObject = ScriptStruct;
-		Variable.Size = ScriptStruct->GetStructureSize();
-	}
-	else if (UEnum* Enum= RigVMTypeUtils::FindObjectFromCPPTypeObjectPath<UEnum>(CPPTypePath))
-	{
-		Variable.TypeName = *CPPTypeFromEnum(Enum);
-		Variable.TypeObject = Enum;
-		Variable.Size = Enum->GetResourceSizeBytes(EResourceSizeMode::EstimatedTotal);
-	}
-	else
-	{
-		check(false);
-	}
-
-	return Variable;
-}
-
-FRigVMExternalVariable RigVMTypeUtils::ExternalVariableFromCPPType(const FName& InName, const FString& InCPPType, UObject* InCPPTypeObject, bool bInPublic, bool bInReadonly)
-{
-	FRigVMExternalVariable Variable;
-	if (InCPPType.StartsWith(TEXT("TMap<")))
-	{
-		return Variable;
-	}
-	
-	Variable.Name = InName;
-	Variable.bIsPublic = bInPublic;
-	Variable.bIsReadOnly = bInReadonly;
-
-	FString CPPType = InCPPType;
-	Variable.bIsArray = RigVMTypeUtils::IsArrayType(CPPType);
-	if (Variable.bIsArray)
-	{
-		CPPType = BaseTypeFromArrayType(CPPType);
-	}
-
-	if (CPPType == BoolType)
-	{
-		Variable.TypeName = *CPPType;
-		Variable.Size = sizeof(bool);
-	}
-	else if (CPPType == FloatType)
-	{
-		Variable.TypeName = *CPPType;
-		Variable.Size = sizeof(float);
-	}
-	else if (CPPType == DoubleType)
-	{
-		Variable.TypeName = *CPPType;
-		Variable.Size = sizeof(double);
-	}
-	else if (CPPType == Int32Type)
-	{
-		Variable.TypeName = *CPPType;
-		Variable.Size = sizeof(int32);
-	}
-	else if (CPPType == FStringType)
-	{
-		Variable.TypeName = *CPPType;
-		Variable.Size = sizeof(FString);
-	}
-	else if (CPPType == FNameType)
-	{
-		Variable.TypeName = *CPPType;
-		Variable.Size = sizeof(FName);
-	}
-	else if(UScriptStruct* ScriptStruct = RigVMTypeUtils::FindObjectFromCPPTypeObjectPath<UScriptStruct>(CPPType))
-	{
-		Variable.TypeName = *RigVMTypeUtils::GetUniqueStructTypeName(ScriptStruct);
-		Variable.TypeObject = ScriptStruct;
-		Variable.Size = ScriptStruct->GetStructureSize();
-	}
-	else if (UEnum* Enum= RigVMTypeUtils::FindObjectFromCPPTypeObjectPath<UEnum>(CPPType))
-	{
-		Variable.TypeName = *CPPTypeFromEnum(Enum);
-		Variable.TypeObject = Enum;
-		Variable.Size = Enum->GetResourceSizeBytes(EResourceSizeMode::EstimatedTotal);
-	}
-	else
-	{
-		Variable.TypeName = *CPPType;
-		Variable.TypeObject = InCPPTypeObject;
-		Variable.Size = InCPPTypeObject->StaticClass()->GetStructureSize();
-	}
-
-	return Variable;
-}
-
 FEdGraphPinType RigVMTypeUtils::PinTypeFromTypeIndex(const TRigVMTypeIndex& InTypeIndex)
 {
 	const FRigVMTemplateArgumentType& Type = FRigVMRegistry::Get().GetType(InTypeIndex);
diff --git a/Engine/Source/Developer/RigVMDeveloper/Public/RigVMDeveloperTypeUtils.h b/Engine/Source/Developer/RigVMDeveloper/Public/RigVMDeveloperTypeUtils.h
--- a/Engine/Source/Developer/RigVMDeveloper/Public/RigVMDeveloperTypeUtils.h
+++ b/Engine/Source/Developer/RigVMDeveloper/Public/RigVMDeveloperTypeUtils.h
@@ -14,16 +14,8 @@ class URigVMPin;
 namespace RigVMTypeUtils
 {
 #if WITH_EDITOR
-	RIGVMDEVELOPER_API FRigVMExternalVariable ExternalVariableFromBPVariableDescription(const FBPVariableDescription& InVariableDescription);
-
 	RIGVMDEVELOPER_API FRigVMExternalVariable ExternalVariableFromRigVMVariableDescription(const FRigVMGraphVariableDescription& InVariableDescription);
-
-	RIGVMDEVELOPER_API FRigVMExternalVariable ExternalVariableFromPinType(const FName& InName, const FEdGraphPinType& InPinType, bool bInPublic = false, bool bInReadonly = false);
-
-	RIGVMDEVELOPER_API FRigVMExternalVariable ExternalVariableFromCPPTypePath(const FName& InName, const FString& InCPPTypePath, bool bInPublic = false, bool bInReadonly = false);
-
-	RIGVMDEVELOPER_API FRigVMExternalVariable ExternalVariableFromCPPType(const FName& InName, const FString& InCPPType, UObject* InCPPTypeObject, bool bInPublic = false, bool bInReadonly = false);
-
+	
 	RIGVMDEVELOPER_API FEdGraphPinType PinTypeFromTypeIndex(const TRigVMTypeIndex& InTypeIndex);
 
 	RIGVMDEVELOPER_API FEdGraphPinType PinTypeFromCPPType(const FName& InCPPType, UObject* InCPPTypeObject);
diff --git a/Engine/Source/Runtime/RigVM/Private/RigVMCore/RigVMExternalVariable.cpp b/Engine/Source/Runtime/RigVM/Private/RigVMCore/RigVMExternalVariable.cpp
new file mode 100644
--- /dev/null
+++ b/Engine/Source/Runtime/RigVM/Private/RigVMCore/RigVMExternalVariable.cpp
@@ -0,0 +1,242 @@
+// Copyright Epic Games, Inc. All Rights Reserved.
+
+#include "RigVMCore/RigVMExternalVariable.h"
+
+#if WITH_EDITOR
+
+#include "EdGraphSchema_K2.h"
+#include "Engine/Blueprint.h"
+
+FRigVMExternalVariable RigVMTypeUtils::ExternalVariableFromBPVariableDescription(
+	const FBPVariableDescription& InVariableDescription)
+{
+	const bool bIsPublic = !((InVariableDescription.PropertyFlags & CPF_DisableEditOnInstance) == CPF_DisableEditOnInstance);
+	const bool bIsReadOnly = ((InVariableDescription.PropertyFlags & CPF_BlueprintReadOnly) == CPF_BlueprintReadOnly);
+	return ExternalVariableFromPinType(InVariableDescription.VarName, InVariableDescription.VarType, bIsPublic, bIsReadOnly);
+}
+
+FRigVMExternalVariable RigVMTypeUtils::ExternalVariableFromPinType(const FName& InName, const FEdGraphPinType& InPinType, bool bInPublic, bool bInReadonly)
+{
+	FRigVMExternalVariable ExternalVariable;
+	ExternalVariable.Name = InName;
+	ExternalVariable.bIsPublic = bInPublic;
+	ExternalVariable.bIsReadOnly = bInReadonly;
+
+	if (InPinType.ContainerType == EPinContainerType::None)
+	{
+		ExternalVariable.bIsArray = false;
+	}
+	else if (InPinType.ContainerType == EPinContainerType::Array)
+	{
+		ExternalVariable.bIsArray = true;
+	}
+	else
+	{
+		return FRigVMExternalVariable();
+	}
+
+	if (InPinType.PinCategory == UEdGraphSchema_K2::PC_Boolean)
+	{
+		ExternalVariable.TypeName = BoolTypeName;
+		ExternalVariable.Size = sizeof(bool);
+	}
+	else if (InPinType.PinCategory == UEdGraphSchema_K2::PC_Int)
+	{
+		ExternalVariable.TypeName = Int32TypeName;
+		ExternalVariable.Size = sizeof(int32);
+	}
+	else if (InPinType.PinCategory == UEdGraphSchema_K2::PC_Enum ||
+		InPinType.PinCategory == UEdGraphSchema_K2::PC_Byte)
+	{
+		if (UEnum* Enum = Cast<UEnum>(InPinType.PinSubCategoryObject))
+		{
+			ExternalVariable.TypeName = Enum->GetFName();
+			ExternalVariable.TypeObject = Enum;
+		}
+		else
+		{
+			ExternalVariable.TypeName = UInt8TypeName;
+		}
+		ExternalVariable.Size = sizeof(uint8);
+	}
+	else if (InPinType.PinCategory == UEdGraphSchema_K2::PC_Real)
+	{
+		if (InPinType.PinSubCategory == UEdGraphSchema_K2::PC_Float)
+		{
+			ExternalVariable.TypeName = FloatTypeName;
+			ExternalVariable.Size = sizeof(float);
+		}
+		else if (InPinType.PinSubCategory == UEdGraphSchema_K2::PC_Double)
+		{
+			ExternalVariable.TypeName = DoubleTypeName;
+			ExternalVariable.Size = sizeof(double);
+		}
+		else
+		{
+			checkf(false, TEXT("Unexpected subcategory for PC_Real pin type."));
+		}
+	}
+	else if (InPinType.PinCategory == UEdGraphSchema_K2::PC_Name)
+	{
+		ExternalVariable.TypeName = FNameTypeName;
+		ExternalVariable.Size = sizeof(FName);
+	}
+	else if (InPinType.PinCategory == UEdGraphSchema_K2::PC_String)
+	{
+		ExternalVariable.TypeName = FStringTypeName;
+		ExternalVariable.Size = sizeof(FString);
+	}
+	else if (InPinType.PinCategory == UEdGraphSchema_K2::PC_Struct)
+	{
+		if (UScriptStruct* Struct = Cast<UScriptStruct>(InPinType.PinSubCategoryObject))
+		{
+			ExternalVariable.TypeName = *RigVMTypeUtils::GetUniqueStructTypeName(Struct);
+			ExternalVariable.TypeObject = Struct;
+			ExternalVariable.Size = Struct->GetStructureSize();
+		}
+	}
+
+	return ExternalVariable;
+}
+
+FRigVMExternalVariable RigVMTypeUtils::ExternalVariableFromCPPTypePath(const FName& InName, const FString& InCPPTypePath, bool bInPublic, bool bInReadonly)
+{
+	FRigVMExternalVariable Variable;
+	if (InCPPTypePath.StartsWith(TEXT("TMap<")))
+	{
+		return Variable;
+	}
+	
+	Variable.Name = InName;
+	Variable.bIsPublic = bInPublic;
+	Variable.bIsReadOnly = bInReadonly;
+
+	FString CPPTypePath = InCPPTypePath;
+	Variable.bIsArray = RigVMTypeUtils::IsArrayType(CPPTypePath);
+	if (Variable.bIsArray)
+	{
+		CPPTypePath = BaseTypeFromArrayType(CPPTypePath);
+	}
+
+	if (CPPTypePath == BoolType)
+	{
+		Variable.TypeName = *CPPTypePath;
+		Variable.Size = sizeof(bool);
+	}
+	else if (CPPTypePath == FloatType)
+	{
+		Variable.TypeName = *CPPTypePath;
+		Variable.Size = sizeof(float);
+	}
+	else if (CPPTypePath == DoubleType)
+	{
+		Variable.TypeName = *CPPTypePath;
+		Variable.Size = sizeof(double);
+	}
+	else if (CPPTypePath == Int32Type)
+	{
+		Variable.TypeName = *CPPTypePath;
+		Variable.Size = sizeof(int32);
+	}
+	else if (CPPTypePath == FStringType)
+	{
+		Variable.TypeName = *CPPTypePath;
+		Variable.Size = sizeof(FString);
+	}
+	else if (CPPTypePath == FNameType)
+	{
+		Variable.TypeName = *CPPTypePath;
+		Variable.Size = sizeof(FName);
+	}
+	else if(UScriptStruct* ScriptStruct = RigVMTypeUtils::FindObjectFromCPPTypeObjectPath<UScriptStruct>(CPPTypePath))
+	{
+		Variable.TypeName = *RigVMTypeUtils::GetUniqueStructTypeName(ScriptStruct);
+		Variable.TypeObject = ScriptStruct;
+		Variable.Size = ScriptStruct->GetStructureSize();
+	}
+	else if (UEnum* Enum= RigVMTypeUtils::FindObjectFromCPPTypeObjectPath<UEnum>(CPPTypePath))
+	{
+		Variable.TypeName = *CPPTypeFromEnum(Enum);
+		Variable.TypeObject = Enum;
+		Variable.Size = Enum->GetResourceSizeBytes(EResourceSizeMode::EstimatedTotal);
+	}
+	else
+	{
+		check(false);
+	}
+
+	return Variable;
+}
+
+FRigVMExternalVariable RigVMTypeUtils::ExternalVariableFromCPPType(const FName& InName, const FString& InCPPType, UObject* InCPPTypeObject, bool bInPublic, bool bInReadonly)
+{
+	FRigVMExternalVariable Variable;
+	if (InCPPType.StartsWith(TEXT("TMap<")))
+	{
+		return Variable;
+	}
+	
+	Variable.Name = InName;
+	Variable.bIsPublic = bInPublic;
+	Variable.bIsReadOnly = bInReadonly;
+
+	FString CPPType = InCPPType;
+	Variable.bIsArray = RigVMTypeUtils::IsArrayType(CPPType);
+	if (Variable.bIsArray)
+	{
+		CPPType = BaseTypeFromArrayType(CPPType);
+	}
+
+	if (CPPType == BoolType)
+	{
+		Variable.TypeName = *CPPType;
+		Variable.Size = sizeof(bool);
+	}
+	else if (CPPType == FloatType)
+	{
+		Variable.TypeName = *CPPType;
+		Variable.Size = sizeof(float);
+	}
+	else if (CPPType == DoubleType)
+	{
+		Variable.TypeName = *CPPType;
+		Variable.Size = sizeof(double);
+	}
+	else if (CPPType == Int32Type)
+	{
+		Variable.TypeName = *CPPType;
+		Variable.Size = sizeof(int32);
+	}
+	else if (CPPType == FStringType)
+	{
+		Variable.TypeName = *CPPType;
+		Variable.Size = sizeof(FString);
+	}
+	else if (CPPType == FNameType)
+	{
+		Variable.TypeName = *CPPType;
+		Variable.Size = sizeof(FName);
+	}
+	else if(UScriptStruct* ScriptStruct = RigVMTypeUtils::FindObjectFromCPPTypeObjectPath<UScriptStruct>(CPPType))
+	{
+		Variable.TypeName = *RigVMTypeUtils::GetUniqueStructTypeName(ScriptStruct);
+		Variable.TypeObject = ScriptStruct;
+		Variable.Size = ScriptStruct->GetStructureSize();
+	}
+	else if (UEnum* Enum= RigVMTypeUtils::FindObjectFromCPPTypeObjectPath<UEnum>(CPPType))
+	{
+		Variable.TypeName = *CPPTypeFromEnum(Enum);
+		Variable.TypeObject = Enum;
+		Variable.Size = Enum->GetResourceSizeBytes(EResourceSizeMode::EstimatedTotal);
+	}
+	else
+	{
+		Variable.TypeName = *CPPType;
+		Variable.TypeObject = InCPPTypeObject;
+		Variable.Size = InCPPTypeObject->StaticClass()->GetStructureSize();
+	}
+
+	return Variable;
+}
+
+#endif
\ No newline at end of file
diff --git a/Engine/Source/Runtime/RigVM/Public/RigVMCore/RigVMExternalVariable.h b/Engine/Source/Runtime/RigVM/Public/RigVMCore/RigVMExternalVariable.h
--- a/Engine/Source/Runtime/RigVM/Public/RigVMCore/RigVMExternalVariable.h
+++ b/Engine/Source/Runtime/RigVM/Public/RigVMCore/RigVMExternalVariable.h
@@ -9,6 +9,9 @@
 #include "RigVMModule.h"
 #include "RigVMTypeUtils.h"
 #include "UObject/UnrealType.h"
+#if WITH_EDITOR
+#include "EdGraphSchema_K2.h"
+#endif
 #include "RigVMExternalVariable.generated.h"
 
 /**
@@ -646,3 +649,15 @@ inline FArchive& operator<<(FArchive& Ar, FRigVMExternalVariable& Variable)
 	return Ar;
 }
 
+namespace RigVMTypeUtils
+{
+#if WITH_EDITOR
+ 	RIGVM_API FRigVMExternalVariable ExternalVariableFromBPVariableDescription(const FBPVariableDescription& InVariableDescription);
+
+	RIGVM_API FRigVMExternalVariable ExternalVariableFromPinType(const FName& InName, const FEdGraphPinType& InPinType, bool bInPublic = false, bool bInReadonly = false);
+
+	RIGVM_API FRigVMExternalVariable ExternalVariableFromCPPTypePath(const FName& InName, const FString& InCPPTypePath, bool bInPublic = false, bool bInReadonly = false);
+
+	RIGVM_API FRigVMExternalVariable ExternalVariableFromCPPType(const FName& InName, const FString& InCPPType, UObject* InCPPTypeObject, bool bInPublic = false, bool bInReadonly = false);
+#endif
+}
diff --git a/Engine/Source/Runtime/RigVM/RigVM.Build.cs b/Engine/Source/Runtime/RigVM/RigVM.Build.cs
--- a/Engine/Source/Runtime/RigVM/RigVM.Build.cs
+++ b/Engine/Source/Runtime/RigVM/RigVM.Build.cs
@@ -18,5 +18,16 @@ public class RigVM : ModuleRules
                 "AnimGraphRuntime",
             }
         );
+
+        if (Target.bBuildEditor == true)
+        {
+            PrivateDependencyModuleNames.AddRange(
+                new string[]
+                {
+                    "UnrealEd",
+                    "BlueprintGraph",
+                }
+            );
+        }
     }
 }
