Fix RigVM potential corruption because we're holding a ref on an array element while adding to that array which could invalidate the ref on reallocs.

#rnx

diff --git a/Engine/Source/Developer/RigVMDeveloper/Private/RigVMCompiler/RigVMCompiler.cpp b/Engine/Source/Developer/RigVMDeveloper/Private/RigVMCompiler/RigVMCompiler.cpp
--- a/Engine/Source/Developer/RigVMDeveloper/Private/RigVMCompiler/RigVMCompiler.cpp
+++ b/Engine/Source/Developer/RigVMDeveloper/Private/RigVMCompiler/RigVMCompiler.cpp
@@ -3261,11 +3261,9 @@ const FRigVMCompilerWorkData::FRigVMASTProxyArray& URigVMCompiler::FindProxiesWi
 
 	for(int32 ProxyIndex = 0; ProxyIndex < PinProxiesToProcess.Num(); ProxyIndex++)
 	{
-		const FRigVMASTProxy& CurrentProxy = PinProxiesToProcess[ProxyIndex];
-
-		if (CurrentProxy.IsValid())
+		if (PinProxiesToProcess[ProxyIndex].IsValid())
 		{
-			if (URigVMPin* Pin = Cast<URigVMPin>(CurrentProxy.GetSubject()))
+			if (URigVMPin* Pin = Cast<URigVMPin>(PinProxiesToProcess[ProxyIndex].GetSubject()))
 			{
 				if (Pin->GetNode()->IsA<URigVMVariableNode>())
 				{
@@ -3281,10 +3279,10 @@ const FRigVMCompilerWorkData::FRigVMASTProxyArray& URigVMCompiler::FindProxiesWi
 					continue;
 				}
 			}
-			PinProxies.Add(CurrentProxy);
+			PinProxies.Add(PinProxiesToProcess[ProxyIndex]);
 		}
 
-		if(const FRigVMASTProxy* SourceProxy = ProxySources.Find(CurrentProxy))
+		if(const FRigVMASTProxy* SourceProxy = ProxySources.Find(PinProxiesToProcess[ProxyIndex]))
 		{
 			if(SourceProxy->IsValid())
 			{
@@ -3295,7 +3293,7 @@ const FRigVMCompilerWorkData::FRigVMASTProxyArray& URigVMCompiler::FindProxiesWi
 			}
 		}
 
-		if(const FRigVMCompilerWorkData::FRigVMASTProxyArray* TargetProxies = WorkData.ProxyTargets.Find(CurrentProxy))
+		if(const FRigVMCompilerWorkData::FRigVMASTProxyArray* TargetProxies = WorkData.ProxyTargets.Find(PinProxiesToProcess[ProxyIndex]))
 		{
 			for(const FRigVMASTProxy& TargetProxy : *TargetProxies)
 			{
