RigVM: Fix heap access bug

diff --git a/Engine/Source/Developer/RigVMDeveloper/Private/RigVMCompiler/RigVMCompiler.cpp b/Engine/Source/Developer/RigVMDeveloper/Private/RigVMCompiler/RigVMCompiler.cpp
--- a/Engine/Source/Developer/RigVMDeveloper/Private/RigVMCompiler/RigVMCompiler.cpp
+++ b/Engine/Source/Developer/RigVMDeveloper/Private/RigVMCompiler/RigVMCompiler.cpp
@@ -1797,12 +1797,13 @@ int32 URigVMCompiler::TraverseCallExtern(const FRigVMCallExternExprAST* InExpr,
 			// The first index of the branch in the overall list of branches is stored in the operator (BranchIndices[0])
 			WorkData.VM->GetByteCode().AddJumpToBranchOp(BlockToRunOperand, BranchIndices[0]);
 
-			const TArray<FName>& BlockNames = Node->GetControlFlowBlocks();
+			// create a copy here for ensure memory validity
+			TArray<FName> BlockNames = Node->GetControlFlowBlocks();
 
 			// traverse all of the blocks now
 			for(int32 BlockIndex = 0; BlockIndex < BlockNames.Num(); BlockIndex++)
 			{
-				const FName& BlockName = BlockNames[BlockIndex];
+				const FName BlockName = BlockNames[BlockIndex];
 				int32 BranchIndex = BranchIndices[BlockIndex];
 				{
 					FRigVMBranchInfo& BranchInfo = WorkData.VM->GetByteCode().BranchInfos[BranchIndex];
