Control Rig: Fix issue with renaming function

diff --git a/Engine/Plugins/Animation/ControlRig/Source/ControlRigDeveloper/Private/Graph/ControlRigGraph.cpp b/Engine/Plugins/Animation/ControlRig/Source/ControlRigDeveloper/Private/Graph/ControlRigGraph.cpp
--- a/Engine/Plugins/Animation/ControlRig/Source/ControlRigDeveloper/Private/Graph/ControlRigGraph.cpp
+++ b/Engine/Plugins/Animation/ControlRig/Source/ControlRigDeveloper/Private/Graph/ControlRigGraph.cpp
@@ -866,8 +866,9 @@ void UControlRigGraph::HandleModifiedEvent(ERigVMGraphNotifType InNotifType, URi
 		{
 			if (URigVMNode* ModelNode = Cast<URigVMNode>(InSubject))
 			{
+				UEdGraphNode* EdNode = FindNodeForModelNodeName(ModelNode->GetPreviousFName());
 				ModelNodePathToEdNode.Remove(ModelNode->GetPreviousFName());
-				if (UControlRigGraphNode* RigNode = Cast<UControlRigGraphNode>(FindNodeForModelNodeName(ModelNode->GetPreviousFName())))
+				if (UControlRigGraphNode* RigNode = Cast<UControlRigGraphNode>(EdNode))
 				{
 					RigNode->SyncGraphNodeNameWithModelNodeName(ModelNode);
 					ModelNodePathToEdNode.Add(ModelNode->GetFName(), RigNode);
diff --git a/Engine/Source/Developer/RigVMDeveloper/Private/RigVMModel/RigVMClient.cpp b/Engine/Source/Developer/RigVMDeveloper/Private/RigVMModel/RigVMClient.cpp
--- a/Engine/Source/Developer/RigVMDeveloper/Private/RigVMModel/RigVMClient.cpp
+++ b/Engine/Source/Developer/RigVMDeveloper/Private/RigVMModel/RigVMClient.cpp
@@ -781,12 +781,23 @@ void FRigVMClient::HandleGraphModifiedEvent(ERigVMGraphNotifType InNotifType, UR
 					URigVMBuildData* BuildData = URigVMBuildData::Get();
 					IRigVMGraphFunctionHost* Host = Cast<IRigVMGraphFunctionHost>(CollapseNode->GetFunctionIdentifier().HostObject.ResolveObject());
 					FRigVMGraphFunctionData* Data = Host->GetRigVMGraphFunctionStore()->FindFunctionByName(CollapseNode->GetPreviousFName());
+					const FRigVMGraphFunctionIdentifier PreviousFunctionId = Data->Header.LibraryPointer;
 					Data->Header = CollapseNode->GetFunctionHeader();
-					
-					BuildData->ForEachFunctionReference(CollapseNode->GetFunctionIdentifier(), [Data](URigVMFunctionReferenceNode* ReferenceNode)
+
+					if (const FRigVMFunctionReferenceArray* FunctionReferencesPtr = BuildData->GraphFunctionReferences.Find(PreviousFunctionId))
 					{
-						ReferenceNode->ReferencedFunctionHeader = Data->Header;
-					});
+						const FRigVMFunctionReferenceArray& FunctionReferences = *FunctionReferencesPtr;
+						BuildData->Modify();
+						FRigVMFunctionReferenceArray& NewFunctionReferences = BuildData->GraphFunctionReferences.Add(Data->Header.LibraryPointer, FunctionReferences);
+						BuildData->GraphFunctionReferences.Remove(PreviousFunctionId);
+						BuildData->MarkPackageDirty();
+
+						for (int32 i=0; i<NewFunctionReferences.Num(); ++i)
+						{
+							NewFunctionReferences[i]->ReferencedFunctionHeader = Data->Header;
+						}
+					}
+					
 					UpdateGraphFunctionData(CollapseNode);
 				}
 			}
diff --git a/Engine/Source/Developer/RigVMDeveloper/Private/RigVMModel/RigVMController.cpp b/Engine/Source/Developer/RigVMDeveloper/Private/RigVMModel/RigVMController.cpp
--- a/Engine/Source/Developer/RigVMDeveloper/Private/RigVMModel/RigVMController.cpp
+++ b/Engine/Source/Developer/RigVMDeveloper/Private/RigVMModel/RigVMController.cpp
@@ -6236,29 +6236,15 @@ bool URigVMController::RenameNode(URigVMNode* InNode, const FName& InNewName, bo
 	{
 		if (URigVMFunctionLibrary* FunctionLibrary = Cast<URigVMFunctionLibrary>(LibraryNode->GetGraph()))
 		{
-			// update the table in the build data
+			// rename each function reference
 			if(URigVMBuildData* BuildData = URigVMBuildData::Get())
 			{
-				for(const TPair< FRigVMGraphFunctionIdentifier, FRigVMFunctionReferenceArray >& Pair: BuildData->GraphFunctionReferences)
+				BuildData->ForEachFunctionReference(LibraryNode->GetFunctionIdentifier(), [this, InNewName](URigVMFunctionReferenceNode* ReferenceNode)
 				{
-					if(Pair.Key.LibraryNode == PreviousObjectPath)
-					{
-						const FRigVMFunctionReferenceArray FunctionReferences = Pair.Value;
-						
-						BuildData->Modify();
-						BuildData->GraphFunctionReferences.Remove(Pair.Key);
-						BuildData->GraphFunctionReferences.Add(LibraryNode->GetFunctionIdentifier(), FunctionReferences);
-						BuildData->MarkPackageDirty();
-						break;
-					}
-				}
+					FRigVMControllerGraphGuard GraphGuard(this, ReferenceNode->GetGraph(), false);
+					RenameNode(ReferenceNode, InNewName, false);
+				});
 			}
-			
-			FunctionLibrary->ForEachReference(LibraryNode->GetFName(), [this, InNewName](URigVMFunctionReferenceNode* ReferenceNode)
-			{
-				FRigVMControllerGraphGuard GraphGuard(this, ReferenceNode->GetGraph(), false);
-                RenameNode(ReferenceNode, InNewName, false);
-			});
 
 			if (FunctionLibrary->PublicFunctionNames.Contains(InNode->PreviousName))
 			{
diff --git a/Engine/Source/Developer/RigVMDeveloper/Public/RigVMModel/RigVMBuildData.h b/Engine/Source/Developer/RigVMDeveloper/Public/RigVMModel/RigVMBuildData.h
--- a/Engine/Source/Developer/RigVMDeveloper/Public/RigVMModel/RigVMBuildData.h
+++ b/Engine/Source/Developer/RigVMDeveloper/Public/RigVMModel/RigVMBuildData.h
@@ -129,6 +129,7 @@ private:
 	bool bIsRunningUnitTest;
 
 	friend class URigVMController;
+	friend struct FRigVMClient;
 	friend class URigVMCompiler;
 };
 
