Control Rig: Set of fixes to re-enable all rigs in ControlRigExample

diff --git a/Engine/Plugins/Animation/ControlRig/Source/ControlRigDeveloper/Private/ControlRigBlueprint.cpp b/Engine/Plugins/Animation/ControlRig/Source/ControlRigDeveloper/Private/ControlRigBlueprint.cpp
--- a/Engine/Plugins/Animation/ControlRig/Source/ControlRigDeveloper/Private/ControlRigBlueprint.cpp
+++ b/Engine/Plugins/Animation/ControlRig/Source/ControlRigDeveloper/Private/ControlRigBlueprint.cpp
@@ -1490,6 +1490,8 @@ void UControlRigBlueprint::RefreshAllModels(EControlRigBlueprintLoadType InLoadT
 			if (Controller->RecomputeAllTemplateFilteredPermutations(false))
 			{
 				ensureMsgf(false, TEXT("Pin type changed during load %s"), *GetPackage()->GetPathName());
+				// if you hit this ensure consider to uncomment the code
+				// in URigVMController::ReportPinTypeChange for debugging
 			}
 
 			for(URigVMNode* ModelNode : GraphToClean->GetNodes())
@@ -4541,14 +4543,17 @@ void UControlRigBlueprint::PatchLinksWithCast()
 			{
 				const URigVMPin* SourcePin = Link->GetSourcePin();
 				const URigVMPin* TargetPin = Link->GetTargetPin();
-				const TRigVMTypeIndex SourceTypeIndex = SourcePin->GetTypeIndex();
-				const TRigVMTypeIndex TargetTypeIndex = TargetPin->GetTypeIndex();
-
-				if(SourceTypeIndex != TargetTypeIndex)
+				if (SourcePin && TargetPin)
 				{
-					if(!FRigVMRegistry::Get().CanMatchTypes(SourceTypeIndex, TargetTypeIndex, true))
+					const TRigVMTypeIndex SourceTypeIndex = SourcePin->GetTypeIndex();
+					const TRigVMTypeIndex TargetTypeIndex = TargetPin->GetTypeIndex();
+					
+					if(SourceTypeIndex != TargetTypeIndex)
 					{
-						LinksWithCast.Emplace(Graph, TWeakObjectPtr<URigVMLink>(Link), SourcePin->GetPinPath(), TargetPin->GetPinPath());
+						if(!FRigVMRegistry::Get().CanMatchTypes(SourceTypeIndex, TargetTypeIndex, true))
+						{
+							LinksWithCast.Emplace(Graph, TWeakObjectPtr<URigVMLink>(Link), SourcePin->GetPinPath(), TargetPin->GetPinPath());
+						}
 					}
 				}
 			}
diff --git a/Engine/Source/Developer/RigVMDeveloper/Private/RigVMModel/Nodes/RigVMTemplateNode.cpp b/Engine/Source/Developer/RigVMDeveloper/Private/RigVMModel/Nodes/RigVMTemplateNode.cpp
--- a/Engine/Source/Developer/RigVMDeveloper/Private/RigVMModel/Nodes/RigVMTemplateNode.cpp
+++ b/Engine/Source/Developer/RigVMDeveloper/Private/RigVMModel/Nodes/RigVMTemplateNode.cpp
@@ -117,6 +117,12 @@ void URigVMTemplateNode::PostLoad()
 
 	ConvertPreferredTypesToTypeIndex();
 	InvalidateCache();
+
+	// the template notation may have changed
+	if(const FRigVMTemplate* Template = GetTemplate())
+	{
+		TemplateNotation = Template->GetNotation();
+	}
 }
 
 void URigVMTemplateNode::ConvertPreferredTypesToString()
diff --git a/Engine/Source/Developer/RigVMDeveloper/Private/RigVMModel/RigVMClient.cpp b/Engine/Source/Developer/RigVMDeveloper/Private/RigVMModel/RigVMClient.cpp
--- a/Engine/Source/Developer/RigVMDeveloper/Private/RigVMModel/RigVMClient.cpp
+++ b/Engine/Source/Developer/RigVMDeveloper/Private/RigVMModel/RigVMClient.cpp
@@ -668,6 +668,9 @@ void FRigVMClient::PatchModelsOnLoad()
 	for(URigVMGraph* Model : AllModels)
 	{
 		URigVMController* Controller = GetOrCreateController(Model);
+		TGuardValue<bool> GuardSuspendTemplateComputation(Controller->bSuspendTemplateComputation, true);
+		TGuardValue<bool> GuardIsTransacting(Controller->bIsTransacting, true);
+		Controller->PatchUnitNodesOnLoad();
 		Controller->PatchDispatchNodesOnLoad();
 		Controller->PatchBranchNodesOnLoad();
 		Controller->PatchIfSelectNodesOnLoad();
diff --git a/Engine/Source/Developer/RigVMDeveloper/Private/RigVMModel/RigVMController.cpp b/Engine/Source/Developer/RigVMDeveloper/Private/RigVMModel/RigVMController.cpp
--- a/Engine/Source/Developer/RigVMDeveloper/Private/RigVMModel/RigVMController.cpp
+++ b/Engine/Source/Developer/RigVMDeveloper/Private/RigVMModel/RigVMController.cpp
@@ -82,6 +82,7 @@ URigVMController::URigVMController()
 	, bIsRunningUnitTest(false)
 	, bIsFullyResolvingTemplateNode(false)
 	, bSuspendRecomputingTemplateFilters(false)
+	, bSuspendTemplateComputation(false)
 #if WITH_EDITOR
 	, bRegisterTemplateNodeUsage(true)
 #endif
@@ -101,6 +102,7 @@ URigVMController::URigVMController(const FObjectInitializer& ObjectInitializer)
 	, bIsRunningUnitTest(false)
 	, bIsFullyResolvingTemplateNode(false)
 	, bSuspendRecomputingTemplateFilters(false)
+	, bSuspendTemplateComputation(false)
 #if WITH_EDITOR
 	, bRegisterTemplateNodeUsage(true)
 #endif
@@ -9602,30 +9604,33 @@ bool URigVMController::AddLink(URigVMPin* OutputPin, URigVMPin* InputPin, bool b
 		}
 	}
 
-	// Find any library pin that doesn't own an argument, but should 
-	if (URigVMLibraryNode* LibraryNode = Graph->GetTypedOuter<URigVMLibraryNode>())
+	// Find any library pin that doesn't own an argument, but should
+	if (!bSuspendTemplateComputation)
 	{
-		for (URigVMPin* Pin : LibraryNode->GetPins())
+		if (URigVMLibraryNode* LibraryNode = Graph->GetTypedOuter<URigVMLibraryNode>())
 		{
-			if (const FRigVMTemplate* Template = LibraryNode->GetTemplate())
+			for (URigVMPin* Pin : LibraryNode->GetPins())
 			{
-				if (!Template->FindArgument(Pin->GetFName()))
+				if (const FRigVMTemplate* Template = LibraryNode->GetTemplate())
 				{
-					if (ShouldPinOwnArgument(Pin))
+					if (!Template->FindArgument(Pin->GetFName()))
 					{
-						URigVMPin* InterfacePin = nullptr;
-						if (URigVMFunctionEntryNode* EntryNode = LibraryNode->GetEntryNode())
-						{
-							InterfacePin = EntryNode->FindPin(Pin->GetName());
-						}
-						if (!InterfacePin)
+						if (ShouldPinOwnArgument(Pin))
 						{
-							if (URigVMFunctionReturnNode* ReturnNode = LibraryNode->GetReturnNode())
+							URigVMPin* InterfacePin = nullptr;
+							if (URigVMFunctionEntryNode* EntryNode = LibraryNode->GetEntryNode())
 							{
-								InterfacePin = ReturnNode->FindPin(Pin->GetName());
+								InterfacePin = EntryNode->FindPin(Pin->GetName());
+							}
+							if (!InterfacePin)
+							{
+								if (URigVMFunctionReturnNode* ReturnNode = LibraryNode->GetReturnNode())
+								{
+									InterfacePin = ReturnNode->FindPin(Pin->GetName());
+								}
 							}
+							AddArgumentForPin(InterfacePin, nullptr, bSetupUndoRedo);
 						}
-						AddArgumentForPin(InterfacePin, nullptr, bSetupUndoRedo);
 					}
 				}
 			}
@@ -9641,50 +9646,53 @@ bool URigVMController::AddLink(URigVMPin* OutputPin, URigVMPin* InputPin, bool b
 	if (bSetupUndoRedo)
 	{
 #if WITH_EDITOR
-		auto ResolveTemplateNodeToCommonTypes = [this](URigVMPin* Pin)
+		if (!bSuspendTemplateComputation)
 		{
-			if(!Pin->IsExecuteContext())
+			auto ResolveTemplateNodeToCommonTypes = [this](URigVMPin* Pin)
 			{
-				return;
-			}
+				if(!Pin->IsExecuteContext())
+				{
+					return;
+				}
 			
-			URigVMTemplateNode* TemplateNode = Cast<URigVMTemplateNode>(Pin->GetNode());
-			if(TemplateNode == nullptr)
-			{
-				return;
-			}
+				URigVMTemplateNode* TemplateNode = Cast<URigVMTemplateNode>(Pin->GetNode());
+				if(TemplateNode == nullptr)
+				{
+					return;
+				}
 
-			const FRigVMTemplate* Template = TemplateNode->GetTemplate();
-			if(Template == nullptr)
-			{
-				return;
-			}
+				const FRigVMTemplate* Template = TemplateNode->GetTemplate();
+				if(Template == nullptr)
+				{
+					return;
+				}
 
-			if(!TemplateNode->HasWildCardPin())
-			{
-				return;
-			}
+				if(!TemplateNode->HasWildCardPin())
+				{
+					return;
+				}
 
-			const FRigVMTemplate::FTypeMap PreferredTypes = GetCommonlyUsedTypesForTemplate(TemplateNode);
-			if(PreferredTypes.IsEmpty())
-			{
-				return;
-			}
+				const FRigVMTemplate::FTypeMap PreferredTypes = GetCommonlyUsedTypesForTemplate(TemplateNode);
+				if(PreferredTypes.IsEmpty())
+				{
+					return;
+				}
 
-			const int32 PreferredPermutation = Template->FindPermutation(PreferredTypes);
-			if(PreferredPermutation != INDEX_NONE)
-			{
-				const TGuardValue<bool> DisableRegisterUseOfTemplate(bRegisterTemplateNodeUsage, false);
-				if(FullyResolveTemplateNode(TemplateNode, PreferredPermutation, true))
+				const int32 PreferredPermutation = Template->FindPermutation(PreferredTypes);
+				if(PreferredPermutation != INDEX_NONE)
 				{
-					static const FString Message = TEXT("Template node was automatically resolved to commonly used types.");
-					SendUserFacingNotification(Message, 0.f, TemplateNode, TEXT("MessageLog.Note"));
+					const TGuardValue<bool> DisableRegisterUseOfTemplate(bRegisterTemplateNodeUsage, false);
+					if(FullyResolveTemplateNode(TemplateNode, PreferredPermutation, true))
+					{
+						static const FString Message = TEXT("Template node was automatically resolved to commonly used types.");
+						SendUserFacingNotification(Message, 0.f, TemplateNode, TEXT("MessageLog.Note"));
+					}
 				}
-			}
-		};
+			};
 
-		ResolveTemplateNodeToCommonTypes(OutputPin);
-		ResolveTemplateNodeToCommonTypes(InputPin);
+			ResolveTemplateNodeToCommonTypes(OutputPin);
+			ResolveTemplateNodeToCommonTypes(InputPin);
+		}
 #endif
 		
 		ActionStack->EndAction(Action);
@@ -15544,8 +15552,14 @@ void URigVMController::ApplyPinState(URigVMPin* InPin, const FPinState& InPinSta
 	{
 		for (const URigVMInjectionInfo::FWeakInfo& InjectionInfo : InPinState.WeakInjectionInfos)
 		{
-			if(const URigVMNode* FormerlyInjectedNode = InjectionInfo.Node.Get())
+			if(URigVMNode* FormerlyInjectedNode = InjectionInfo.Node.Get())
 			{
+				if (FormerlyInjectedNode->IsInjected())
+				{
+					URigVMInjectionInfo* Injection = Cast<URigVMInjectionInfo>(FormerlyInjectedNode->GetOuter());
+					RenameObject(FormerlyInjectedNode, nullptr, InPin->GetGraph());
+					DestroyObject(Injection);
+				}
 				if(InjectionInfo.bInjectedAsInput)
 				{
 					const FString OutputPinPath = URigVMPin::JoinPinPath(FormerlyInjectedNode->GetNodePath(), InjectionInfo.OutputPinName.ToString()); 
@@ -15686,6 +15700,19 @@ void URigVMController::ReportAndNotifyError(const FString& InMessage) const
 	SendUserFacingNotification(InMessage, 0.f, nullptr, TEXT("MessageLog.Error"));
 }
 
+void URigVMController::ReportPinTypeChange(URigVMPin* InPin, const FString& InNewCPPType)
+{
+	/*
+	UE_LOG(LogRigVMDeveloper,
+		Warning,
+		TEXT("Pin '%s' is about to change type from '%s' to '%s'."),
+		*InPin->GetPinPath(),
+		*InPin->GetCPPType(),
+		*InNewCPPType
+	);
+	*/
+}
+
 void URigVMController::SendUserFacingNotification(const FString& InMessage, float InDuration, const UObject* InSubject, const FName& InBrushName) const
 {
 #if WITH_EDITOR
@@ -17178,7 +17205,7 @@ bool URigVMController::UpdateTemplateNodePinTypes(URigVMTemplateNode* InNode, bo
 			{
 				bShouldUnresolve = true;
 			}
-				
+
 			if (bShouldUnresolve)
 			{
 				// Unresolve
@@ -17220,6 +17247,10 @@ bool URigVMController::UpdateTemplateNodePinTypes(URigVMTemplateNode* InNode, bo
 				if (Pin->GetCPPType() != CPPType || Pin->GetCPPTypeObject() != CPPObjectType)
 				{
 					bAnyTypeChanged = !Pin->IsExecuteContext();
+					if(bAnyTypeChanged)
+					{
+						ReportPinTypeChange(Pin, CPPType);
+					}
 					ChangePinType(Pin, CPPType, CPPObjectType, bSetupUndoRedo, false, false, false, bInitializeDefaultValue);
 				}
 			}
@@ -17229,6 +17260,11 @@ bool URigVMController::UpdateTemplateNodePinTypes(URigVMTemplateNode* InNode, bo
 				if (Pin->GetTypeIndex() != FinalPinTypes[PinIndex])
 				{
 					bAnyTypeChanged = !Pin->IsExecuteContext();
+					if(bAnyTypeChanged)
+					{
+						const FString CPPType = FRigVMRegistry::Get().GetType(FinalPinTypes[PinIndex]).CPPType.ToString();
+						ReportPinTypeChange(Pin, CPPType);
+					}
 					ChangePinType(Pin, FinalPinTypes[PinIndex], bSetupUndoRedo, false, false, false, bInitializeDefaultValue);
 				}
 
@@ -19347,6 +19383,82 @@ FRigVMTemplate::FTypeMap URigVMController::GetCommonlyUsedTypesForTemplate(
 	return Template->GetArgumentTypesFromString(TypesString);
 }
 
+void URigVMController::PatchUnitNodesOnLoad()
+{
+	if (const URigVMGraph* Graph = GetGraph())
+	{
+		TArray<URigVMUnitNode*> UnitNodesToTurnToDispatches;
+		
+		// check for unit nodes that should be dispatches
+		for(URigVMNode* Node : Graph->GetNodes())
+		{
+			if(URigVMUnitNode* UnitNode = Cast<URigVMUnitNode>(Node))
+			{
+				if(const FRigVMTemplate* Template = UnitNode->GetTemplate())
+				{
+					if(Template->GetDispatchFactory() != nullptr)
+					{
+						UnitNodesToTurnToDispatches.Add(UnitNode);
+					}
+				}
+			}
+		}
+
+		// convert unit nodes to dispatches
+		for(URigVMUnitNode* UnitNode : UnitNodesToTurnToDispatches)
+		{
+			TArray<TPair<FString, FString>> LinkedPaths = GetLinkedPinPaths(UnitNode);
+			const FVector2D NodePosition = UnitNode->GetPosition();
+			const FString NodeName = UnitNode->GetName();
+			TMap<FString, FPinState> PinStates = GetPinStates(UnitNode, false);
+			
+			FRigVMTemplate::FTypeMap TypeMap;
+			for(URigVMPin* Pin : UnitNode->GetPins())
+			{
+				if(Pin->GetDirection() == ERigVMPinDirection::Input ||
+					Pin->GetDirection() == ERigVMPinDirection::Visible ||
+					Pin->GetDirection() == ERigVMPinDirection::IO ||
+					Pin->GetDirection() == ERigVMPinDirection::Output)
+				{
+					TypeMap.Add(Pin->GetFName(), Pin->GetTypeIndex());
+				}
+			}
+
+			const FRigVMTemplate* Template = UnitNode->GetTemplate();
+
+			RemoveNode(UnitNode, false, false, false, false);
+
+			URigVMTemplateNode* NewNode = AddTemplateNode(
+				Template->GetNotation(),
+				NodePosition, 
+				NodeName, 
+				false, 
+				false);
+
+			TArray<int32> Permutations;
+			Template->Resolve(TypeMap, Permutations, false);
+
+			for(URigVMPin* Pin : NewNode->GetPins())
+			{
+				if(Pin->IsWildCard())
+				{
+					if(const TRigVMTypeIndex* ResolvedTypeIndex = TypeMap.Find(Pin->GetFName()))
+					{
+						if(!FRigVMRegistry::Get().IsWildCardType(*ResolvedTypeIndex))
+						{
+							ChangePinType(Pin, *ResolvedTypeIndex, false, false);
+							NewNode->UpdateFilteredPermutations(Pin, {*ResolvedTypeIndex});
+						}
+					}
+				}
+			}
+
+			ApplyPinStates(NewNode, PinStates, {}, false);
+			RestoreLinkedPaths(LinkedPaths, {}, {}, false);
+		}
+	}
+}
+
 void URigVMController::PatchDispatchNodesOnLoad()
 {
 	if (const URigVMGraph* Graph = GetGraph())
@@ -19503,7 +19615,7 @@ void URigVMController::PatchArrayNodesOnLoad()
 			const ERigVMOpCode OpCode = ArrayNode->GetOpCode();
 			TMap<FString, FPinState> PinStates = GetPinStates(ArrayNode, true);
 			
-			RemoveNode(ArrayNode, false, true, false, false);
+			RemoveNode(ArrayNode, false, false, false, false);
 
 			URigVMNode* NewNode = AddArrayNode(OpCode, CPPType, CPPTypeObject, NodePosition, NodeName, false, false, true);
 			ApplyPinStates(NewNode, PinStates, {}, false);
diff --git a/Engine/Source/Developer/RigVMDeveloper/Public/RigVMModel/RigVMController.h b/Engine/Source/Developer/RigVMDeveloper/Public/RigVMModel/RigVMController.h
--- a/Engine/Source/Developer/RigVMDeveloper/Public/RigVMModel/RigVMController.h
+++ b/Engine/Source/Developer/RigVMDeveloper/Public/RigVMModel/RigVMController.h
@@ -924,6 +924,7 @@ public:
 	void ReportAndNotifyInfo(const FString& InMessage) const;
 	void ReportAndNotifyWarning(const FString& InMessage) const;
 	void ReportAndNotifyError(const FString& InMessage) const;
+	void ReportPinTypeChange(URigVMPin* InPin, const FString& InNewCPPType);
 	void SendUserFacingNotification(const FString& InMessage, float InDuration = 0.f, const UObject* InSubject = nullptr, const FName& InBrushName = TEXT("MessageLog.Warning")) const;
 
 	template <typename FmtType, typename... Types>
@@ -1168,6 +1169,7 @@ public:
 protected:
 
 	// backwards compatibility code
+	void PatchUnitNodesOnLoad();
 	void PatchDispatchNodesOnLoad();
 	void PatchBranchNodesOnLoad();
 	void PatchIfSelectNodesOnLoad();
@@ -1237,6 +1239,7 @@ private:
 public:
 	
 	bool bSuspendRecomputingTemplateFilters;
+	bool bSuspendTemplateComputation;
 
 private:
 
diff --git a/Engine/Source/Runtime/RigVM/Private/RigVMCore/RigVMRegistry.cpp b/Engine/Source/Runtime/RigVM/Private/RigVMCore/RigVMRegistry.cpp
--- a/Engine/Source/Runtime/RigVM/Private/RigVMCore/RigVMRegistry.cpp
+++ b/Engine/Source/Runtime/RigVM/Private/RigVMCore/RigVMRegistry.cpp
@@ -12,7 +12,7 @@
 #include "UObject/UObjectIterator.h"
 #include "Misc/CoreDelegates.h"
 #include "Misc/DelayedAutoRegister.h"
-
+#include "RigVMFunctions//RigVMDispatch_Core.h"
 
 const FName FRigVMRegistry::TemplateNameMetaName = TEXT("TemplateName");
 
@@ -1215,6 +1215,26 @@ const FRigVMTemplate* FRigVMRegistry::FindTemplate(const FName& InNotation, bool
 	FString FactoryName, ArgumentsString;
 	if(NotationString.Split(TEXT("("), &FactoryName, &ArgumentsString))
 	{
+		FRigVMRegistry* MutableThis = (FRigVMRegistry*)this;
+		
+		// deal with a couple of custom cases
+		static const TMap<FString, FString> CoreDispatchMap =
+		{
+			{
+				TEXT("Equals::Execute"),
+				MutableThis->FindOrAddDispatchFactory<FRigVMDispatch_CoreEquals>()->GetFactoryName().ToString()
+			},
+			{
+				TEXT("NotEquals::Execute"),
+				MutableThis->FindOrAddDispatchFactory<FRigVMDispatch_CoreNotEquals>()->GetFactoryName().ToString()
+			},
+		};
+
+		if(const FString* RemappedDispatch = CoreDispatchMap.Find(FactoryName))
+		{
+			FactoryName = *RemappedDispatch;
+		}
+		
 		if(const FRigVMDispatchFactory* Factory = FindDispatchFactory(*FactoryName))
 		{
 			TGuardValue<bool> ReEntryGuard(IsDispatchingFunction, true);
