RigVMCompiler: Consider all graphs for function references

diff --git a/Engine/Source/Developer/RigVMDeveloper/Private/RigVMCompiler/RigVMCompiler.cpp b/Engine/Source/Developer/RigVMDeveloper/Private/RigVMCompiler/RigVMCompiler.cpp
--- a/Engine/Source/Developer/RigVMDeveloper/Private/RigVMCompiler/RigVMCompiler.cpp
+++ b/Engine/Source/Developer/RigVMDeveloper/Private/RigVMCompiler/RigVMCompiler.cpp
@@ -342,8 +342,9 @@ bool URigVMCompiler::Compile(TArray<URigVMGraph*> InGraphs, URigVMController* In
 	TMap<FString, const FRigVMFunctionCompilationData*> CurrentCompiledFunctions;
 
 	// Gather function compilation data
+	for(URigVMGraph* Graph : InGraphs)
 	{
-		TArray<URigVMNode*> Nodes = InGraphs[0]->GetNodes();
+		TArray<URigVMNode*> Nodes = Graph->GetNodes();
 		for (int32 i=0; i<Nodes.Num(); ++i)
 		{
 			if (URigVMFunctionReferenceNode* ReferenceNode = Cast<URigVMFunctionReferenceNode>(Nodes[i]))
