RigVM: Introduce patch result structure

diff --git a/Engine/Plugins/Animation/ControlRig/Source/ControlRigDeveloper/Private/ControlRigBlueprint.cpp b/Engine/Plugins/Animation/ControlRig/Source/ControlRigDeveloper/Private/ControlRigBlueprint.cpp
--- a/Engine/Plugins/Animation/ControlRig/Source/ControlRigDeveloper/Private/ControlRigBlueprint.cpp
+++ b/Engine/Plugins/Animation/ControlRig/Source/ControlRigDeveloper/Private/ControlRigBlueprint.cpp
@@ -855,7 +855,14 @@ void UControlRigBlueprint::PostLoad()
 		InitializeModelIfRequired(false /* recompile vm */);
 		{
 			TGuardValue<bool> GuardNotifications(bSuspendModelNotificationsForSelf, true);
-			GetRigVMClient()->PatchModelsOnLoad();
+			
+			const FRigVMClientPatchResult PatchResult = GetRigVMClient()->PatchModelsOnLoad();
+			if(PatchResult.RequiresToMarkPackageDirty())
+			{
+				(void)MarkPackageDirty();
+				bDirtyDuringLoad = true;
+			}
+			
 			PatchFunctionReferencesOnLoad();
 			PatchVariableNodesOnLoad();
 			PatchVariableNodesWithIncorrectType();
diff --git a/Engine/Source/Developer/RigVMDeveloper/Private/RigVMModel/RigVMClient.cpp b/Engine/Source/Developer/RigVMDeveloper/Private/RigVMModel/RigVMClient.cpp
--- a/Engine/Source/Developer/RigVMDeveloper/Private/RigVMModel/RigVMClient.cpp
+++ b/Engine/Source/Developer/RigVMDeveloper/Private/RigVMModel/RigVMClient.cpp
@@ -661,8 +661,10 @@ void FRigVMClient::DestroyObject(UObject* InObject)
 	}
 }
 
-void FRigVMClient::PatchModelsOnLoad()
+FRigVMClientPatchResult FRigVMClient::PatchModelsOnLoad()
 {
+	FRigVMClientPatchResult Result;
+	
 	TArray<URigVMGraph*> AllModels = GetAllModels(true, true);
 	TGuardValue<bool> ClientIgnoreModificationsGuard(bIgnoreModelNotifications, true);
 	for(URigVMGraph* Model : AllModels)
@@ -670,12 +672,15 @@ void FRigVMClient::PatchModelsOnLoad()
 		URigVMController* Controller = GetOrCreateController(Model);
 		TGuardValue<bool> GuardSuspendTemplateComputation(Controller->bSuspendTemplateComputation, true);
 		TGuardValue<bool> GuardIsTransacting(Controller->bIsTransacting, true);
-		Controller->PatchUnitNodesOnLoad();
-		Controller->PatchDispatchNodesOnLoad();
-		Controller->PatchBranchNodesOnLoad();
-		Controller->PatchIfSelectNodesOnLoad();
-		Controller->PatchArrayNodesOnLoad();
+		
+		Result.Merge(Controller->PatchUnitNodesOnLoad());
+		Result.Merge(Controller->PatchDispatchNodesOnLoad());
+		Result.Merge(Controller->PatchBranchNodesOnLoad());
+		Result.Merge(Controller->PatchIfSelectNodesOnLoad());
+		Result.Merge(Controller->PatchArrayNodesOnLoad());
 	}
+
+	return Result;
 }
 
 void FRigVMClient::PostDuplicateHost(const FString& InOldPathName, const FString& InNewPathName)
diff --git a/Engine/Source/Developer/RigVMDeveloper/Private/RigVMModel/RigVMController.cpp b/Engine/Source/Developer/RigVMDeveloper/Private/RigVMModel/RigVMController.cpp
--- a/Engine/Source/Developer/RigVMDeveloper/Private/RigVMModel/RigVMController.cpp
+++ b/Engine/Source/Developer/RigVMDeveloper/Private/RigVMModel/RigVMController.cpp
@@ -70,6 +70,16 @@ FRigVMControllerCompileBracketScope::~FRigVMControllerCompileBracketScope()
 	Graph->Notify(ERigVMGraphNotifType::InteractionBracketClosed, nullptr);
 }
 
+void FRigVMClientPatchResult::Merge(const FRigVMClientPatchResult& InOther)
+{
+	bSucceeded = Succeeded() && InOther.Succeeded();
+	bChangedContent = ChangedContent() || InOther.ChangedContent();
+	bRequiresToMarkPackageDirty = RequiresToMarkPackageDirty() || InOther.RequiresToMarkPackageDirty();
+	ErrorMessages.Append(InOther.GetErrorMessages());
+	RemovedNodes.Append(InOther.GetRemovedNodes());
+	AddedNodes.Append(InOther.GetAddedNodes());
+}
+
 URigVMController::URigVMController()
 	: bValidatePinDefaults(true)
 	, bSuspendRecomputingOuterTemplateFilters(false)
@@ -19383,8 +19393,10 @@ FRigVMTemplate::FTypeMap URigVMController::GetCommonlyUsedTypesForTemplate(
 	return Template->GetArgumentTypesFromString(TypesString);
 }
 
-void URigVMController::PatchUnitNodesOnLoad()
+FRigVMClientPatchResult URigVMController::PatchUnitNodesOnLoad()
 {
+	FRigVMClientPatchResult Result;
+	
 	if (const URigVMGraph* Graph = GetGraph())
 	{
 		TArray<URigVMUnitNode*> UnitNodesToTurnToDispatches;
@@ -19426,6 +19438,9 @@ void URigVMController::PatchUnitNodesOnLoad()
 
 			const FRigVMTemplate* Template = UnitNode->GetTemplate();
 
+			Result.RemovedNodes.Add(UnitNode->GetPathName());
+			Result.bChangedContent = true;
+
 			RemoveNode(UnitNode, false, false, false, false);
 
 			URigVMTemplateNode* NewNode = AddTemplateNode(
@@ -19435,6 +19450,8 @@ void URigVMController::PatchUnitNodesOnLoad()
 				false, 
 				false);
 
+			Result.AddedNodes.Add(NewNode);
+
 			TArray<int32> Permutations;
 			Template->Resolve(TypeMap, Permutations, false);
 
@@ -19457,10 +19474,14 @@ void URigVMController::PatchUnitNodesOnLoad()
 			RestoreLinkedPaths(LinkedPaths, {}, {}, false);
 		}
 	}
+
+	return Result;
 }
 
-void URigVMController::PatchDispatchNodesOnLoad()
+FRigVMClientPatchResult URigVMController::PatchDispatchNodesOnLoad()
 {
+	FRigVMClientPatchResult Result;
+
 	if (const URigVMGraph* Graph = GetGraph())
 	{
 		for(URigVMNode* Node : Graph->GetNodes())
@@ -19472,7 +19493,9 @@ void URigVMController::PatchDispatchNodesOnLoad()
 				{
 					if(Template->GetNotation() != DispatchNode->TemplateNotation)
 					{
+						Result.bChangedContent = Result.bChangedContent || DispatchNode->TemplateNotation != Template->GetNotation(); 
 						DispatchNode->TemplateNotation = Template->GetNotation();
+						
 						if(!DispatchNode->ResolvedFunctionName.IsEmpty())
 						{
 							FString FactoryName, ArgumentsString;
@@ -19491,6 +19514,7 @@ void URigVMController::PatchDispatchNodesOnLoad()
 										{
 											DispatchNode->ResolvedFunctionName = Function->GetName();
 											DispatchNode->ResolvedPermutation = Template->FindPermutation(Function);
+											Result.bChangedContent = true;
 										}
 									}
 								}
@@ -19501,10 +19525,14 @@ void URigVMController::PatchDispatchNodesOnLoad()
 			}
 		}
 	}
+
+	return Result;
 }
 
-void URigVMController::PatchBranchNodesOnLoad()
+FRigVMClientPatchResult URigVMController::PatchBranchNodesOnLoad()
 {
+	FRigVMClientPatchResult Result;
+
 	if (const URigVMGraph* Graph = GetGraph())
 	{
 		TArray<URigVMNode*> BranchNodes = Graph->GetNodes().FilterByPredicate([](URigVMNode* Node)
@@ -19520,10 +19548,15 @@ void URigVMController::PatchBranchNodesOnLoad()
 			const URigVMPin* OldConditionPin = BranchNode->FindPin(GET_MEMBER_NAME_CHECKED(FRigVMFunction_ControlFlowBranch, Condition).ToString());
 			const FString ConditionDefault = GetPinDefaultValue(OldConditionPin->GetPinPath());
 
+			Result.RemovedNodes.Add(BranchNode->GetPathName());
+			Result.bChangedContent = true;
+
 			RemoveNode(BranchNode, false, true, false, false);
 
 			const URigVMNode* NewNode = AddUnitNode(FRigVMFunction_ControlFlowBranch::StaticStruct(), FRigVMStruct::ExecuteName, NodePosition, NodeName, false, false);
 
+			Result.AddedNodes.Add(NewNode);
+
 			if(!ConditionDefault.IsEmpty())
 			{
 				const URigVMPin* ConditionPin = NewNode->FindPin(GET_MEMBER_NAME_CHECKED(FRigVMFunction_ControlFlowBranch, Condition).ToString());
@@ -19532,10 +19565,14 @@ void URigVMController::PatchBranchNodesOnLoad()
 			RestoreLinkedPaths(LinkedPaths, {}, {}, false);
 		}
 	}
+	
+	return Result;
 }
 
-void URigVMController::PatchIfSelectNodesOnLoad()
+FRigVMClientPatchResult URigVMController::PatchIfSelectNodesOnLoad()
 {
+	FRigVMClientPatchResult Result;
+	
 	if (const URigVMGraph* Graph = GetGraph())
 	{
 		TArray<URigVMNode*> IfOrSelectNodes = Graph->GetNodes().FilterByPredicate([](URigVMNode* Node)
@@ -19552,7 +19589,10 @@ void URigVMController::PatchIfSelectNodesOnLoad()
 			const FString NodeName = IfOrSelectNode->GetName();
 			const TRigVMTypeIndex TypeIndex = IfOrSelectNode->GetPins().Last()->GetTypeIndex();
 			TMap<FString, FPinState> PinStates = GetPinStates(IfOrSelectNode, true);
-			
+
+			Result.RemovedNodes.Add(IfOrSelectNode->GetPathName());
+			Result.bChangedContent = true;
+
 			RemoveNode(IfOrSelectNode, false, true, false, false);
 
 			const FRigVMDispatchFactory* Factory = FRigVMRegistry::Get().FindOrAddDispatchFactory(
@@ -19566,6 +19606,8 @@ void URigVMController::PatchIfSelectNodesOnLoad()
 				false, 
 				false);
 
+			Result.AddedNodes.Add(NewNode);
+
 			if(!FRigVMRegistry::Get().IsWildCardType(TypeIndex))
 			{
 				TArray<int32> Permutations;
@@ -19593,10 +19635,14 @@ void URigVMController::PatchIfSelectNodesOnLoad()
 			RestoreLinkedPaths(LinkedPaths, {}, {}, false);
 		}
 	}
+
+	return Result;
 }
 
-void URigVMController::PatchArrayNodesOnLoad()
+FRigVMClientPatchResult URigVMController::PatchArrayNodesOnLoad()
 {
+	FRigVMClientPatchResult Result;
+	
 	if (const URigVMGraph* Graph = GetGraph())
 	{
 		TArray<URigVMNode*> ArrayNodes = Graph->GetNodes().FilterByPredicate([](URigVMNode* Node)
@@ -19614,14 +19660,22 @@ void URigVMController::PatchArrayNodesOnLoad()
 			UObject* CPPTypeObject = ArrayNode->GetCPPTypeObject();
 			const ERigVMOpCode OpCode = ArrayNode->GetOpCode();
 			TMap<FString, FPinState> PinStates = GetPinStates(ArrayNode, true);
-			
+
+			Result.RemovedNodes.Add(ArrayNode->GetPathName());
+			Result.bChangedContent = true;
+
 			RemoveNode(ArrayNode, false, false, false, false);
 
 			URigVMNode* NewNode = AddArrayNode(OpCode, CPPType, CPPTypeObject, NodePosition, NodeName, false, false, true);
+
+			Result.AddedNodes.Add(NewNode);
+
 			ApplyPinStates(NewNode, PinStates, {}, false);
 			RestoreLinkedPaths(LinkedPaths, {}, {}, false);
 		}
 	}
+
+	return Result;
 }
 
 void URigVMController::PostDuplicateHost(const FString& InOldPathName, const FString& InNewPathName)
diff --git a/Engine/Source/Developer/RigVMDeveloper/Public/RigVMModel/RigVMClient.h b/Engine/Source/Developer/RigVMDeveloper/Public/RigVMModel/RigVMClient.h
--- a/Engine/Source/Developer/RigVMDeveloper/Public/RigVMModel/RigVMClient.h
+++ b/Engine/Source/Developer/RigVMDeveloper/Public/RigVMModel/RigVMClient.h
@@ -140,7 +140,7 @@ public:
 	static void DestroyObject(UObject* InObject);
 
 	// backwards compatibility
-	void PatchModelsOnLoad();
+	FRigVMClientPatchResult PatchModelsOnLoad();
 
 	// work to be done after a duplication of the source asset
 	void PostDuplicateHost(const FString& InOldPathName, const FString& InNewPathName);
diff --git a/Engine/Source/Developer/RigVMDeveloper/Public/RigVMModel/RigVMController.h b/Engine/Source/Developer/RigVMDeveloper/Public/RigVMModel/RigVMController.h
--- a/Engine/Source/Developer/RigVMDeveloper/Public/RigVMModel/RigVMController.h
+++ b/Engine/Source/Developer/RigVMDeveloper/Public/RigVMModel/RigVMController.h
@@ -125,6 +125,44 @@ protected:
 	const uint8* Memory;
 };
 
+// A struct describing the result of a backwards compatibility patch
+USTRUCT()
+struct RIGVMDEVELOPER_API FRigVMClientPatchResult
+{
+public:
+
+	GENERATED_BODY()
+	
+	FRigVMClientPatchResult()
+		: bSucceeded(true)
+		, bChangedContent(false)
+		, bRequiresToMarkPackageDirty(false)
+	{}
+
+	bool Succeeded() const { return bSucceeded; }
+	bool ChangedContent() const { return bChangedContent; }
+	bool RequiresToMarkPackageDirty() const { return bRequiresToMarkPackageDirty; }
+
+	const TArray<FString>& GetErrorMessages() const { return ErrorMessages; }
+	const TArray<FString>& GetRemovedNodes() const { return RemovedNodes; }
+	const TArray<TWeakObjectPtr<const URigVMNode>>& GetAddedNodes() const { return AddedNodes; }
+
+private:
+
+	void Merge(const FRigVMClientPatchResult& InOther);
+	
+	bool bSucceeded;
+	bool bChangedContent;
+	bool bRequiresToMarkPackageDirty;
+
+	TArray<FString> ErrorMessages;
+	TArray<FString> RemovedNodes;
+	TArray<TWeakObjectPtr<const URigVMNode>> AddedNodes;
+
+	friend struct FRigVMClient;
+	friend class URigVMController;
+};
+
 /**
  * The Controller is the sole authority to perform changes
  * on the Graph. The Controller itself is stateless.
@@ -1169,11 +1207,11 @@ public:
 protected:
 
 	// backwards compatibility code
-	void PatchUnitNodesOnLoad();
-	void PatchDispatchNodesOnLoad();
-	void PatchBranchNodesOnLoad();
-	void PatchIfSelectNodesOnLoad();
-	void PatchArrayNodesOnLoad();
+	FRigVMClientPatchResult PatchUnitNodesOnLoad();
+	FRigVMClientPatchResult PatchDispatchNodesOnLoad();
+	FRigVMClientPatchResult PatchBranchNodesOnLoad();
+	FRigVMClientPatchResult PatchIfSelectNodesOnLoad();
+	FRigVMClientPatchResult PatchArrayNodesOnLoad();
 
 	// work to do after a duplication of the host asset
 	void PostDuplicateHost(const FString& InOldPathName, const FString& InNewPathName);
