RigVMDeveloper: Deal with invalid links

diff --git a/Engine/Source/Developer/RigVMDeveloper/Private/RigVMModel/RigVMLink.cpp b/Engine/Source/Developer/RigVMDeveloper/Private/RigVMModel/RigVMLink.cpp
--- a/Engine/Source/Developer/RigVMDeveloper/Private/RigVMModel/RigVMLink.cpp
+++ b/Engine/Source/Developer/RigVMDeveloper/Private/RigVMModel/RigVMLink.cpp
@@ -54,6 +54,16 @@ int32 URigVMLink::GetLinkIndex() const
 
 URigVMPin* URigVMLink::GetSourcePin() const
 {
+	if(SourcePin)
+	{
+		if(SourcePin->GetOutermost() == GetTransientPackage())
+		{
+			if(GetOutermost() != GetTransientPackage())
+			{
+				SourcePin = nullptr;
+			}
+		}
+	}
 	if (SourcePin == nullptr)
 	{
 		SourcePin = GetGraph()->FindPin(SourcePinPath);
@@ -63,6 +73,16 @@ URigVMPin* URigVMLink::GetSourcePin() const
 
 URigVMPin* URigVMLink::GetTargetPin() const
 {
+	if(TargetPin)
+	{
+		if(TargetPin->GetOutermost() == GetTransientPackage())
+		{
+			if(GetOutermost() != GetTransientPackage())
+			{
+				TargetPin = nullptr;
+			}
+		}
+	}
 	if (TargetPin == nullptr)
 	{
 		TargetPin = GetGraph()->FindPin(TargetPinPath);
