Control Rig: Fix RigVM.Graph.NodeSelection and RigVM.UndoRedo.ClearSelection unit tests failing

diff --git a/Engine/Source/Developer/RigVMDeveloper/Private/RigVMModel/RigVMController.cpp b/Engine/Source/Developer/RigVMDeveloper/Private/RigVMModel/RigVMController.cpp
--- a/Engine/Source/Developer/RigVMDeveloper/Private/RigVMModel/RigVMController.cpp
+++ b/Engine/Source/Developer/RigVMDeveloper/Private/RigVMModel/RigVMController.cpp
@@ -6370,7 +6370,7 @@ bool URigVMController::SetNodeSelection(const TArray<FName>& InNodeNames, bool b
 		ActionStack->BeginAction(Action);
 	}
 
-	bool bSelectedSomething = false;
+	bool bSelectionChanged = false;
 
 	TArray<FName> PreviousSelection = Graph->GetSelectNodes();
 	for (const FName& PreviouslySelectedNode : PreviousSelection)
@@ -6379,8 +6379,7 @@ bool URigVMController::SetNodeSelection(const TArray<FName>& InNodeNames, bool b
 		{
 			if(Graph->SelectedNodes.Remove(PreviouslySelectedNode) > 0)
 			{
-				bSelectedSomething = true;
-				break;
+				bSelectionChanged = true;
 			}
 		}
 	}
@@ -6393,14 +6392,14 @@ bool URigVMController::SetNodeSelection(const TArray<FName>& InNodeNames, bool b
 			Graph->SelectedNodes.AddUnique(InNodeName);
 			if (PreviousNum != Graph->SelectedNodes.Num())
 			{
-				bSelectedSomething = true;
+				bSelectionChanged = true;
 			}
 		}
 	}
 
 	if (bSetupUndoRedo)
 	{
-		if (bSelectedSomething)
+		if (bSelectionChanged)
 		{
 			const TArray<FName>& SelectedNodes = Graph->GetSelectNodes();
 			if (SelectedNodes.Num() == 0)
@@ -6426,7 +6425,7 @@ bool URigVMController::SetNodeSelection(const TArray<FName>& InNodeNames, bool b
 		}
 	}
 
-	if (bSelectedSomething)
+	if (bSelectionChanged)
 	{
 		Notify(ERigVMGraphNotifType::NodeSelectionChanged, nullptr);
 	}
@@ -6452,7 +6451,7 @@ bool URigVMController::SetNodeSelection(const TArray<FName>& InNodeNames, bool b
 											*ArrayStr));
 	}
 
-	return bSelectedSomething;
+	return bSelectionChanged;
 }
 
 bool URigVMController::SetNodePosition(URigVMNode* InNode, const FVector2D& InPosition, bool bSetupUndoRedo, bool bMergeUndoAction, bool bPrintPythonCommand)
