RigVM: Avoid recreating pins if not required during RepopulatePins

diff --git a/Engine/Source/Developer/RigVMDeveloper/Private/RigVMModel/RigVMController.cpp b/Engine/Source/Developer/RigVMDeveloper/Private/RigVMModel/RigVMController.cpp
--- a/Engine/Source/Developer/RigVMDeveloper/Private/RigVMModel/RigVMController.cpp
+++ b/Engine/Source/Developer/RigVMDeveloper/Private/RigVMModel/RigVMController.cpp
@@ -11,7 +11,6 @@
 #include "RigVMCore/RigVMExecuteContext.h"
 #include "RigVMCore/RigVMUnknownType.h"
 #include "RigVMCore/RigVMByteCode.h"
-#include "RigVMFunctions/RigVMFunction_ControlFlow.h"
 #include "RigVMCompiler/RigVMCompiler.h"
 #include "RigVMDeveloperModule.h"
 #include "UObject/PropertyPortFlags.h"
@@ -24,12 +23,6 @@
 #include "RigVMPythonUtils.h"
 #include "RigVMTypeUtils.h"
 #include "Engine/UserDefinedStruct.h"
-#include "RigVMFunctions/RigVMDispatch_If.h"
-#include "RigVMFunctions/RigVMDispatch_Select.h"
-#include "RigVMFunctions/RigVMDispatch_Array.h"
-#include "RigVMModel/RigVMClient.h"
-#include "RigVMModel/Nodes/RigVMBranchNode.h"
-#include "RigVMModel/Nodes/RigVMArrayNode.h"
 
 #include UE_INLINE_GENERATED_CPP_BY_NAME(RigVMController)
 
@@ -70,16 +63,6 @@ FRigVMControllerCompileBracketScope::~FRigVMControllerCompileBracketScope()
 	Graph->Notify(ERigVMGraphNotifType::InteractionBracketClosed, nullptr);
 }
 
-void FRigVMClientPatchResult::Merge(const FRigVMClientPatchResult& InOther)
-{
-	bSucceeded = Succeeded() && InOther.Succeeded();
-	bChangedContent = ChangedContent() || InOther.ChangedContent();
-	bRequiresToMarkPackageDirty = RequiresToMarkPackageDirty() || InOther.RequiresToMarkPackageDirty();
-	ErrorMessages.Append(InOther.GetErrorMessages());
-	RemovedNodes.Append(InOther.GetRemovedNodes());
-	AddedNodes.Append(InOther.GetAddedNodes());
-}
-
 URigVMController::URigVMController()
 	: bValidatePinDefaults(true)
 	, bSuspendRecomputingOuterTemplateFilters(false)
@@ -87,12 +70,10 @@ URigVMController::URigVMController()
 	, bReportWarningsAndErrors(true)
 	, bIgnoreRerouteCompactnessChanges(false)
 	, UserLinkDirection(ERigVMPinDirection::Invalid)
-	, bEnableTypeCasting(true)
 	, bIsTransacting(false)
 	, bIsRunningUnitTest(false)
 	, bIsFullyResolvingTemplateNode(false)
 	, bSuspendRecomputingTemplateFilters(false)
-	, bSuspendTemplateComputation(false)
 #if WITH_EDITOR
 	, bRegisterTemplateNodeUsage(true)
 #endif
@@ -107,12 +88,10 @@ URigVMController::URigVMController(const FObjectInitializer& ObjectInitializer)
 	, bReportWarningsAndErrors(true)
 	, bIgnoreRerouteCompactnessChanges(false)
 	, UserLinkDirection(ERigVMPinDirection::Invalid)
-	, bEnableTypeCasting(true)
 	, bIsTransacting(false)
 	, bIsRunningUnitTest(false)
 	, bIsFullyResolvingTemplateNode(false)
 	, bSuspendRecomputingTemplateFilters(false)
-	, bSuspendTemplateComputation(false)
 #if WITH_EDITOR
 	, bRegisterTemplateNodeUsage(true)
 #endif
@@ -293,12 +272,400 @@ void URigVMController::SetIsRunningUnitTest(bool bIsRunning)
 {
 	bIsRunningUnitTest = bIsRunning;
 
-	if(URigVMBuildData* BuildData = URigVMBuildData::Get())
+	if(URigVMBuildData* BuildData = GetBuildData())
 	{
 		BuildData->SetIsRunningUnitTest(bIsRunning);
 	}	
 }
 
+URigVMController::FPinInfo::FPinInfo()
+	: ParentIndex(INDEX_NONE)
+	, Name(NAME_None)
+	, Direction(ERigVMPinDirection::Invalid)
+	, TypeIndex(INDEX_NONE)
+	, bIsArray(false)
+	, Property(nullptr)
+	, bIsExpanded(false)
+	, bIsConstant(false)
+	, bIsDynamicArray(false)
+{
+}
+
+URigVMController::FPinInfo::FPinInfo(const URigVMPin* InPin, int32 InParentIndex, ERigVMPinDirection InDirection)
+	: ParentIndex(InParentIndex)
+	, Name(InPin->GetName())
+	, Direction(InDirection == ERigVMPinDirection::Invalid ? InPin->GetDirection() : InDirection)
+	, TypeIndex(InPin->GetTypeIndex())
+	, bIsArray(InPin->IsArray())
+	, Property(nullptr)
+	, bIsExpanded(InPin->IsExpanded())
+	, bIsConstant(InPin->IsDefinedAsConstant())
+	, bIsDynamicArray(InPin->IsDynamicArray())
+{
+	// this method describes the info as currently represented in the model.
+
+	CorrectExecuteTypeIndex();
+	if(InPin->GetSubPins().IsEmpty())
+	{
+		DefaultValue = InPin->GetDefaultValue();
+	}
+}
+
+URigVMController::FPinInfo::FPinInfo(FProperty* InProperty, ERigVMPinDirection InDirection, int32 InParentIndex, const uint8* InDefaultValueMemory)
+	: ParentIndex(InParentIndex)
+	, Name(InProperty->GetFName())
+	, Direction(InDirection)
+	, TypeIndex(INDEX_NONE)
+	, bIsArray(InProperty->IsA<FArrayProperty>())
+	, Property(InProperty)
+	, bIsExpanded(false)
+	, bIsConstant(false)
+	, bIsDynamicArray(false)
+{
+	// this method describes the info as needed based on the property structure
+
+	if (Direction == ERigVMPinDirection::Invalid)
+	{
+		Direction = FRigVMStruct::GetPinDirectionFromProperty(Property);
+	}
+
+#if WITH_EDITOR
+
+	if (CastField<FArrayProperty>(InProperty->GetOwnerProperty()) == nullptr)
+	{
+		const FString DisplayNameText = InProperty->GetDisplayNameText().ToString();
+		if (!DisplayNameText.IsEmpty())
+		{
+			DisplayName = *DisplayNameText;
+		}
+	}
+	
+	bIsConstant = InProperty->HasMetaData(TEXT("Constant"));
+	CustomWidgetName = InProperty->GetMetaData(TEXT("CustomWidget"));
+	if (InProperty->HasMetaData(FRigVMStruct::ExpandPinByDefaultMetaName))
+	{
+		bIsExpanded = true;
+	}
+
+#endif
+
+#if WITH_EDITOR
+	if (Direction == ERigVMPinDirection::Hidden)
+	{
+		if (!InProperty->HasMetaData(TEXT("ArraySize")))
+		{
+			bIsDynamicArray = true;
+		}
+	}
+	if (bIsDynamicArray)
+	{
+		if (InProperty->HasMetaData(FRigVMStruct::SingletonMetaName))
+		{
+			bIsDynamicArray = false;
+		}
+	}
+#endif
+
+	UObject* CPPTypeObject = nullptr;
+	
+	FProperty* PropertyForType = InProperty;
+	if (const FArrayProperty* ArrayProperty = CastField<FArrayProperty>(PropertyForType))
+	{
+		PropertyForType = ArrayProperty->Inner;
+	}
+
+	if (const FStructProperty* StructProperty = CastField<FStructProperty>(PropertyForType))
+	{
+		CPPTypeObject = StructProperty->Struct;
+	}
+	else if (const FObjectProperty* ObjectProperty = CastField<FObjectProperty>(PropertyForType))
+	{
+		if(RigVMCore::SupportsUObjects())
+		{
+			CPPTypeObject = ObjectProperty->PropertyClass;
+		}
+	}
+	else if (const FInterfaceProperty* InterfaceProperty = CastField<FInterfaceProperty>(PropertyForType))
+	{
+		if (RigVMCore::SupportsUInterfaces())
+		{
+			CPPTypeObject = InterfaceProperty->InterfaceClass;
+		}
+	}
+	else
+	{
+		if (const FEnumProperty* EnumProperty = CastField<FEnumProperty>(PropertyForType))
+		{
+			CPPTypeObject = EnumProperty->GetEnum();
+		}
+		if (const FByteProperty* ByteProperty = CastField<FByteProperty>(PropertyForType))
+		{
+			CPPTypeObject = ByteProperty->Enum;
+		}
+
+		if(InDefaultValueMemory)
+		{
+			InProperty->ExportText_Direct(DefaultValue, InDefaultValueMemory, InDefaultValueMemory, nullptr, PPF_None, nullptr);
+		}
+	}
+
+	FString ExtendedCppType;
+	FString CPPType = InProperty->GetCPPType(&ExtendedCppType);
+	CPPType += ExtendedCppType;
+	CPPType = RigVMTypeUtils::PostProcessCPPType(CPPType, CPPTypeObject);
+	TypeIndex = FRigVMRegistry::Get().GetTypeIndexFromCPPType(CPPType);
+	CorrectExecuteTypeIndex();
+}
+
+void URigVMController::FPinInfo::CorrectExecuteTypeIndex()
+{
+	const FRigVMRegistry& Registry = FRigVMRegistry::Get();
+	if(Registry.IsExecuteType(TypeIndex))
+	{
+		TRigVMTypeIndex DefaultExecuteType = RigVMTypeUtils::TypeIndex::Execute;
+		if(Registry.IsArrayType(TypeIndex))
+		{
+			DefaultExecuteType = Registry.GetArrayTypeFromBaseTypeIndex(DefaultExecuteType);
+		}
+		TypeIndex = DefaultExecuteType;
+	}
+}
+
+uint32 GetTypeHash(const URigVMController::FPinInfo& InPin)
+{
+	uint32 Hash = 0; //GetTypeHash(InPin.ParentIndex);
+	Hash = HashCombine(Hash, GetTypeHash(InPin.Name));
+	Hash = HashCombine(Hash, GetTypeHash((int32)InPin.Direction));
+	Hash = HashCombine(Hash, GetTypeHash((int32)InPin.TypeIndex));
+	Hash = HashCombine(Hash, GetTypeHash(InPin.bIsArray));
+	// we are not hashing the parent index,  pinpath, default value or the property since
+	// it doesn't matter for the structure validity of the node
+	return Hash;
+}
+
+URigVMController::FPinInfoArray::FPinInfoArray(const URigVMNode* InNode)
+{
+	// this method adds all pins as currently represented in the model.
+	for(const URigVMPin* Pin : InNode->GetPins())
+	{
+		(void)AddPin(Pin, INDEX_NONE);
+	}
+}
+
+int32 URigVMController::FPinInfoArray::AddPin(const URigVMPin* InPin, int32 InParentIndex, ERigVMPinDirection InDirection)
+{
+	// this method adds all pins as currently represented in the model.
+	const int32 Index = Pins.Emplace(InPin, InParentIndex, InDirection);
+	for(const URigVMPin* SubPin : InPin->GetSubPins())
+	{
+		(void)AddPin(SubPin, Index, InDirection);
+	}
+	return Index;
+}
+
+URigVMController::FPinInfoArray::FPinInfoArray(const URigVMNode* InNode, URigVMController* InController)
+{
+	// this method adds pins as needed based on the property structure
+	for(const URigVMPin* Pin : InNode->GetPins())
+	{
+		const FString DefaultValue = Pin->GetDefaultValue();
+		(void)AddPin(InController, INDEX_NONE, Pin->GetFName(), Pin->GetDirection(), Pin->GetTypeIndex(), DefaultValue, nullptr);
+	}
+}
+
+int32 URigVMController::FPinInfoArray::AddPin(FProperty* InProperty, URigVMController* InController,
+	ERigVMPinDirection InDirection, int32 InParentIndex, const uint8* InDefaultValueMemory)
+{
+	// this method adds pins as needed based on the property structure
+
+	check(InDefaultValueMemory);
+	
+	const int32 Index = Pins.Emplace(InProperty, InDirection, InParentIndex, InDefaultValueMemory);
+
+	if (const FStructProperty* StructProperty = CastField<FStructProperty>(InProperty))
+	{
+		(void)AddPins(StructProperty->Struct, InController, Pins[Index].Direction, Index, InDefaultValueMemory);
+	}
+	else if (const FArrayProperty* ArrayProperty = CastField<FArrayProperty>(InProperty))
+	{
+		FScriptArrayHelper ArrayHelper(ArrayProperty, InDefaultValueMemory);
+		for(int32 ElementIndex = 0; ElementIndex < ArrayHelper.Num(); ElementIndex++)
+		{
+			const uint8* ElementDefaultValueMemory = ArrayHelper.GetRawPtr(ElementIndex);
+			const int32 SubIndex = AddPin(ArrayProperty->Inner, InController, Pins[Index].Direction, Index, ElementDefaultValueMemory);
+			Pins[SubIndex].Name = *FString::FormatAsNumber(ElementIndex);
+		}
+	}
+
+	return Index;
+}
+
+int32 URigVMController::FPinInfoArray::AddPin(URigVMController* InController, int32 InParentIndex, const FName& InName, ERigVMPinDirection InDirection,
+	TRigVMTypeIndex InTypeIndex, const FString& InDefaultValue, const uint8* InDefaultValueMemory)
+{
+	const FRigVMRegistry& Registry = FRigVMRegistry::Get();
+		
+	FPinInfo Info;
+	Info.ParentIndex = InParentIndex;
+	Info.Name = InName;
+	Info.Direction = InDirection;
+	Info.TypeIndex = InTypeIndex;
+	Info.bIsArray = Registry.IsArrayType(InTypeIndex);
+	Info.DefaultValue = InDefaultValue;
+	Info.CorrectExecuteTypeIndex();
+	const int32 Index = Pins.Add(Info);
+
+	const FRigVMTemplateArgumentType& Type = Registry.GetType(InTypeIndex);
+	if(!Type.IsWildCard())
+	{
+		if(Info.bIsArray)
+		{
+			const TRigVMTypeIndex& ElementTypeIndex = Registry.GetBaseTypeFromArrayTypeIndex(InTypeIndex);
+			const FRigVMTemplateArgumentType& ElementType = Registry.GetType(ElementTypeIndex);
+
+			const TArray<FString> Elements = URigVMPin::SplitDefaultValue(InDefaultValue);
+			for(int32 ElementIndex = 0; ElementIndex < Elements.Num(); ElementIndex++)
+			{
+				const FString& ElementDefaultValue = Elements[ElementIndex];
+				const uint8* ElementDefaultValueMemory = nullptr;
+				FStructOnScope ElementDefaultValueMemoryScope;
+
+				if(UScriptStruct* ElementScriptStruct = Cast<UScriptStruct>(ElementType.CPPTypeObject))
+				{
+					ElementDefaultValueMemoryScope = FStructOnScope(ElementScriptStruct);
+		
+					FRigVMPinDefaultValueImportErrorContext ErrorPipe;
+					ElementScriptStruct->ImportText(*ElementDefaultValue, ElementDefaultValueMemoryScope.GetStructMemory(), nullptr, PPF_None, &ErrorPipe, FString());
+					ElementDefaultValueMemory = ElementDefaultValueMemoryScope.GetStructMemory();
+				}
+
+				AddPin(InController, Index, *FString::FormatAsNumber(ElementIndex), InDirection, ElementTypeIndex, ElementDefaultValue, ElementDefaultValueMemory);
+			}
+		}
+		else if(UScriptStruct* ScriptStruct = Cast<UScriptStruct>(Type.CPPTypeObject))
+		{
+			const uint8* DefaultValueMemory = InDefaultValueMemory;
+			FStructOnScope DefaultValueMemoryScope;
+			if(DefaultValueMemory == nullptr)
+			{
+				FRigVMPinDefaultValueImportErrorContext ErrorPipe;
+				DefaultValueMemoryScope = FStructOnScope(ScriptStruct);
+				ScriptStruct->ImportText(*InDefaultValue, DefaultValueMemoryScope.GetStructMemory(), nullptr, PPF_None, &ErrorPipe, FString());
+				DefaultValueMemory = DefaultValueMemoryScope.GetStructMemory();
+			}
+			AddPins(ScriptStruct, InController, InDirection, Index, DefaultValueMemory);
+		}
+	}
+
+	return Index;
+}
+
+void URigVMController::FPinInfoArray::AddPins(UScriptStruct* InScriptStruct, URigVMController* InController,
+	ERigVMPinDirection InDirection, int32 InParentIndex, const uint8* InDefaultValueMemory)
+{
+	if (InController->ShouldStructBeUnfolded(InScriptStruct))
+	{
+		TArray<UStruct*> StructsToVisit = FRigVMTemplate::GetSuperStructs(InScriptStruct, true);
+		for(UStruct* StructToVisit : StructsToVisit)
+		{
+			// using EFieldIterationFlags::None excludes the
+			// properties of the super struct in this iterator.
+			for (TFieldIterator<FProperty> It(StructToVisit, EFieldIterationFlags::None); It; ++It)
+			{
+				const uint8* DefaultValueMemory = nullptr;
+				if(InDefaultValueMemory)
+				{
+					DefaultValueMemory = It->ContainerPtrToValuePtr<uint8>(InDefaultValueMemory);
+				}
+				(void)AddPin(*It, InController, InDirection, InParentIndex, DefaultValueMemory);
+			}
+		}
+	}
+}
+
+const FString& URigVMController::FPinInfoArray::GetPinPath(const int32 InIndex) const
+{
+	if(!Pins.IsValidIndex(InIndex))
+	{
+		static const FString EmptyString;
+		return EmptyString;
+	}
+
+	if(Pins[InIndex].PinPath.IsEmpty())
+	{
+		if(Pins[InIndex].ParentIndex == INDEX_NONE)
+		{
+			Pins[InIndex].PinPath = Pins[InIndex].Name.ToString();
+		}
+		else
+		{
+			Pins[InIndex].PinPath = URigVMPin::JoinPinPath(GetPinPath(Pins[InIndex].ParentIndex), Pins[InIndex].Name.ToString());
+		}
+	}
+
+	return Pins[InIndex].PinPath;
+}
+
+int32 URigVMController::FPinInfoArray::GetIndexFromPinPath(const FString& InPinPath) const
+{
+	if(PinPathLookup.Num() != Num())
+	{
+		PinPathLookup.Reset();
+		for(int32 Index=0; Index < Num(); Index++)
+		{
+			PinPathLookup.Add(GetPinPath(Index), Index);
+		}
+	}
+
+	if(const int32* Index = PinPathLookup.Find(InPinPath))
+	{
+		return *Index;
+	}
+	return INDEX_NONE;
+}
+
+const URigVMController::FPinInfo* URigVMController::FPinInfoArray::GetPinFromPinPath(const FString& InPinPath) const
+{
+	int32 Index = GetIndexFromPinPath(InPinPath);
+	if(Pins.IsValidIndex(Index))
+	{
+		return &Pins[Index];
+	}
+	return nullptr;
+}
+
+int32 URigVMController::FPinInfoArray::GetRootIndex(const int32 InIndex) const
+{
+	if(Pins.IsValidIndex(InIndex))
+	{
+		if(Pins[InIndex].ParentIndex == INDEX_NONE)
+		{
+			return InIndex;
+		}
+		return GetRootIndex(Pins[InIndex].ParentIndex);
+	}
+	return INDEX_NONE;
+}
+
+uint32 GetTypeHash(const URigVMController::FPinInfoArray& InPins)
+{
+	TArray<uint32> Hashes;
+	Hashes.Reserve(InPins.Num());
+	
+	uint32 OverAllHash = GetTypeHash(InPins.Num());
+	for(const URigVMController::FPinInfo& Info : InPins)
+	{
+		uint32 PinHash = GetTypeHash(Info);
+		if(Info.ParentIndex != INDEX_NONE)
+		{
+			PinHash = HashCombine(PinHash, Hashes[Info.ParentIndex]);
+		}
+		Hashes.Add(PinHash);
+		OverAllHash = HashCombine(OverAllHash, PinHash); 
+	}
+	return OverAllHash;
+}
+
 void URigVMController::HandleModifiedEvent(ERigVMGraphNotifType InNotifType, URigVMGraph* InGraph, UObject* InSubject)
 {
 	switch (InNotifType)
@@ -616,6 +983,36 @@ TArray<FString> URigVMController::GetAddNodePythonCommands(URigVMNode* Node) con
 					*RigVMPythonUtils::LinearColorToPythonString(CommentNode->GetNodeColor()),
 					*NodeName));	
 	}
+	else if (const URigVMBranchNode* BranchNode = Cast<URigVMBranchNode>(Node))
+	{
+		// add_branch_node(position=[0.0, 0.0], node_name='', undo=True)
+		Commands.Add(FString::Printf(TEXT("blueprint.get_controller_by_name('%s').add_branch_node(%s, '%s')"),
+						*GraphName,
+						*RigVMPythonUtils::Vector2DToPythonString(BranchNode->GetPosition()),
+						*NodeName));	
+	}
+	else if (const URigVMIfNode* IfNode = Cast<URigVMIfNode>(Node))
+	{
+		// add_if_node(cpp_type, cpp_type_object_path, position=[0.0, 0.0], node_name='', undo=True)
+		URigVMPin* ResultPin = IfNode->FindPin(URigVMIfNode::ResultName);
+		Commands.Add(FString::Printf(TEXT("blueprint.get_controller_by_name('%s').add_if_node('%s', '%s', %s, '%s')"),
+						*GraphName,
+						*ResultPin->GetCPPType(),
+						*ResultPin->CPPTypeObject->GetPathName(),
+						*RigVMPythonUtils::Vector2DToPythonString(IfNode->GetPosition()),
+						*NodeName));
+	}
+	else if (const URigVMSelectNode* SelectNode = Cast<URigVMSelectNode>(Node))
+	{
+		// add_select_node(cpp_type, cpp_type_object_path, position=[0.0, 0.0], node_name='', undo=True)
+		URigVMPin* ResultPin = SelectNode->FindPin(URigVMSelectNode::ResultName);
+		Commands.Add(FString::Printf(TEXT("blueprint.get_controller_by_name('%s').add_select_node('%s', '%s', %s, '%s')"),
+						*GraphName,
+						*ResultPin->GetCPPType(),
+						*ResultPin->CPPTypeObject->GetPathName(),
+						*RigVMPythonUtils::Vector2DToPythonString(SelectNode->GetPosition()),
+						*NodeName));
+	}
 	else if (const URigVMRerouteNode* RerouteNode = Cast<URigVMRerouteNode>(Node))
 	{
 		// add_free_reroute_node(bool bShowAsFullNode, const FString& InCPPType, const FName& InCPPTypeObjectPath, bool bIsConstant, const FName& InCustomWidgetName, const FString& InDefaultValue, const FVector2D& InPosition = FVector2D::ZeroVector, const FString& InNodeName = TEXT(""), bool bSetupUndoRedo = true);
@@ -630,6 +1027,31 @@ TArray<FString> URigVMController::GetAddNodePythonCommands(URigVMNode* Node) con
 				*RigVMPythonUtils::Vector2DToPythonString(RerouteNode->GetPosition()),
 				*NodeName));
 	}
+	else if (const URigVMArrayNode* ArrayNode = Cast<URigVMArrayNode>(Node))
+	{
+		// add_array_node(opcode, cpp_type, cpp_type_object, position=[0.0, 0.0], node_name='', undo=True)
+		if (ArrayNode->GetCPPTypeObject())
+		{
+			static constexpr TCHAR ArrayNodeFormat[] = TEXT("blueprint.get_controller_by_name('%s').add_array_node_from_object_path(%s, '%s', '%s', %s, '%s')");
+			Commands.Add(FString::Printf(ArrayNodeFormat,
+					*GraphName,
+					*RigVMPythonUtils::EnumValueToPythonString<ERigVMOpCode>((int64)ArrayNode->GetOpCode()),
+					*ArrayNode->GetCPPType(),
+					*ArrayNode->GetCPPTypeObject()->GetPathName(),
+					*RigVMPythonUtils::Vector2DToPythonString(ArrayNode->GetPosition()),
+					*NodeName));	
+		}
+		else
+		{
+			static constexpr TCHAR ArrayNodeFormat[] = TEXT("blueprint.get_controller_by_name('%s').add_array_node(%s, '%s', None, %s, '%s')");
+			Commands.Add(FString::Printf(ArrayNodeFormat,
+					*GraphName,
+					*RigVMPythonUtils::EnumValueToPythonString<ERigVMOpCode>((int64)ArrayNode->GetOpCode()),
+					*ArrayNode->GetCPPType(),
+					*RigVMPythonUtils::Vector2DToPythonString(ArrayNode->GetPosition()),
+					*NodeName));	
+		}
+	}
 	else if (const URigVMEnumNode* EnumNode = Cast<URigVMEnumNode>(Node))
 	{
 		// add_enum_node(cpp_type_object_path, position=[0.0, 0.0], node_name='', undo=True)
@@ -639,43 +1061,39 @@ TArray<FString> URigVMController::GetAddNodePythonCommands(URigVMNode* Node) con
 						*RigVMPythonUtils::Vector2DToPythonString(EnumNode->GetPosition()),
 						*NodeName));
 	}
-	else if (const URigVMFunctionReferenceNode* RefNode = Cast<URigVMFunctionReferenceNode>(Node))
+	else if (const URigVMLibraryNode* LibraryNode = Cast<URigVMLibraryNode>(Node))
 	{
-		if (RefNode->GetReferencedFunctionHeader().LibraryPointer.HostObject == GetGraph()->GetDefaultFunctionLibrary()->GetFunctionHostObjectPath())
+		const FString ContainedGraphName = GetSanitizedGraphName(LibraryNode->GetContainedGraph()->GetGraphName());
+
+		// AddFunctionReferenceNode(URigVMLibraryNode* InFunctionDefinition, const FVector2D& InNodePosition = FVector2D::ZeroVector, const FString& InNodeName = TEXT(""), bool bSetupUndoRedo = true);
+		URigVMFunctionLibrary* Library = LibraryNode->GetLibrary();
+		if (!Library || Library == GetGraph()->GetDefaultFunctionLibrary())
 		{
 			Commands.Add(FString::Printf(TEXT("blueprint.get_controller_by_name('%s').add_function_reference_node(function_%s, %s, '%s')"),
-					*GraphName,
-					*RigVMPythonUtils::PythonizeName(RefNode->LoadReferencedNode()->GetContainedGraph()->GetGraphName()),
-					*RigVMPythonUtils::Vector2DToPythonString(RefNode->GetPosition()), 
-					*NodeName));
+						*GraphName,
+						*RigVMPythonUtils::PythonizeName(ContainedGraphName),
+						*RigVMPythonUtils::Vector2DToPythonString(LibraryNode->GetPosition()), 
+						*NodeName));
 		}
 		else
 		{
-			Commands.Add(FString::Printf(TEXT("blueprint.get_controller_by_name('%s').add_external_function_reference_node('%s', '%s', %s, '%s')"),
+			Commands.Add(FString::Printf(TEXT("function_blueprint = unreal.load_object(name = '%s', outer = None)"),
+				*Library->GetOuter()->GetPathName()));
+			Commands.Add(FString::Printf(TEXT("blueprint.get_controller_by_name('%s').add_function_reference_node(function_blueprint.get_local_function_library().find_function('%s'), %s, '%s')"),
 						*GraphName,
-						*RefNode->GetReferencedFunctionHeader().LibraryPointer.HostObject.ToString(),
-						*RefNode->GetReferencedFunctionHeader().Name.ToString(),
-						*RigVMPythonUtils::Vector2DToPythonString(RefNode->GetPosition()), 
+						*NodeName,
+						*RigVMPythonUtils::Vector2DToPythonString(LibraryNode->GetPosition()), 
 						*NodeName));
 		}
-	}
-	else if (const URigVMCollapseNode* CollapseNode = Cast<URigVMCollapseNode>(Node))
-	{
-		const FString ContainedGraphName = GetSanitizedGraphName(CollapseNode->GetContainedGraph()->GetGraphName());
-		// AddFunctionReferenceNode(URigVMLibraryNode* InFunctionDefinition, const FVector2D& InNodePosition = FVector2D::ZeroVector, const FString& InNodeName = TEXT(""), bool bSetupUndoRedo = true);
-		Commands.Add(FString::Printf(TEXT("blueprint.get_controller_by_name('%s').add_function_reference_node(function_%s, %s, '%s')"),
+
+		if (Node->IsA<URigVMCollapseNode>())
+		{
+			Commands.Add(FString::Printf(TEXT("blueprint.get_controller_by_name('%s').promote_function_reference_node_to_collapse_node('%s')"),
 					*GraphName,
-					*RigVMPythonUtils::PythonizeName(ContainedGraphName),
-					*RigVMPythonUtils::Vector2DToPythonString(CollapseNode->GetPosition()), 
 					*NodeName));
-	
-
-		Commands.Add(FString::Printf(TEXT("blueprint.get_controller_by_name('%s').promote_function_reference_node_to_collapse_node('%s')"),
-				*GraphName,
-				*NodeName));
-		Commands.Add(FString::Printf(TEXT("library_controller.remove_function_from_library('%s')"),
-				*ContainedGraphName));
-	
+			Commands.Add(FString::Printf(TEXT("library_controller.remove_function_from_library('%s')"),
+					*ContainedGraphName));
+		}
 	}
 	else if (const URigVMInvokeEntryNode* InvokeEntryNode = Cast<URigVMInvokeEntryNode>(Node))
 	{
@@ -788,21 +1206,6 @@ URigVMUnitNode* URigVMController::AddUnitNode(UScriptStruct* InScriptStruct, con
 		return nullptr;
 	}
 
-	if(IRigVMClientHost* ClientHost = GetImplementingOuter<IRigVMClientHost>())
-	{
-		if(const FRigVMClient* Client = ClientHost->GetRigVMClient())
-		{
-			if(!Function->SupportsExecuteContextStruct(Client->GetExecuteContextStruct()))
-			{
-				ReportErrorf(TEXT("Cannot add node for function '%s' - incompatible execute context: '%s' vs '%s'."),
-						*Function->GetName(),
-						*Function->GetExecuteContextStruct()->GetStructCPPName(),
-						*Client->GetExecuteContextStruct()->GetStructCPPName());
-				return nullptr;
-			}
-		}
-	}
-
 	FString StructureError;
 	if (!FRigVMStruct::ValidateStruct(InScriptStruct, &StructureError))
 	{
@@ -810,6 +1213,8 @@ URigVMUnitNode* URigVMController::AddUnitNode(UScriptStruct* InScriptStruct, con
 		return nullptr;
 	}
 
+#if UE_RIGVM_ENABLE_TEMPLATE_NODES
+
 	if(const FRigVMTemplate* Template = Function->GetTemplate())
 	{
 		if(bSetupUndoRedo)
@@ -856,8 +1261,11 @@ URigVMUnitNode* URigVMController::AddUnitNode(UScriptStruct* InScriptStruct, con
 		return TemplateNode;
 	}
 
+#endif
+
 	FStructOnScope StructOnScope(InScriptStruct);
 	FRigVMStruct* StructMemory = (FRigVMStruct*)StructOnScope.GetStructMemory();
+	InScriptStruct->InitializeDefaultValue((uint8*)StructMemory);
 	const bool bIsEventNode = (!StructMemory->GetEventName().IsNone());
 	if (bIsEventNode)
 	{
@@ -897,10 +1305,7 @@ URigVMUnitNode* URigVMController::AddUnitNode(UScriptStruct* InScriptStruct, con
 
 	FString ExportedDefaultValue;
 	CreateDefaultValueForStructIfRequired(InScriptStruct, ExportedDefaultValue);
-	{
-		TGuardValue<bool> SuspendNotifications(bSuspendNotifications, true);
-		AddPinsForStruct(InScriptStruct, Node, nullptr, ERigVMPinDirection::Invalid, ExportedDefaultValue, true);
-	}
+	AddPinsForStruct(InScriptStruct, Node, nullptr, ERigVMPinDirection::Invalid, ExportedDefaultValue, true);
 
 	Graph->Nodes.Add(Node);
 	if (!bSuspendNotifications)
@@ -959,7 +1364,7 @@ URigVMUnitNode* URigVMController::AddUnitNodeFromStructPath(const FString& InScr
 		return nullptr;
 	}
 
-	UScriptStruct* ScriptStruct = RigVMTypeUtils::FindObjectFromCPPTypeObjectPath<UScriptStruct>(InScriptStructPath);
+	UScriptStruct* ScriptStruct = URigVMPin::FindObjectFromCPPTypeObjectPath<UScriptStruct>(InScriptStructPath);
 	if (ScriptStruct == nullptr)
 	{
 		ReportErrorf(TEXT("Cannot find struct for path '%s'."), *InScriptStructPath);
@@ -1177,7 +1582,7 @@ URigVMVariableNode* URigVMController::AddVariableNode(const FName& InVariableNam
 	}
 	if (InCPPTypeObject == nullptr)
 	{
-		InCPPTypeObject = RigVMTypeUtils::FindObjectFromCPPTypeObjectPath<UObject>(InCPPType);
+		InCPPTypeObject = URigVMPin::FindObjectFromCPPTypeObjectPath<UObject>(InCPPType);
 	}
 
 	FString CPPType = RigVMTypeUtils::PostProcessCPPType(InCPPType, InCPPTypeObject);
@@ -1248,10 +1653,7 @@ URigVMVariableNode* URigVMController::AddVariableNode(const FName& InVariableNam
 	{
 		FString DefaultValue = InDefaultValue;
 		CreateDefaultValueForStructIfRequired(ValuePin->GetScriptStruct(), DefaultValue);
-		{
-			TGuardValue<bool> SuspendNotifications(bSuspendNotifications, true);
-			AddPinsForStruct(ValuePin->GetScriptStruct(), Node, ValuePin, ValuePin->Direction, DefaultValue, false);
-		}
+		AddPinsForStruct(ValuePin->GetScriptStruct(), Node, ValuePin, ValuePin->Direction, DefaultValue, false);
 	}
 	else if (!InDefaultValue.IsEmpty() && InDefaultValue != TEXT("()"))
 	{
@@ -1312,7 +1714,7 @@ URigVMVariableNode* URigVMController::AddVariableNodeFromObjectPath(const FName&
 	UObject* CPPTypeObject = nullptr;
 	if (!InCPPTypeObjectPath.IsEmpty())
 	{
-		CPPTypeObject = RigVMTypeUtils::FindObjectFromCPPTypeObjectPath<UObject>(InCPPTypeObjectPath);
+		CPPTypeObject = URigVMPin::FindObjectFromCPPTypeObjectPath<UObject>(InCPPTypeObjectPath);
 		if (CPPTypeObject == nullptr)
 		{
 			ReportErrorf(TEXT("Cannot find cpp type object for path '%s'."), *InCPPTypeObjectPath);
@@ -1703,7 +2105,7 @@ void URigVMController::OnExternalVariableTypeChangedFromObjectPath(const FName&
 	UObject* CPPTypeObject = nullptr;
 	if (!InCPPTypeObjectPath.IsEmpty())
 	{
-		CPPTypeObject = RigVMTypeUtils::FindObjectFromCPPTypeObjectPath<UObject>(InCPPTypeObjectPath);
+		CPPTypeObject = URigVMPin::FindObjectFromCPPTypeObjectPath<UObject>(InCPPTypeObjectPath);
 		if (CPPTypeObject == nullptr)
 		{
 			ReportErrorf(TEXT("Cannot find cpp type object for path '%s'."), *InCPPTypeObjectPath);
@@ -1876,27 +2278,6 @@ bool URigVMController::UnresolveTemplateNodes(const TArray<URigVMNode*>& InNodes
 					URigVMNode* OtherNode = SourcePin->GetNode() == Node ? TargetPin->GetNode() : SourcePin->GetNode();
 					if (!InNodes.Contains(OtherNode))
 					{
-						const URigVMPin* PinOnNode = SourcePin->GetNode() == Node ? SourcePin : TargetPin;
-						if(PinOnNode->IsExecuteContext())
-						{
-							continue;
-						}
-						
-						if (const URigVMTemplateNode* TemplateNode = Cast<URigVMTemplateNode>(Node))
-						{
-							if(const FRigVMTemplate* Template = TemplateNode->GetTemplate())
-							{
-								const URigVMPin* RootPin = PinOnNode->GetRootPin();
-								if(const FRigVMTemplateArgument* Argument = Template->FindArgument(RootPin->GetFName()))
-								{
-									if(Argument->IsSingleton())
-									{
-										continue;
-									}
-								}
-							}
-						}
-						
 						BreakLink(SourcePin, TargetPin, bSetupUndoRedo);
 					}			
 				}
@@ -2221,7 +2602,7 @@ URigVMParameterNode* URigVMController::AddParameterNodeFromObjectPath(const FNam
 	UObject* CPPTypeObject = nullptr;
 	if (!InCPPTypeObjectPath.IsEmpty())
 	{
-		CPPTypeObject = RigVMTypeUtils::FindObjectFromCPPTypeObjectPath<UObject>(InCPPTypeObjectPath);
+		CPPTypeObject = URigVMPin::FindObjectFromCPPTypeObjectPath<UObject>(InCPPTypeObjectPath);
 		if (CPPTypeObject == nullptr)
 		{
 			ReportErrorf(TEXT("Cannot find cpp type object for path '%s'."), *InCPPTypeObjectPath);
@@ -2437,7 +2818,6 @@ URigVMRerouteNode* URigVMController::AddRerouteNodeOnPin(const FString& InPinPat
 
 	if (ValuePin->IsStruct())
 	{
-		TGuardValue<bool> SuspendNotifications(bSuspendNotifications, true);
 		AddPinsForStruct(ValuePin->GetScriptStruct(), Node, ValuePin, ValuePin->Direction, FString(), false);
 	}
 
@@ -2727,7 +3107,7 @@ URigVMInjectionInfo* URigVMController::AddInjectedNodeFromStructPath(const FStri
 		return nullptr;
 	}
 
-	UScriptStruct* ScriptStruct = RigVMTypeUtils::FindObjectFromCPPTypeObjectPath<UScriptStruct>(InScriptStructPath);
+	UScriptStruct* ScriptStruct = URigVMPin::FindObjectFromCPPTypeObjectPath<UScriptStruct>(InScriptStructPath);
 	if (ScriptStruct == nullptr)
 	{
 		ReportErrorf(TEXT("Cannot find struct for path '%s'."), *InScriptStructPath);
@@ -3386,39 +3766,13 @@ bool URigVMController::CanImportNodesFromText(const FString& InText)
 		return false;
 	}
 
-	FRigVMControllerObjectFactory Factory(this);
-	if (!Factory.CanCreateObjectsFromText(InText))
-	{
-		return false;
-	}
-
-	URigVMGraph* Graph = GetGraph();
-	check(Graph);
-	Factory.ProcessBuffer(Graph, RF_Transactional, InText);
-
-	if (Factory.CreatedNodes.Num() == 0)
+	if (GetGraph()->IsA<URigVMFunctionLibrary>())
 	{
 		return false;
 	}
 
-	// If the graph is a function library, make sure we are pasting unconnected collapse nodes
-	if (Graph->IsA<URigVMFunctionLibrary>())
-	{
-		if (!Factory.CreatedLinks.IsEmpty())
-		{
-			return false;
-		}
-		
-		for (URigVMNode* Node : Factory.CreatedNodes)
-		{
-			if (!Node->IsA<URigVMCollapseNode>())
-			{
-				return false;
-			}
-		}
-	}
-
-	return true;
+	FRigVMControllerObjectFactory Factory(nullptr);
+	return Factory.CanCreateObjectsFromText(InText);
 }
 
 TArray<FName> URigVMController::ImportNodesFromText(const FString& InText, bool bSetupUndoRedo, bool bPrintPythonCommands)
@@ -3462,9 +3816,9 @@ TArray<FName> URigVMController::ImportNodesFromText(const FString& InText, bool
 	for (int32 i=0; i<CreatedNodes.Num(); ++i)
 	{
 		URigVMNode* CreatedNode = CreatedNodes[i];
-		if (URigVMCollapseNode* CollapseNode = Cast<URigVMCollapseNode>(CreatedNode))
+		if (URigVMLibraryNode* LibraryNode = Cast<URigVMLibraryNode>(CreatedNode))
 		{
-			if (URigVMGraph* ContainedGraph = CollapseNode->GetContainedGraph())
+			if (URigVMGraph* ContainedGraph = LibraryNode->GetContainedGraph())
 			{
 				EditGuards.Emplace(ContainedGraph->bEditable, true);
 				CreatedNodes.Append(ContainedGraph->GetNodes());
@@ -3500,7 +3854,6 @@ TArray<FName> URigVMController::ImportNodesFromText(const FString& InText, bool
 		{
 			if (URigVMUnitNode* UnitNode = Cast<URigVMUnitNode>(SubNodes[SubNodeIndex]))
 			{
-				TGuardValue<bool> SuspendNotifications(bSuspendNotifications, true);
 				RepopulatePinsOnNode(UnitNode);
 			}
 		}
@@ -3520,7 +3873,6 @@ TArray<FName> URigVMController::ImportNodesFromText(const FString& InText, bool
 				{
 					FRigVMControllerGraphGuard GraphGuard(this, CollapseNode->GetContainedGraph(), false);
 					TGuardValue<bool> GuardEditGraph(CollapseNode->ContainedGraph->bEditable, true);
-					TGuardValue<bool> SuspendNotifications(bSuspendNotifications, true);
 					ReattachLinksToPinObjects();
 				}
 				
@@ -3587,9 +3939,12 @@ TArray<FName> URigVMController::ImportNodesFromText(const FString& InText, bool
 
 			if (URigVMFunctionReferenceNode* FunctionRefNode = Cast<URigVMFunctionReferenceNode>(SubNode))
 			{
-				if(URigVMBuildData* BuildData = URigVMBuildData::Get())
+				if (URigVMLibraryNode* FunctionDefinition = FunctionRefNode->GetReferencedNode())
 				{
-					BuildData->RegisterFunctionReference(FunctionRefNode->GetReferencedFunctionHeader().LibraryPointer, FunctionRefNode);
+					if(URigVMBuildData* BuildData = GetBuildData())
+					{
+						BuildData->RegisterFunctionReference(FunctionDefinition, FunctionRefNode);
+					}
 				}
 			}
 
@@ -3755,14 +4110,14 @@ TArray<FName> URigVMController::ImportNodesFromText(const FString& InText, bool
 									ActionStack->AddAction(FRigVMInjectNodeIntoPinAction(TargetPin->GetTypedOuter<URigVMInjectionInfo>()));	
 								}
 							}
-							Notify(ERigVMGraphNotifType::LinkAdded, CreatedLink);
 						}
+						Notify(ERigVMGraphNotifType::LinkAdded, CreatedLink);
 						continue;
 					}
 				}
 			}
 
-			ReportErrorf(TEXT("Cannot import link '%s'."), *URigVMLink::GetPinPathRepresentation(CreatedLink->SourcePinPath, CreatedLink->TargetPinPath));
+			ReportErrorf(TEXT("Cannot import link '%s -> %s'."), *CreatedLink->SourcePinPath, *CreatedLink->TargetPinPath);
 			DestroyObject(CreatedLink);
 		}
 
@@ -3796,53 +4151,8 @@ TArray<FName> URigVMController::ImportNodesFromText(const FString& InText, bool
 	return NodeNames;
 }
 
-URigVMLibraryNode* URigVMController::LocalizeFunctionFromPath(const FString& InHostPath, const FName& InFunctionName, bool bLocalizeDependentPrivateFunctions, bool bSetupUndoRedo, bool bPrintPythonCommand)
-{
-	if (!IsValidGraph())
-	{
-		return nullptr;
-	}
-
-	if (!bIsTransacting && !IsGraphEditable())
-	{
-		return nullptr;
-	}
-
-	URigVMGraph* Graph = GetGraph();
-	check(Graph);
-
-	if (Graph->IsA<URigVMFunctionLibrary>())
-	{
-		ReportError(TEXT("Cannot add function reference nodes to function library graphs."));
-		return nullptr;
-	}
-
-	UObject* HostObject = StaticLoadObject(UObject::StaticClass(), NULL, *InHostPath, NULL, LOAD_None, NULL);
-	if (!HostObject)
-	{
-		ReportErrorf(TEXT("Failed to load the Host object %s."), *InHostPath);
-		return nullptr;
-	}
-
-	IRigVMGraphFunctionHost* FunctionHost = Cast<IRigVMGraphFunctionHost>(HostObject);
-	if (!FunctionHost)
-	{
-		ReportError(TEXT("Host object is not a IRigVMGraphFunctionHost."));
-		return nullptr;
-	}
-
-	FRigVMGraphFunctionData* Data = FunctionHost->GetRigVMGraphFunctionStore()->FindFunctionByName(InFunctionName);
-	if (!Data)
-	{
-		ReportErrorf(TEXT("Function %s not found in host %s."), *InFunctionName.ToString(), *InHostPath);
-		return nullptr;
-	}
-
-	return LocalizeFunction(Data->Header.LibraryPointer, bLocalizeDependentPrivateFunctions, bSetupUndoRedo, bPrintPythonCommand);
-}
-
 URigVMLibraryNode* URigVMController::LocalizeFunction(
-	const FRigVMGraphFunctionIdentifier& InFunctionDefinition,
+	URigVMLibraryNode* InFunctionDefinition,
 	bool bLocalizeDependentPrivateFunctions,
 	bool bSetupUndoRedo,
 	bool bPrintPythonCommand)
@@ -3857,10 +4167,15 @@ URigVMLibraryNode* URigVMController::LocalizeFunction(
 		return nullptr;
 	}
 
-	TArray<FRigVMGraphFunctionIdentifier> FunctionsToLocalize;
+	if(InFunctionDefinition == nullptr)
+	{
+		return nullptr;
+	}
+
+	TArray<URigVMLibraryNode*> FunctionsToLocalize;
 	FunctionsToLocalize.Add(InFunctionDefinition);
 
-	TMap<FRigVMGraphFunctionIdentifier, URigVMLibraryNode*> Results = LocalizeFunctions(FunctionsToLocalize, bLocalizeDependentPrivateFunctions, bSetupUndoRedo, bPrintPythonCommand);
+	TMap<URigVMLibraryNode*, URigVMLibraryNode*> Results = LocalizeFunctions(FunctionsToLocalize, bLocalizeDependentPrivateFunctions, bSetupUndoRedo, bPrintPythonCommand);
 
 	URigVMLibraryNode** LocalizedFunctionPtr = Results.Find(FunctionsToLocalize[0]);
 	if(LocalizedFunctionPtr)
@@ -3870,12 +4185,12 @@ URigVMLibraryNode* URigVMController::LocalizeFunction(
 	return nullptr;
 }
 
-TMap<FRigVMGraphFunctionIdentifier, URigVMLibraryNode*> URigVMController::LocalizeFunctions(
-	TArray<FRigVMGraphFunctionIdentifier> InFunctionDefinitions,
+TMap<URigVMLibraryNode*, URigVMLibraryNode*> URigVMController::LocalizeFunctions(
+	TArray<URigVMLibraryNode*> InFunctionDefinitions,
 	bool bLocalizeDependentPrivateFunctions,
 	bool bSetupUndoRedo, bool bPrintPythonCommand)
 {
-	TMap<FRigVMGraphFunctionIdentifier, URigVMLibraryNode*> LocalizedFunctions;
+	TMap<URigVMLibraryNode*, URigVMLibraryNode*> LocalizedFunctions;
 
 	if(!IsValidGraph())
 	{
@@ -3896,65 +4211,87 @@ TMap<FRigVMGraphFunctionIdentifier, URigVMLibraryNode*> URigVMController::Locali
 		return LocalizedFunctions;
 	}
 
-	TArray<FRigVMGraphFunctionData*> FunctionsToLocalize;
+	TArray<URigVMLibraryNode*> FunctionsToLocalize;
 
-	TArray<FRigVMGraphFunctionIdentifier> NodesToVisit;
-	for(const FRigVMGraphFunctionIdentifier& FunctionDefinition : InFunctionDefinitions)
+	TArray<URigVMLibraryNode*> NodesToVisit;
+	for(URigVMLibraryNode* FunctionDefinition : InFunctionDefinitions)
 	{
 		NodesToVisit.AddUnique(FunctionDefinition);
-		FunctionsToLocalize.AddUnique(FRigVMGraphFunctionData::FindFunctionData(FunctionDefinition));
+		FunctionsToLocalize.AddUnique(FunctionDefinition);
 	}
 
-	const FSoftObjectPath ThisFunctionHost = ThisLibrary->GetFunctionHostObjectPath();
-	
 	// find all functions to localize
 	for(int32 NodeToVisitIndex=0; NodeToVisitIndex<NodesToVisit.Num(); NodeToVisitIndex++)
 	{
-		FRigVMGraphFunctionIdentifier& NodeToVisit = NodesToVisit[NodeToVisitIndex];
-
-		// Already local
-		if (NodeToVisit.HostObject == ThisFunctionHost)
-		{
-			continue;
-		}
+		URigVMLibraryNode* NodeToVisit = NodesToVisit[NodeToVisitIndex];
 
-		bool bIsPublic;
-		FRigVMGraphFunctionData* FunctionData = FRigVMGraphFunctionData::FindFunctionData(NodeToVisit, &bIsPublic);
-		if (!FunctionData)
+		if(URigVMCollapseNode* CollapseNode = Cast<URigVMCollapseNode>(NodeToVisit))
 		{
-			ReportAndNotifyErrorf(TEXT("Cannot localize function - could not find function %s in host %s."), *NodeToVisit.LibraryNode.ToString(), *NodeToVisit.HostObject.ToString());
-			return LocalizedFunctions;
-		}
+			const TArray<URigVMNode*>& ContainedNodes = CollapseNode->GetContainedNodes();
+			for(URigVMNode* ContainedNode : ContainedNodes)
+			{
+				if(URigVMLibraryNode* ContainedLibraryNode = Cast<URigVMLibraryNode>(ContainedNode))
+				{
+					NodesToVisit.AddUnique(ContainedLibraryNode);
+				}
+			}
 
-		// Do not localize public functions
-		if (bIsPublic)
-		{
-			continue;
-		}
+			if(URigVMFunctionLibrary* OtherLibrary = Cast<URigVMFunctionLibrary>(CollapseNode->GetOuter()))
+			{
+				if(OtherLibrary != ThisLibrary)
+				{
+					bool bIsAvailable = false;
+					if(IsFunctionAvailableDelegate.IsBound())
+					{
+						bIsAvailable = IsFunctionAvailableDelegate.Execute(CollapseNode);
+					}
 
-		if (!bLocalizeDependentPrivateFunctions)
-		{
-			ReportAndNotifyErrorf(TEXT("Cannot localize function - dependency %s is private."), *NodeToVisit.LibraryNode.ToString());
-			return LocalizedFunctions;
+					if(!bIsAvailable)
+					{
+						if(!bLocalizeDependentPrivateFunctions)
+						{
+							ReportAndNotifyErrorf(TEXT("Cannot localize function - dependency %s is private."), *CollapseNode->GetPathName());
+							return LocalizedFunctions;
+						}
+						
+						FunctionsToLocalize.AddUnique(CollapseNode);
+					}
+				}
+			}
 		}
 
-		FunctionsToLocalize.AddUnique(FunctionData);
-
-		// Look for its dependencies
-		for (TPair<FRigVMGraphFunctionIdentifier, uint32>& Pair : FunctionData->Header.Dependencies)
+		else if(URigVMFunctionReferenceNode* FunctionReferencedNode = Cast<URigVMFunctionReferenceNode>(NodeToVisit))
 		{
-			NodesToVisit.AddUnique(Pair.Key);
+			if(FunctionReferencedNode->GetLibrary() != ThisLibrary)
+			{
+				if(URigVMCollapseNode* FunctionDefinition = Cast<URigVMCollapseNode>(FunctionReferencedNode->GetReferencedNode()))
+				{
+					NodesToVisit.AddUnique(FunctionDefinition);
+				}
+			}
 		}
 	}
 	
 	// sort the functions to localize based on their nesting
-	Algo::Sort(FunctionsToLocalize, [](FRigVMGraphFunctionData* A, FRigVMGraphFunctionData* B) -> bool
+	Algo::Sort(FunctionsToLocalize, [](URigVMLibraryNode* A, URigVMLibraryNode* B) -> bool
 	{
 		check(A);
 		check(B);
-		return B->Header.Dependencies.Contains(A->Header.LibraryPointer);
+		return B->Contains(A);
 	});
 
+	// export all of the content for each node
+	TMap<URigVMLibraryNode*, FString> ExportedTextPerFunction;
+	for(URigVMLibraryNode* FunctionToLocalize : FunctionsToLocalize)
+	{
+		URigVMFunctionLibrary* OtherLibrary = Cast<URigVMFunctionLibrary>(FunctionToLocalize->GetOuter());
+		FRigVMControllerGraphGuard GraphGuard(this, OtherLibrary, false);
+
+		const TArray<FName> NodeNamesToExport = {FunctionToLocalize->GetFName()};
+		const FString ExportedText = ExportNodesToText(NodeNamesToExport);
+		ExportedTextPerFunction.Add(FunctionToLocalize, ExportedText);
+	}
+
 	FRigVMControllerCompileBracketScope CompileScope(this);
 	if (bSetupUndoRedo)
 	{
@@ -3964,23 +4301,34 @@ TMap<FRigVMGraphFunctionIdentifier, URigVMLibraryNode*> URigVMController::Locali
 	// import the functions to our local function library
 	{
 		FRigVMControllerGraphGuard GraphGuard(this, ThisLibrary, bSetupUndoRedo);
-		for(FRigVMGraphFunctionData* FunctionToLocalize : FunctionsToLocalize)
+
+		// override the availability and check up later
+		TGuardValue<FRigVMController_IsFunctionAvailableDelegate> IsFunctionAvailableGuard(IsFunctionAvailableDelegate,
+			FRigVMController_IsFunctionAvailableDelegate::CreateLambda([](URigVMLibraryNode*)
+			{
+				return true;
+			})
+		);
+
+		for(URigVMLibraryNode* FunctionToLocalize : FunctionsToLocalize)
 		{
-			if (URigVMLibraryNode* ReferencedFunction = Cast<URigVMLibraryNode>(FunctionToLocalize->Header.LibraryPointer.LibraryNode.TryLoad()))
+			const FString& ExportedText = ExportedTextPerFunction.FindChecked(FunctionToLocalize);
+			TArray<FName> ImportedNodeNames = ImportNodesFromText(ExportedText);
+			if(ImportedNodeNames.Num() != 1)
 			{
-				if (IRigVMClientHost* ClientHost = ReferencedFunction->GetImplementingOuter<IRigVMClientHost>())
-				{
-					ClientHost->GetRigVMClient()->UpdateGraphFunctionSerializedGraph(ReferencedFunction);
-				}
+				ReportErrorf(TEXT("Not possible to localize function %s"), *FunctionToLocalize->GetPathName());
+				continue;
 			}
-			
-			TArray<FName> NodeNames = ImportNodesFromText(FunctionToLocalize->SerializedCollapsedNode, false);
-			if (NodeNames.Num() > 0)
+
+			URigVMLibraryNode* LocalizedFunction = Cast<URigVMLibraryNode>(GetGraph()->FindNodeByName(ImportedNodeNames[0]));
+			if(LocalizedFunction == nullptr)
 			{
-				URigVMLibraryNode* LocalizedFunction = ThisLibrary->FindFunction(NodeNames[0]);
-				LocalizedFunctions.Add(FunctionToLocalize->Header.LibraryPointer, LocalizedFunction);
-				ThisLibrary->LocalizedFunctions.FindOrAdd(FunctionToLocalize->Header.LibraryPointer.LibraryNode.ToString(), LocalizedFunction);
+				ReportErrorf(TEXT("Not possible to localize function %s"), *FunctionToLocalize->GetPathName());
+				continue;
 			}
+
+			LocalizedFunctions.Add(FunctionToLocalize, LocalizedFunction);
+			ThisLibrary->LocalizedFunctions.FindOrAdd(FunctionToLocalize->GetPathName(), LocalizedFunction);
 		}
 	}
 
@@ -4004,7 +4352,8 @@ TMap<FRigVMGraphFunctionIdentifier, URigVMLibraryNode*> URigVMController::Locali
 			}
 			else if(URigVMFunctionReferenceNode* FunctionReferenceNode = Cast<URigVMFunctionReferenceNode>(NodeToUpdate))
 			{
-				URigVMLibraryNode** RemappedNodePtr = LocalizedFunctions.Find(FunctionReferenceNode->GetReferencedFunctionHeader().LibraryPointer);
+				URigVMLibraryNode* ReferencedNode = FunctionReferenceNode->GetReferencedNode();
+				URigVMLibraryNode** RemappedNodePtr = LocalizedFunctions.Find(ReferencedNode);
 				if(RemappedNodePtr)
 				{
 					URigVMLibraryNode* RemappedNode = *RemappedNodePtr;
@@ -4021,19 +4370,26 @@ TMap<FRigVMGraphFunctionIdentifier, URigVMLibraryNode*> URigVMController::Locali
 
 	if (bPrintPythonCommand)
 	{
-		for (const FRigVMGraphFunctionIdentifier& Identifier : InFunctionDefinitions)
+		FString FunctionNames = TEXT("[");
+		for (auto It = InFunctionDefinitions.CreateConstIterator(); It; ++It)
 		{
-			const FString GraphName = GetSanitizedGraphName(GetGraph()->GetGraphName());
-			FRigVMGraphFunctionData* FunctionData = FRigVMGraphFunctionData::FindFunctionData(Identifier);
-			const FString FunctionDefinitionName = GetSanitizedNodeName(FunctionData->Header.Name.ToString());
-
-			RigVMPythonUtils::Print(GetGraphOuterName(), 
-				FString::Printf(TEXT("blueprint.get_controller_by_name('%s').localize_function_from_path('%s', '%s', %s)"),
-						*GraphName,
-						*Identifier.HostObject.ToString(),
-						*FunctionDefinitionName,
-						bLocalizeDependentPrivateFunctions ? TEXT("True") : TEXT("False")));
+			FunctionNames += FString::Printf(TEXT("unreal.load_object(name = '%s', outer = None).get_local_function_library().find_function('%s')"),
+				*(*It)->GetLibrary()->GetOuter()->GetPathName(),
+				*(*It)->GetName());
+			if (It.GetIndex() < InFunctionDefinitions.Num() - 1)
+			{
+				FunctionNames += TEXT(", ");
+			}
 		}
+		FunctionNames += TEXT("]");
+
+		const FString GraphName = GetSanitizedGraphName(GetGraph()->GetGraphName());
+
+		RigVMPythonUtils::Print(GetGraphOuterName(), 
+							FString::Printf(TEXT("blueprint.get_controller_by_name('%s').localize_functions(%s, %s)"),
+											*GraphName,
+											*FunctionNames,
+											(bLocalizeDependentPrivateFunctions) ? TEXT("True") : TEXT("False")));
 	}
 
 	return LocalizedFunctions;
@@ -4411,7 +4767,6 @@ URigVMCollapseNode* URigVMController::CollapseNodes(const TArray<URigVMNode*>& I
 
 		if (CollapsedPin->IsStruct())
 		{
-			TGuardValue<bool> SuspendNotifications(bSuspendNotifications, true);
 			AddPinsForStruct(CollapsedPin->GetScriptStruct(), CollapseNode, CollapsedPin, CollapsedPin->GetDirection(), FString(), false);
 		}
 
@@ -4433,10 +4788,7 @@ URigVMCollapseNode* URigVMController::CollapseNodes(const TArray<URigVMNode*>& I
 		EntryNode = NewObject<URigVMFunctionEntryNode>(CollapseNode->ContainedGraph, TEXT("Entry"));
 		CollapseNode->ContainedGraph->Nodes.Add(EntryNode);
 		EntryNode->Position = -Diagonal * 0.5f - FVector2D(250.f, 0.f);
-		{
-			TGuardValue<bool> SuspendNotifications(bSuspendNotifications, true);
-			RefreshFunctionPins(EntryNode);
-		}
+		RefreshFunctionPins(EntryNode, false);
 
 		Notify(ERigVMGraphNotifType::NodeAdded, EntryNode);
 
@@ -4445,10 +4797,7 @@ URigVMCollapseNode* URigVMController::CollapseNodes(const TArray<URigVMNode*>& I
 			ReturnNode = NewObject<URigVMFunctionReturnNode>(CollapseNode->ContainedGraph, TEXT("Return"));
 			CollapseNode->ContainedGraph->Nodes.Add(ReturnNode);
 			ReturnNode->Position = FVector2D(Diagonal.X, -Diagonal.Y) * 0.5f + FVector2D(300.f, 0.f);
-			{
-				TGuardValue<bool> SuspendNotifications(bSuspendNotifications, true);
-				RefreshFunctionPins(ReturnNode);
-			}
+			RefreshFunctionPins(ReturnNode, false);
 
 			Notify(ERigVMGraphNotifType::NodeAdded, ReturnNode);
 		}
@@ -4738,22 +5087,8 @@ TArray<URigVMNode*> URigVMController::ExpandLibraryNode(URigVMLibraryNode* InNod
 		return TArray<URigVMNode*>();
 	}
 
-	URigVMGraph* InnerGraph = InNode->GetContainedGraph();
-	if (URigVMFunctionReferenceNode* RefNode = Cast<URigVMFunctionReferenceNode>(InNode))
-	{
-		if (URigVMLibraryNode* LibraryNode = RefNode->LoadReferencedNode())
-		{
-			InnerGraph = LibraryNode->GetContainedGraph();
-		}
-		else
-		{
-			ReportError(TEXT("Cannot expand nodes from function reference because the source graph is not found."));
-			return TArray<URigVMNode*>();			
-		}
-	}	
-
-	TArray<URigVMNode*> ContainedNodes = InnerGraph->GetNodes();
-	TArray<URigVMLink*> ContainedLinks = InnerGraph->GetLinks();
+	TArray<URigVMNode*> ContainedNodes = InNode->GetContainedNodes();
+	TArray<URigVMLink*> ContainedLinks = InNode->GetContainedLinks();
 	if (ContainedNodes.Num() == 0)
 	{
 		return TArray<URigVMNode*>();
@@ -4814,10 +5149,10 @@ TArray<URigVMNode*> URigVMController::ExpandLibraryNode(URigVMLibraryNode* InNod
 	// exist, they will be reused. If a local variable is not used, it will not be created.
 	if (URigVMFunctionReferenceNode* FunctionReferenceNode = Cast<URigVMFunctionReferenceNode>(InNode))
 	{
-		TArray<FRigVMGraphVariableDescription> LocalVariables = InnerGraph->LocalVariables;
+		TArray<FRigVMGraphVariableDescription> LocalVariables = FunctionReferenceNode->GetContainedGraph()->LocalVariables;
 		TArray<FRigVMExternalVariable> CurrentVariables = GetAllVariables();
 		TArray<FRigVMGraphVariableDescription> VariablesToAdd;
-		for (const URigVMNode* Node : InnerGraph->GetNodes())
+		for (const URigVMNode* Node : FunctionReferenceNode->GetContainedGraph()->GetNodes())
 		{
 			if (const URigVMVariableNode* VariableNode = Cast<URigVMVariableNode>(Node))
 			{
@@ -4854,7 +5189,7 @@ TArray<URigVMNode*> URigVMController::ExpandLibraryNode(URigVMLibraryNode* InNod
 						}
 						else if(bVariableIncompatible)
 						{
-							ReportErrorf(TEXT("Found variable %s of incompatible type with a local variable inside function %s"), *LocalVariable.Name.ToString(), *FunctionReferenceNode->GetReferencedFunctionHeader().Name.ToString());
+							ReportErrorf(TEXT("Found variable %s of incompatible type with a local variable inside function %s"), *LocalVariable.Name.ToString(), *FunctionReferenceNode->GetReferencedNode()->GetName());
 							if (bSetupUndoRedo)
 							{
 								ActionStack->CancelAction(ExpandAction, this);
@@ -4881,7 +5216,7 @@ TArray<URigVMNode*> URigVMController::ExpandLibraryNode(URigVMLibraryNode* InNod
 
 	FString TextContent;
 	{
-		FRigVMControllerGraphGuard GraphGuard(this, InnerGraph, false);
+		FRigVMControllerGraphGuard GraphGuard(this, InNode->GetContainedGraph(), false);
 		TextContent = ExportNodesToText(NodeNames);
 	}
 
@@ -4945,11 +5280,11 @@ TArray<URigVMNode*> URigVMController::ExpandLibraryNode(URigVMLibraryNode* InNod
 
 	// c) create a map from the entry node to the contained graph
 	TMap<FString, TArray<FString>> FromEntryNode;
-	if (URigVMFunctionEntryNode* EntryNode = InnerGraph->GetEntryNode())
+	if (URigVMFunctionEntryNode* EntryNode = InNode->GetEntryNode())
 	{
 		TArray<URigVMLink*> EntryLinks = EntryNode->GetLinks();
 
-		for (URigVMNode* Node : InnerGraph->GetNodes())
+		for (URigVMNode* Node : InNode->GetContainedGraph()->GetNodes())
 		{
 			if (URigVMVariableNode* VariableNode = Cast<URigVMVariableNode>(Node))
 			{
@@ -4998,7 +5333,7 @@ TArray<URigVMNode*> URigVMController::ExpandLibraryNode(URigVMLibraryNode* InNod
 
 	// d) create a map from the contained graph from to the return node
 	TMap<FString, TArray<FString>> ToReturnNode;
-	if (URigVMFunctionReturnNode* ReturnNode = InnerGraph->GetReturnNode())
+	if (URigVMFunctionReturnNode* ReturnNode = InNode->GetReturnNode())
 	{
 		TArray<URigVMLink*> ReturnLinks = ReturnNode->GetLinks();
 		for (URigVMLink* Link : ReturnLinks)
@@ -5288,7 +5623,7 @@ TArray<URigVMNode*> URigVMController::ExpandLibraryNode(URigVMLibraryNode* InNod
 				FString ReturnNodeName, ReturnPinPath;
 				if (URigVMPin::SplitPinPathAtStart(RemappedTargetPinPath, ReturnNodeName, ReturnPinPath))
 				{
-					if(Cast<URigVMFunctionReturnNode>(InnerGraph->FindNode(ReturnNodeName)))
+					if(Cast<URigVMFunctionReturnNode>(InNode->GetContainedGraph()->FindNode(ReturnNodeName)))
 					{
 						if(FromLibraryNode.Contains(ReturnPinPath))
 						{
@@ -5657,7 +5992,7 @@ URigVMCollapseNode* URigVMController::PromoteFunctionReferenceNodeToCollapseNode
 	URigVMGraph* Graph = GetGraph();
 	check(Graph);
 
-	const URigVMCollapseNode* FunctionDefinition = Cast<URigVMCollapseNode>(InFunctionRefNode->LoadReferencedNode());
+	URigVMCollapseNode* FunctionDefinition = Cast<URigVMCollapseNode>(InFunctionRefNode->GetReferencedNode());
 	if (FunctionDefinition == nullptr)
 	{
 		return nullptr;
@@ -5722,11 +6057,6 @@ URigVMCollapseNode* URigVMController::PromoteFunctionReferenceNodeToCollapseNode
 		LinkPaths.Add(TPair< FString, FString >(Link->GetSourcePin()->GetPinPath(), Link->GetTargetPin()->GetPinPath()));
 	}
 
-	if (bSetupUndoRedo)
-	{
-		ActionStack->AddAction(FRigVMPromoteNodeAction(InFunctionRefNode, NodeName, FunctionDefinition->GetPathName()));
-	}
-
 	RemoveNode(InFunctionRefNode, false, true);
 
 	if (RequestNewExternalVariableDelegate.IsBound())
@@ -5744,8 +6074,7 @@ URigVMCollapseNode* URigVMController::PromoteFunctionReferenceNodeToCollapseNode
 		{
 			FRigVMByteArray ArchiveBytes;
 			FMemoryWriter ArchiveWriter(ArchiveBytes);
-			FRigVMTemplate Template = FunctionDefinition->Template;
-			Template.Serialize(ArchiveWriter);
+			FunctionDefinition->Template.Serialize(ArchiveWriter);
 
 			FMemoryReader ArchiveReader(ArchiveBytes);
 			CollapseNode->Template.Serialize(ArchiveReader);
@@ -5755,17 +6084,13 @@ URigVMCollapseNode* URigVMController::PromoteFunctionReferenceNodeToCollapseNode
 	
 		{
 			FRigVMControllerGraphGuard Guard(this, CollapseNode->GetContainedGraph(), false);
-			TGuardValue<bool> SuspendNotifications(bSuspendNotifications, true);
 			ReattachLinksToPinObjects();
 
 			for (URigVMNode* Node : CollapseNode->GetContainedGraph()->GetNodes())
 			{
 				if (URigVMVariableNode* VariableNode = Cast<URigVMVariableNode>(Node))
 				{
-					TArray<URigVMLink*> VariableLinks = VariableNode->GetLinks();
-					DetachLinksFromPinObjects(&VariableLinks);
-					RepopulatePinsOnNode(VariableNode);
-					ReattachLinksToPinObjects(false, &VariableLinks);
+					RepopulatePinsOnNode(VariableNode, true, false, false, true);
 				}
 			}
 
@@ -5784,6 +6109,11 @@ URigVMCollapseNode* URigVMController::PromoteFunctionReferenceNodeToCollapseNode
 		}
 	}
 
+	if (bSetupUndoRedo)
+	{
+		ActionStack->AddAction(FRigVMPromoteNodeAction(InFunctionRefNode, NodeName, FunctionDefinition->GetPathName()));
+	}
+
 	if(bRemoveFunctionDefinition)
 	{
 		FRigVMControllerGraphGuard Guard(this, FunctionDefinition->GetRootGraph(), false);
@@ -5805,24 +6135,23 @@ void URigVMController::SetReferencedFunction(URigVMFunctionReferenceNode* InFunc
 		return;
 	}
 	
-	FRigVMGraphFunctionHeader OldReferencedNode = InFunctionRefNode->GetReferencedFunctionHeader();
-	InFunctionRefNode->ReferencedFunctionHeader = InNewReferencedNode->GetFunctionHeader();
-	
-	if(!(OldReferencedNode == InFunctionRefNode->ReferencedFunctionHeader))
+	URigVMLibraryNode* OldReferencedNode = InFunctionRefNode->GetReferencedNode();
+	if(OldReferencedNode != InNewReferencedNode)
 	{
-		if(URigVMBuildData* BuildData = URigVMBuildData::Get())
+		if(URigVMBuildData* BuildData = GetBuildData())
 		{
-			BuildData->UnregisterFunctionReference(OldReferencedNode.LibraryPointer, InFunctionRefNode);
-			BuildData->RegisterFunctionReference(InFunctionRefNode->ReferencedFunctionHeader.LibraryPointer, InFunctionRefNode);
+			BuildData->UnregisterFunctionReference(OldReferencedNode, InFunctionRefNode);
+			BuildData->RegisterFunctionReference(InNewReferencedNode, InFunctionRefNode);
 		}
 	}
 
+	InFunctionRefNode->SetReferencedNode(InNewReferencedNode);
 	
 	FRigVMControllerGraphGuard GraphGuard(this, InFunctionRefNode->GetGraph(), false);
 	GetGraph()->Notify(ERigVMGraphNotifType::NodeReferenceChanged, InFunctionRefNode);
 }
 
-void URigVMController::RefreshFunctionPins(URigVMNode* InNode)
+void URigVMController::RefreshFunctionPins(URigVMNode* InNode, bool bNotify)
 {
 	if (InNode == nullptr)
 	{
@@ -5834,10 +6163,7 @@ void URigVMController::RefreshFunctionPins(URigVMNode* InNode)
 
 	if (EntryNode || ReturnNode)
 	{
-		TArray<URigVMLink*> Links = InNode->GetLinks();
-		DetachLinksFromPinObjects(&Links);
-		RepopulatePinsOnNode(InNode, false);
-		ReattachLinksToPinObjects(false, &Links);
+		RepopulatePinsOnNode(InNode, false, bNotify, false, true);
 	}
 }
 
@@ -6007,18 +6333,17 @@ bool URigVMController::RemoveNode(URigVMNode* InNode, bool bSetupUndoRedo, bool
 		// If we are removing a reference, remove the function references to this node in the function library
 		if(URigVMFunctionReferenceNode* FunctionReferenceNode = Cast<URigVMFunctionReferenceNode>(LibraryNode))
 		{
-			if(URigVMBuildData* BuildData = URigVMBuildData::Get())
+			if(URigVMBuildData* BuildData = GetBuildData())
 			{
-				BuildData->UnregisterFunctionReference(FunctionReferenceNode->GetReferencedFunctionHeader().LibraryPointer, FunctionReferenceNode);
+				BuildData->UnregisterFunctionReference(FunctionReferenceNode->GetReferencedNode(), FunctionReferenceNode);
 			}
 		}
 		// If we are removing a function, remove all the references first
 		else if (URigVMFunctionLibrary* FunctionLibrary = Cast<URigVMFunctionLibrary>(LibraryNode->GetGraph()))
 		{
-			if(URigVMBuildData* BuildData = URigVMBuildData::Get())
+			if(URigVMBuildData* BuildData = GetBuildData())
 			{
-				FRigVMGraphFunctionIdentifier Identifier = LibraryNode->GetFunctionIdentifier();
-				if (const FRigVMFunctionReferenceArray* ReferencesEntry = BuildData->FindFunctionReferences(Identifier))
+				if (const FRigVMFunctionReferenceArray* ReferencesEntry = BuildData->FindFunctionReferences(LibraryNode))
 				{
 					// make a copy since we'll be modifying the array
 					TArray< TSoftObjectPtr<URigVMFunctionReferenceNode> > FunctionReferences = ReferencesEntry->FunctionReferences;
@@ -6037,7 +6362,7 @@ bool URigVMController::RemoveNode(URigVMNode* InNode, bool bSetupUndoRedo, bool
 					}
 				}
 
-				BuildData->GraphFunctionReferences.Remove(Identifier);
+				BuildData->FunctionReferences.Remove(LibraryNode);
 			}
 			
 			for(const auto& Pair : FunctionLibrary->LocalizedFunctions)
@@ -6048,16 +6373,6 @@ bool URigVMController::RemoveNode(URigVMNode* InNode, bool bSetupUndoRedo, bool
 					break;
 				}
 			}
-
-			if (FunctionLibrary->PublicFunctionNames.Contains(LibraryNode->GetFName()))
-			{
-				FunctionLibrary->PublicFunctionNames.Remove(LibraryNode->GetFName());
-
-				if (bSetupUndoRedo)
-				{
-					ActionStack->AddAction(FRigVMMarkFunctionPublicAction(LibraryNode->GetFName(), true));
-				}
-			}
 		}
 	}
 
@@ -6248,21 +6563,30 @@ bool URigVMController::RenameNode(URigVMNode* InNode, const FName& InNewName, bo
 	{
 		if (URigVMFunctionLibrary* FunctionLibrary = Cast<URigVMFunctionLibrary>(LibraryNode->GetGraph()))
 		{
-			// rename each function reference
-			if(URigVMBuildData* BuildData = URigVMBuildData::Get())
+			// update the table in the build data
+			if(URigVMBuildData* BuildData = GetBuildData())
 			{
-				BuildData->ForEachFunctionReference(LibraryNode->GetFunctionIdentifier(), [this, InNewName](URigVMFunctionReferenceNode* ReferenceNode)
+				for(const TPair< TSoftObjectPtr<URigVMLibraryNode>, FRigVMFunctionReferenceArray >& Pair: BuildData->FunctionReferences)
 				{
-					FRigVMControllerGraphGuard GraphGuard(this, ReferenceNode->GetGraph(), false);
-					RenameNode(ReferenceNode, InNewName, false);
-				});
+					if(Pair.Key.ToSoftObjectPath() == PreviousObjectPath)
+					{
+						const TSoftObjectPtr<URigVMLibraryNode> SoftObjectPtr(InNode);
+						const FRigVMFunctionReferenceArray FunctionReferences = Pair.Value;
+						
+						BuildData->Modify();
+						BuildData->FunctionReferences.Remove(Pair.Key);
+						BuildData->FunctionReferences.Add(SoftObjectPtr, FunctionReferences);
+						BuildData->MarkPackageDirty();
+						break;
+					}
+				}
 			}
-
-			if (FunctionLibrary->PublicFunctionNames.Contains(InNode->PreviousName))
+			
+			FunctionLibrary->ForEachReference(LibraryNode->GetFName(), [this, InNewName](URigVMFunctionReferenceNode* ReferenceNode)
 			{
-				FunctionLibrary->PublicFunctionNames.Remove(InNode->PreviousName);
-				FunctionLibrary->PublicFunctionNames.Add(ValidNewName);
-			}
+				FRigVMControllerGraphGuard GraphGuard(this, ReferenceNode->GetGraph(), false);
+                RenameNode(ReferenceNode, InNewName, false);
+			});
 		}
 	}
 
@@ -6370,7 +6694,7 @@ bool URigVMController::SetNodeSelection(const TArray<FName>& InNodeNames, bool b
 		ActionStack->BeginAction(Action);
 	}
 
-	bool bSelectionChanged = false;
+	bool bSelectedSomething = false;
 
 	TArray<FName> PreviousSelection = Graph->GetSelectNodes();
 	for (const FName& PreviouslySelectedNode : PreviousSelection)
@@ -6379,7 +6703,8 @@ bool URigVMController::SetNodeSelection(const TArray<FName>& InNodeNames, bool b
 		{
 			if(Graph->SelectedNodes.Remove(PreviouslySelectedNode) > 0)
 			{
-				bSelectionChanged = true;
+				Notify(ERigVMGraphNotifType::NodeDeselected, Graph->FindNodeByName(PreviouslySelectedNode));
+				bSelectedSomething = true;
 			}
 		}
 	}
@@ -6392,14 +6717,15 @@ bool URigVMController::SetNodeSelection(const TArray<FName>& InNodeNames, bool b
 			Graph->SelectedNodes.AddUnique(InNodeName);
 			if (PreviousNum != Graph->SelectedNodes.Num())
 			{
-				bSelectionChanged = true;
+				Notify(ERigVMGraphNotifType::NodeSelected, NodeToSelect);
+				bSelectedSomething = true;
 			}
 		}
 	}
 
 	if (bSetupUndoRedo)
 	{
-		if (bSelectionChanged)
+		if (bSelectedSomething)
 		{
 			const TArray<FName>& SelectedNodes = Graph->GetSelectNodes();
 			if (SelectedNodes.Num() == 0)
@@ -6425,7 +6751,7 @@ bool URigVMController::SetNodeSelection(const TArray<FName>& InNodeNames, bool b
 		}
 	}
 
-	if (bSelectionChanged)
+	if (bSelectedSomething)
 	{
 		Notify(ERigVMGraphNotifType::NodeSelectionChanged, nullptr);
 	}
@@ -6451,7 +6777,7 @@ bool URigVMController::SetNodeSelection(const TArray<FName>& InNodeNames, bool b
 											*ArrayStr));
 	}
 
-	return bSelectionChanged;
+	return bSelectedSomething;
 }
 
 bool URigVMController::SetNodePosition(URigVMNode* InNode, const FVector2D& InPosition, bool bSetupUndoRedo, bool bMergeUndoAction, bool bPrintPythonCommand)
@@ -7358,7 +7684,7 @@ bool URigVMController::SetPinDefaultValue(const FString& InPinPath, const FStrin
 	return true;
 }
 
-bool URigVMController::SetPinDefaultValue(URigVMPin* InPin, const FString& InDefaultValue, bool bResizeArrays, bool bSetupUndoRedo, bool bMergeUndoAction)
+bool URigVMController::SetPinDefaultValue(URigVMPin* InPin, const FString& InDefaultValue, bool bResizeArrays, bool bSetupUndoRedo, bool bMergeUndoAction, bool bNotify)
 {
 	if (!bIsTransacting && !IsGraphEditable())
 	{
@@ -7375,6 +7701,8 @@ bool URigVMController::SetPinDefaultValue(URigVMPin* InPin, const FString& InDef
 		ensure(!InDefaultValue.IsEmpty());
 	}
 
+	TGuardValue<bool> Guard(bSuspendNotifications, !bNotify);
+
 	URigVMGraph* Graph = GetGraph();
 	check(Graph);
  
@@ -7621,6 +7949,7 @@ FString URigVMController::GetPinInitialDefaultValueFromStruct(UScriptStruct* Scr
 	{
 		TSharedPtr<FStructOnScope> StructOnScope = MakeShareable(new FStructOnScope(ScriptStruct));
 		uint8* Memory = (uint8*)StructOnScope->GetStructMemory();
+		ScriptStruct->InitializeDefaultValue(Memory);
 
 		if (InPin->GetScriptStruct() == ScriptStruct)
 		{
@@ -7663,6 +7992,7 @@ FString URigVMController::GetPinInitialDefaultValueFromStruct(UScriptStruct* Scr
 					UScriptStruct* InnerStruct = StructProperty->Struct;
 					StructOnScope = MakeShareable(new FStructOnScope(InnerStruct));
 					Memory = (uint8 *)StructOnScope->GetStructMemory();
+					InnerStruct->InitializeDefaultValue(Memory);
 				}
 				continue;
 			}
@@ -8401,10 +8731,7 @@ URigVMPin* URigVMController::InsertArrayPin(URigVMPin* ArrayPin, int32 InIndex,
 		{
 			FString DefaultValue = InDefaultValue;
 			CreateDefaultValueForStructIfRequired(ScriptStruct, DefaultValue);
-			{
-				TGuardValue<bool> SuspendNotifications(bSuspendNotifications, true);
-				AddPinsForStruct(ScriptStruct, Pin->GetNode(), Pin, Pin->Direction, DefaultValue, false);
-			}
+			AddPinsForStruct(ScriptStruct, Pin->GetNode(), Pin, Pin->Direction, DefaultValue, false);
 		}
 	}
 	else if (Pin->IsArray())
@@ -8479,17 +8806,8 @@ bool URigVMController::RemoveArrayPin(const FString& InArrayElementPinPath, bool
 		ActionStack->BeginAction(Action);
 	}
 
-	// we need to keep at least one element for fixed size arrays
-	if(ArrayPin->IsExecuteContext() || ArrayPin->IsFixedSizeArray())
-	{
-		if(ArrayPin->GetArraySize() == 1)
-		{
-			return false;
-		}
-	}
-
 	int32 IndexToRemove = ArrayElementPin->GetPinIndex();
-	if (!RemovePin(ArrayElementPin, bSetupUndoRedo))
+	if (!RemovePin(ArrayElementPin, bSetupUndoRedo, true))
 	{
 		return false;
 	}
@@ -8525,7 +8843,7 @@ bool URigVMController::RemoveArrayPin(const FString& InArrayElementPinPath, bool
 	return true;
 }
 
-bool URigVMController::RemovePin(URigVMPin* InPinToRemove, bool bSetupUndoRedo)
+bool URigVMController::RemovePin(URigVMPin* InPinToRemove, bool bSetupUndoRedo, bool bNotify)
 {
 	if (!bIsTransacting && !IsGraphEditable())
 	{
@@ -8553,13 +8871,13 @@ bool URigVMController::RemovePin(URigVMPin* InPinToRemove, bool bSetupUndoRedo)
 	TArray<URigVMPin*> SubPins = InPinToRemove->GetSubPins();
 	for (URigVMPin* SubPin : SubPins)
 	{
-		if (!RemovePin(SubPin, bSetupUndoRedo))
+		if (!RemovePin(SubPin, bSetupUndoRedo, bNotify))
 		{
 			return false;
 		}
 	}
 
-	if (!bSuspendNotifications)
+	if (bNotify)
 	{
 		Notify(ERigVMGraphNotifType::PinRemoved, InPinToRemove);
 	}
@@ -8652,7 +8970,7 @@ bool URigVMController::SetArrayPinSize(const FString& InArrayPinPath, int32 InSi
 
 	if (!InDefaultValue.IsEmpty() && InDefaultValue != TEXT("()"))
 	{
-		SetPinDefaultValue(Pin, InDefaultValue, false, bSetupUndoRedo, true);
+		SetPinDefaultValue(Pin, InDefaultValue, false, bSetupUndoRedo, true, !bSuspendNotifications);
 	}
 
 	if (bSetupUndoRedo)
@@ -9169,7 +9487,7 @@ bool URigVMController::PromotePinToVariable(URigVMPin* InPin, bool bCreateVariab
 }
 
 bool URigVMController::AddLink(const FString& InOutputPinPath, const FString& InInputPinPath, bool bSetupUndoRedo,
-	bool bPrintPythonCommand, ERigVMPinDirection InUserDirection, bool bCreateCastNode)
+	bool bPrintPythonCommand, ERigVMPinDirection InUserDirection)
 {
 	if(!IsValidGraph())
 	{
@@ -9213,7 +9531,7 @@ bool URigVMController::AddLink(const FString& InOutputPinPath, const FString& In
 	}
 	InputPin = InputPin->GetPinForLink();
 
-	const bool bSuccess = AddLink(OutputPin, InputPin, bSetupUndoRedo, InUserDirection, bCreateCastNode);
+	const bool bSuccess = AddLink(OutputPin, InputPin, bSetupUndoRedo, InUserDirection);
 	if (bSuccess && bPrintPythonCommand)
 	{
 		const FString GraphName = GetSanitizedGraphName(GetGraph()->GetGraphName());
@@ -9231,7 +9549,7 @@ bool URigVMController::AddLink(const FString& InOutputPinPath, const FString& In
 	return bSuccess;
 }
 
-bool URigVMController::AddLink(URigVMPin* OutputPin, URigVMPin* InputPin, bool bSetupUndoRedo, ERigVMPinDirection InUserDirection, bool bCreateCastNode)
+bool URigVMController::AddLink(URigVMPin* OutputPin, URigVMPin* InputPin, bool bSetupUndoRedo, ERigVMPinDirection InUserDirection)
 {
 	if (!bIsTransacting && !IsGraphEditable())
 	{
@@ -9270,7 +9588,7 @@ bool URigVMController::AddLink(URigVMPin* OutputPin, URigVMPin* InputPin, bool b
 	if (!bIsTransacting)
 	{
 		FString FailureReason;
-		if (!Graph->CanLink(OutputPin, InputPin, &FailureReason, GetCurrentByteCode(), UserLinkDirection, bCreateCastNode))
+		if (!Graph->CanLink(OutputPin, InputPin, &FailureReason, GetCurrentByteCode(), UserLinkDirection))
 		{
 			if(OutputPin->IsExecuteContext() && InputPin->IsExecuteContext())
 			{
@@ -9292,243 +9610,60 @@ bool URigVMController::AddLink(URigVMPin* OutputPin, URigVMPin* InputPin, bool b
 	FRigVMBaseAction Action;
 	if (bSetupUndoRedo)
 	{
-		Action.Title = TEXT("Add Link");
+		Action.Title = FString::Printf(TEXT("Add Link"));
 		ActionStack->BeginAction(Action);
 	}
 
-	// check if we need to inject cast node
-	if(bCreateCastNode &&
-		bEnableTypeCasting &&
-		OutputPin->GetTypeIndex() != InputPin->GetTypeIndex() &&
-		!OutputPin->IsWildCard() &&
-		!InputPin->IsWildCard())
+	if (OutputPin->IsExecuteContext())
+	{
+		BreakAllLinks(OutputPin, false, bSetupUndoRedo);
+	}
+
+	bool bAddedTemporaryOutputPreferredType = false, bAddedTemporaryInputPreferredType = false;
+	for (int32 i=0; i<2; ++i)
 	{
-		if(!FRigVMRegistry::Get().CanMatchTypes(OutputPin->GetTypeIndex(), InputPin->GetTypeIndex(), true))
+		URigVMPin* EdgePin = (i==0) ? OutputPin : InputPin;
+		bool& bAddedTemporaryPreferredType = (i==0) ? bAddedTemporaryOutputPreferredType : bAddedTemporaryInputPreferredType;
+		if (EdgePin->IsStructMember())
 		{
-			// returns true if this template node can resolve to the type without breaking any links
-			auto CanTemplateResolveToType = [](const TRigVMTypeIndex& InTypeIndex, URigVMPin* InPin) -> bool
+			// We need to make sure the pin is not removed if any links are broken
+			// We temporarily add a preferred type to make sure that doesn't happen
+			if (URigVMTemplateNode* TemplateNode = Cast<URigVMTemplateNode>(EdgePin->GetNode()))
 			{
-				if(const URigVMTemplateNode* TemplateNode = Cast<URigVMTemplateNode>(InPin->GetNode()))
+				if (TemplateNode->IsSingleton())
 				{
-					if(const FRigVMTemplate* Template = TemplateNode->GetTemplate())
-					{
-						if(!InPin->IsRootPin())
-						{
-							return false;
-						}
-						
-						if(const FRigVMTemplateArgument* Argument = Template->FindArgument(InPin->GetFName()))
-						{
-							if(!Argument->SupportsTypeIndex(InTypeIndex))
-							{
-								return false;
-							}
-						}
-						
-						// we only rely on the cast if the number of links on non-singleton arguments is > 0
-						const TArray<URigVMLink*> Links = InPin->GetNode()->GetLinks();
-						int32 NumLinksRemaining = Links.Num();
-						for(const URigVMLink* Link : Links)
-						{
-							const URigVMPin* SourcePin = Link->GetSourcePin();
-							const URigVMPin* TargetPin = Link->GetTargetPin();
-
-							if(SourcePin && TargetPin)
-							{
-								if(SourcePin->IsExecuteContext())
-								{
-									NumLinksRemaining--;
-									continue;
-								}
-								
-								const URigVMPin* PinOnTemplateNode = SourcePin->GetNode() == TemplateNode ? SourcePin : TargetPin;
-								if(PinOnTemplateNode->IsRootPin() && PinOnTemplateNode != InPin)
-								{
-									const FName ArgumentName = PinOnTemplateNode->GetFName();
-									if(const FRigVMTemplateArgument* Argument = Template->FindArgument(ArgumentName))
-									{
-										if(Argument->IsSingleton())
-										{
-											NumLinksRemaining--;
-										}
-									}
-								}
-								else
-								{
-									NumLinksRemaining--;
-								}
-							}
-						}
-						
-						return NumLinksRemaining == 0;
-					}
+					continue;
 				}
 				
-				return false;
-			};
-
-			if(bCreateCastNode)
-			{
-				// todo: this should take into account links which have been detached
-				// during load we don't filter out the links and assume the template node can cast
-				// correctly to each type. we need to look at the links and their map rather then
-				// URigVMPin::GetLinks
-				if(CanTemplateResolveToType(InputPin->GetTypeIndex(), OutputPin) || 
-					CanTemplateResolveToType(OutputPin->GetTypeIndex(), InputPin))
+				bAddedTemporaryPreferredType = true;
+				URigVMPin* RootPin = EdgePin->GetRootPin();
+				for (FRigVMTemplatePreferredType& Preferred : TemplateNode->PreferredPermutationPairs)
 				{
-					bCreateCastNode = false;
+					if (Preferred.GetArgument() == RootPin->GetFName())
+					{
+						bAddedTemporaryPreferredType = false;
+						break;
+					}
 				}
-			}
-			
-			if(bCreateCastNode)
-			{
-				const FRigVMFunction* CastFunction = RigVMTypeUtils::GetCastForTypeIndices(OutputPin->GetTypeIndex(), InputPin->GetTypeIndex());
 
-				if(CastFunction == nullptr)
+				if (bAddedTemporaryPreferredType)
 				{
-					// this is potentially a template node which has more types available.
-					// look through the filtered types and find a matching cast.
+					AddPreferredType(TemplateNode, RootPin->GetFName(), RootPin->GetTypeIndex(), bSetupUndoRedo);					
 				}
+			}
+		}
+	}
 
-				if(CastFunction == nullptr)
-				{
-					if (bSetupUndoRedo)
-					{
-						ActionStack->CancelAction(Action, this);
-					}
-					return false;
-				}
-
-				const FVector2D OutputPinPosition = OutputPin->GetNode()->GetPosition() + FVector2D(150, 40.f) + FVector2D(0, 16) * OutputPin->GetRootPin()->GetPinIndex();
-				const FVector2D InputPinPosition = InputPin->GetNode()->GetPosition() + FVector2D(-75, 40.f) + FVector2D(0, 16) * InputPin->GetRootPin()->GetPinIndex();
-				const FVector2D CastPosition = (OutputPinPosition + InputPinPosition) * 0.5; 
-
-				const URigVMNode* CastNode = nullptr;
-
-				// try to find an existing cast node
-				if(CastNode == nullptr)
-				{
-					for(URigVMLink* ExistingLink : OutputPin->GetLinks())
-					{
-						if(ExistingLink->GetSourcePin() == OutputPin)
-						{
-							if(URigVMUnitNode* ExistingCastNode = Cast<URigVMUnitNode>(ExistingLink->GetTargetPin()->GetNode()))
-							{
-								if(ExistingCastNode->GetScriptStruct() == CastFunction->Struct)
-								{
-									CastNode = ExistingCastNode;
-									break;
-								}
-							}
-						}
-					}
-				}
-
-				if(CastNode == nullptr)
-				{
-					CastNode = AddUnitNode(CastFunction->Struct, CastFunction->GetMethodName(), CastPosition, FString(), bSetupUndoRedo, false);
-				}
-				
-				if(CastNode == nullptr)
-				{
-					if (bSetupUndoRedo)
-					{
-						ActionStack->CancelAction(Action, this);
-					}
-					return false;
-				}
-
-				static const FString CastTemplateValueName = RigVMTypeUtils::GetCastTemplateValueName().ToString();
-				static const FString CastTemplateResultName = RigVMTypeUtils::GetCastTemplateResultName().ToString();
-				URigVMPin* ValuePin = CastNode->FindPin(CastTemplateValueName);
-				URigVMPin* ResultPin = CastNode->FindPin(CastTemplateResultName);
-
-				if(!OutputPin->IsLinkedTo(ValuePin))
-				{
-					if(!AddLink(OutputPin, ValuePin, bSetupUndoRedo))
-					{
-						if (bSetupUndoRedo)
-						{
-							ActionStack->CancelAction(Action, this);
-						}
-						return false;
-					}
-				}
-
-				if(!ResultPin->IsLinkedTo(InputPin))
-				{
-					if(!AddLink(ResultPin, InputPin, bSetupUndoRedo))
-					{
-						if (bSetupUndoRedo)
-						{
-							ActionStack->CancelAction(Action, this);
-						}
-						return false;
-					}
-				}
-
-				if(bSetupUndoRedo)
-				{
-					Action.Title = TEXT("Add Link with Cast");
-					ActionStack->EndAction(Action);
-				}
-				return true;
-			}
-		}
-	}
-
-	if (OutputPin->IsExecuteContext())
-	{
-		BreakAllLinks(OutputPin, false, bSetupUndoRedo);
-	}
-
-	bool bAddedTemporaryOutputPreferredType = false, bAddedTemporaryInputPreferredType = false;
-	for (int32 i=0; i<2; ++i)
-	{
-		URigVMPin* EdgePin = (i==0) ? OutputPin : InputPin;
-		bool& bAddedTemporaryPreferredType = (i==0) ? bAddedTemporaryOutputPreferredType : bAddedTemporaryInputPreferredType;
-		if (EdgePin->IsStructMember())
-		{
-			// We need to make sure the pin is not removed if any links are broken
-			// We temporarily add a preferred type to make sure that doesn't happen
-			if (URigVMTemplateNode* TemplateNode = Cast<URigVMTemplateNode>(EdgePin->GetNode()))
-			{
-				if (TemplateNode->IsSingleton())
-				{
-					continue;
-				}
-				
-				bAddedTemporaryPreferredType = true;
-				URigVMPin* RootPin = EdgePin->GetRootPin();
-				for (FRigVMTemplatePreferredType& Preferred : TemplateNode->PreferredPermutationPairs)
-				{
-					if (Preferred.GetArgument() == RootPin->GetFName())
-					{
-						bAddedTemporaryPreferredType = false;
-						break;
-					}
-				}
-
-				if (bAddedTemporaryPreferredType)
-				{
-					AddPreferredType(TemplateNode, RootPin->GetFName(), RootPin->GetTypeIndex(), bSetupUndoRedo);					
-				}
-			}
-		}
-	}
-
-	BreakAllLinks(InputPin, true, bSetupUndoRedo);
-	if (bSetupUndoRedo)
-	{
-		BreakAllLinksRecursive(InputPin, true, true, bSetupUndoRedo);
-		BreakAllLinksRecursive(InputPin, true, false, bSetupUndoRedo);
-	}
+	BreakAllLinks(InputPin, true, bSetupUndoRedo);
+	if (bSetupUndoRedo)
+	{
+		BreakAllLinksRecursive(InputPin, true, true, bSetupUndoRedo);
+		BreakAllLinksRecursive(InputPin, true, false, bSetupUndoRedo);
+	}
 
 	// resolve types on the pins if needed
-	if((InputPin->GetCPPTypeObject() != OutputPin->GetCPPTypeObject() ||
-		OutputPin->GetCPPType() != InputPin->GetCPPType()) &&
-		!InputPin->IsExecuteContext() &&
-		!OutputPin->IsExecuteContext())
+	if(InputPin->GetCPPTypeObject() != OutputPin->GetCPPTypeObject() ||
+		OutputPin->GetCPPType() != InputPin->GetCPPType())
 	{
 		bool bOutputPinCanChangeType = OutputPin->IsWildCard();
 		bool bInputPinCanChangeType = InputPin->IsWildCard();
@@ -9544,11 +9679,11 @@ bool URigVMController::AddLink(URigVMPin* OutputPin, URigVMPin* InputPin, bool b
 			Notify(ERigVMGraphNotifType::InteractionBracketOpened, nullptr);
 			if(OutputPin->GetNode()->IsA<URigVMRerouteNode>())
 			{
-				SetPinDefaultValue(OutputPin, InputPin->GetDefaultValue(), true, bSetupUndoRedo, false);
+				SetPinDefaultValue(OutputPin, InputPin->GetDefaultValue(), true, bSetupUndoRedo, false, true);
 			}
 			if(InputPin->GetNode()->IsA<URigVMRerouteNode>())
 			{
-				SetPinDefaultValue(OutputPin, OutputPin->GetDefaultValue(), true, bSetupUndoRedo, false);
+				SetPinDefaultValue(OutputPin, OutputPin->GetDefaultValue(), true, bSetupUndoRedo, false, true);
 			}
 			Notify(ERigVMGraphNotifType::InteractionBracketClosed, nullptr);
 		}
@@ -9563,7 +9698,7 @@ bool URigVMController::AddLink(URigVMPin* OutputPin, URigVMPin* InputPin, bool b
 	// Before adding the link, let's resolve input and ouput pin types
 	// If templates, we will filter the permutations that support this link
 	// If any links need to be broken before perfoming this connection, try to find them and break them
-	if (!bIsTransacting && !InputPin->IsExecuteContext() && !OutputPin->IsExecuteContext())
+	if (!bIsTransacting)
 	{
 		URigVMPin* FirstToResolve = (InUserDirection == ERigVMPinDirection::Input) ? OutputPin : InputPin;
 		URigVMPin* SecondToResolve = (FirstToResolve == OutputPin) ? InputPin : OutputPin;
@@ -9612,34 +9747,28 @@ bool URigVMController::AddLink(URigVMPin* OutputPin, URigVMPin* InputPin, bool b
 		}
 	}
 
-	// Find any library pin that doesn't own an argument, but should
-	if (!bSuspendTemplateComputation)
+	// Find any library pin that doesn't own an argument, but should 
+	if (URigVMLibraryNode* LibraryNode = Graph->GetTypedOuter<URigVMLibraryNode>())
 	{
-		if (URigVMLibraryNode* LibraryNode = Graph->GetTypedOuter<URigVMLibraryNode>())
+		for (URigVMPin* Pin : LibraryNode->GetPins())
 		{
-			for (URigVMPin* Pin : LibraryNode->GetPins())
+			if (!LibraryNode->GetTemplate()->FindArgument(Pin->GetFName()))
 			{
-				if (const FRigVMTemplate* Template = LibraryNode->GetTemplate())
+				if (ShouldPinOwnArgument(Pin))
 				{
-					if (!Template->FindArgument(Pin->GetFName()))
+					URigVMPin* InterfacePin = nullptr;
+					if (URigVMFunctionEntryNode* EntryNode = LibraryNode->GetEntryNode())
+					{
+						InterfacePin = EntryNode->FindPin(Pin->GetName());
+					}
+					if (!InterfacePin)
 					{
-						if (ShouldPinOwnArgument(Pin))
+						if (URigVMFunctionReturnNode* ReturnNode = LibraryNode->GetReturnNode())
 						{
-							URigVMPin* InterfacePin = nullptr;
-							if (URigVMFunctionEntryNode* EntryNode = LibraryNode->GetEntryNode())
-							{
-								InterfacePin = EntryNode->FindPin(Pin->GetName());
-							}
-							if (!InterfacePin)
-							{
-								if (URigVMFunctionReturnNode* ReturnNode = LibraryNode->GetReturnNode())
-								{
-									InterfacePin = ReturnNode->FindPin(Pin->GetName());
-								}
-							}
-							AddArgumentForPin(InterfacePin, nullptr, bSetupUndoRedo);
+							InterfacePin = ReturnNode->FindPin(Pin->GetName());
 						}
 					}
+					AddArgumentForPin(InterfacePin, nullptr, bSetupUndoRedo);
 				}
 			}
 		}
@@ -9654,53 +9783,49 @@ bool URigVMController::AddLink(URigVMPin* OutputPin, URigVMPin* InputPin, bool b
 	if (bSetupUndoRedo)
 	{
 #if WITH_EDITOR
-		if (!bSuspendTemplateComputation)
+		auto ResolveTemplateNodeToCommonTypes = [this](URigVMPin* Pin)
 		{
-			auto ResolveTemplateNodeToCommonTypes = [this](URigVMPin* Pin)
+			if(!Pin->IsExecuteContext())
 			{
-				if(!Pin->IsExecuteContext())
-				{
-					return;
-				}
+				return;
+			}
 			
-				URigVMTemplateNode* TemplateNode = Cast<URigVMTemplateNode>(Pin->GetNode());
-				if(TemplateNode == nullptr)
-				{
-					return;
-				}
+			URigVMTemplateNode* TemplateNode = Cast<URigVMTemplateNode>(Pin->GetNode());
+			if(TemplateNode == nullptr)
+			{
+				return;
+			}
 
-				const FRigVMTemplate* Template = TemplateNode->GetTemplate();
-				if(Template == nullptr)
-				{
-					return;
-				}
+			const FRigVMTemplate* Template = TemplateNode->GetTemplate();
+			if(Template == nullptr)
+			{
+				return;
+			}
 
-				if(!TemplateNode->HasWildCardPin())
-				{
-					return;
-				}
+			if(!TemplateNode->HasWildCardPin())
+			{
+				return;
+			}
 
-				const FRigVMTemplate::FTypeMap PreferredTypes = GetCommonlyUsedTypesForTemplate(TemplateNode);
-				if(PreferredTypes.IsEmpty())
-				{
-					return;
-				}
+			const FRigVMTemplate::FTypeMap PreferredTypes = GetCommonlyUsedTypesForTemplate(TemplateNode);
+			if(PreferredTypes.IsEmpty())
+			{
+				return;
+			}
 
-				const int32 PreferredPermutation = Template->FindPermutation(PreferredTypes);
-				if(PreferredPermutation != INDEX_NONE)
+			const int32 PreferredPermutation = Template->FindPermutation(PreferredTypes);
+			if(PreferredPermutation != INDEX_NONE)
+			{
+				const TGuardValue<bool> DisableRegisterUseOfTemplate(bRegisterTemplateNodeUsage, false);
+				if(FullyResolveTemplateNode(TemplateNode, PreferredPermutation, true))
 				{
-					const TGuardValue<bool> DisableRegisterUseOfTemplate(bRegisterTemplateNodeUsage, false);
-					if(FullyResolveTemplateNode(TemplateNode, PreferredPermutation, true))
-					{
-						static const FString Message = TEXT("Template node was automatically resolved to commonly used types.");
-						SendUserFacingNotification(Message, 0.f, TemplateNode, TEXT("MessageLog.Note"));
-					}
+					static const FString Message = TEXT("Template node was automatically resolved to commonly used types.");
+					SendUserFacingNotification(Message, 0.f, TemplateNode, TEXT("MessageLog.Note"));
 				}
-			};
-
-			ResolveTemplateNodeToCommonTypes(OutputPin);
-			ResolveTemplateNodeToCommonTypes(InputPin);
-		}
+			}
+		};
+		ResolveTemplateNodeToCommonTypes(OutputPin);
+		ResolveTemplateNodeToCommonTypes(InputPin);
 #endif
 		
 		ActionStack->EndAction(Action);
@@ -9869,8 +9994,8 @@ bool URigVMController::BreakLink(URigVMPin* OutputPin, URigVMPin* InputPin, bool
 
 			// each time a link is removed, existing orphaned pins
 			// may become unused and thus can be removed
-			RemoveUnusedOrphanedPins(OutputPin->GetNode());
-			RemoveUnusedOrphanedPins(InputPin->GetNode());
+			RemoveUnusedOrphanedPins(OutputPin->GetNode(), true);
+			RemoveUnusedOrphanedPins(InputPin->GetNode(), true);
 			
 			if (!bIsTransacting && !bSuspendRecomputingTemplateFilters)
 			{
@@ -10080,7 +10205,7 @@ FName URigVMController::AddExposedPin(const FName& InPinName, ERigVMPinDirection
 		}
 		if (CPPTypeObject == nullptr)
 		{
-			CPPTypeObject = RigVMTypeUtils::FindObjectFromCPPTypeObjectPath<UObject>(InCPPTypeObjectPath.ToString());
+			CPPTypeObject = URigVMPin::FindObjectFromCPPTypeObjectPath<UObject>(InCPPTypeObjectPath.ToString());
 		}
 	}
 
@@ -10157,10 +10282,7 @@ FName URigVMController::AddExposedPin(const FName& InPinName, ERigVMPinDirection
 
 		FString DefaultValue = InDefaultValue;
 		CreateDefaultValueForStructIfRequired(Pin->GetScriptStruct(), DefaultValue);
-		{
-			TGuardValue<bool> SuspendNotifications(bSuspendNotifications, true);
-			AddPinsForStruct(Pin->GetScriptStruct(), LibraryNode, Pin, Pin->Direction, DefaultValue, false);
-		}
+		AddPinsForStruct(Pin->GetScriptStruct(), LibraryNode, Pin, Pin->Direction, DefaultValue, false);
 	}
 
 	FRigVMControllerCompileBracketScope CompileScope(this);
@@ -10186,10 +10308,7 @@ FName URigVMController::AddExposedPin(const FName& InPinName, ERigVMPinDirection
 	{
 		EntryNode = NewObject<URigVMFunctionEntryNode>(Graph, TEXT("Entry"));
 		Graph->Nodes.Add(EntryNode);
-		{
-			TGuardValue<bool> SuspendNotifications(bSuspendNotifications, true);
-			RefreshFunctionPins(EntryNode);
-		}
+		RefreshFunctionPins(EntryNode, false);
 		EntryNode->InitializeFilteredPermutations();
 
 		Notify(ERigVMGraphNotifType::NodeAdded, EntryNode);
@@ -10200,17 +10319,14 @@ FName URigVMController::AddExposedPin(const FName& InPinName, ERigVMPinDirection
 	{
 		ReturnNode = NewObject<URigVMFunctionReturnNode>(Graph, TEXT("Return"));
 		Graph->Nodes.Add(ReturnNode);
-		{
-			TGuardValue<bool> SuspendNotifications(bSuspendNotifications, true);
-			RefreshFunctionPins(ReturnNode);
-		}
+		RefreshFunctionPins(ReturnNode, false);
 		ReturnNode->InitializeFilteredPermutations();
 
 		Notify(ERigVMGraphNotifType::NodeAdded, ReturnNode);
 	}
 
-	RefreshFunctionPins(EntryNode);
-	RefreshFunctionPins(ReturnNode);
+	RefreshFunctionPins(EntryNode, true);
+	RefreshFunctionPins(ReturnNode, true);
 
 	// If the type is not a wildcard, add pin as template argument
 	if (InCPPType != RigVMTypeUtils::GetWildCardCPPType() && InCPPType != RigVMTypeUtils::GetWildCardArrayCPPType())
@@ -10451,7 +10567,6 @@ bool URigVMController::AddArgumentForPin(URigVMPin* InPin, URigVMPin* InToLinkPi
 		// Temporarily disconnect connected links
 		if (CurTemplate.NumArguments() > 0)
 		{
-			TGuardValue<bool> SuspendNotifications(bSuspendNotifications, true);
 			DetachLinksFromPinObjects(&ConnectedLinks);
 		}
 	
@@ -10493,7 +10608,6 @@ bool URigVMController::AddArgumentForPin(URigVMPin* InPin, URigVMPin* InToLinkPi
 		// Reconnect  links
 		if (CurTemplate.NumArguments() > 0)
 		{
-			TGuardValue<bool> SuspendNotifications(bSuspendNotifications, true);
 			ReattachLinksToPinObjects(false, &ConnectedLinks);
 		}
 		
@@ -10696,7 +10810,7 @@ bool URigVMController::RemoveExposedPin(const FName& InPinName, bool bSetupUndoR
 		TGuardValue<bool> GuardRecomputeOuterFilters(bSuspendRecomputingOuterTemplateFilters, true);
 		{
 			FRigVMControllerGraphGuard GraphGuard(this, LibraryNode->GetGraph(), bSetupUndoRedo);
-			bSuccessfullyRemovedPin = RemovePin(Pin, bSetupUndoRedo);
+			bSuccessfullyRemovedPin = RemovePin(Pin, bSetupUndoRedo, true);
 		}
 
 		TArray<URigVMVariableNode*> NodesToRemove;
@@ -10715,8 +10829,8 @@ bool URigVMController::RemoveExposedPin(const FName& InPinName, bool bSetupUndoR
 			RemoveNode(NodesToRemove[i], bSetupUndoRedo);
 		}
 
-		RefreshFunctionPins(Graph->GetEntryNode());
-		RefreshFunctionPins(Graph->GetReturnNode());
+		RefreshFunctionPins(Graph->GetEntryNode(), true);
+		RefreshFunctionPins(Graph->GetReturnNode(), true);
 		RefreshFunctionReferences(LibraryNode, false);
 	}
 
@@ -11036,7 +11150,7 @@ bool URigVMController::ChangeExposedPinType(const FName& InPinName, const FStrin
 	UObject* CPPTypeObject = nullptr;
 	if (!InCPPTypeObjectPath.IsNone())
 	{
-		CPPTypeObject = RigVMTypeUtils::FindObjectFromCPPTypeObjectPath<UObject>(InCPPTypeObjectPath.ToString());
+		CPPTypeObject = URigVMPin::FindObjectFromCPPTypeObjectPath<UObject>(InCPPTypeObjectPath.ToString());
 		if(CPPTypeObject)
 		{
 			if(const UScriptStruct* CPPTypeStruct = Cast<UScriptStruct>(CPPTypeObject))
@@ -11191,7 +11305,7 @@ bool URigVMController::ChangeExposedPinType(const FName& InPinName, const FStrin
 
 			if (bSuccessChangingType)
 			{
-				RemoveUnusedOrphanedPins(LibraryNode);
+				RemoveUnusedOrphanedPins(LibraryNode, true);
 			}
 		}
 		if (!bSuccessChangingType)
@@ -11208,11 +11322,8 @@ bool URigVMController::ChangeExposedPinType(const FName& InPinName, const FStrin
 	{
 		if (InterfaceNode)
 		{
-			const TArray<URigVMLink*> Links = InterfaceNode->GetLinks();
-			DetachLinksFromPinObjects(&Links);
-			RepopulatePinsOnNode(InterfaceNode, true, bSetupOrphanPins);
-			ReattachLinksToPinObjects(true, &Links, bSetupOrphanPins);
-			RemoveUnusedOrphanedPins(InterfaceNode);
+			RepopulatePinsOnNode(InterfaceNode, true, true, bSetupOrphanPins, true);
+			RemoveUnusedOrphanedPins(InterfaceNode, true);
 		}
 	}
 
@@ -11224,7 +11335,7 @@ bool URigVMController::ChangeExposedPinType(const FName& InPinName, const FStrin
 			{
 				FRigVMControllerGraphGuard GraphGuard(this, ReferenceNode->GetGraph(), bSetupUndoRedo);
 				ChangePinType(ReferencedNodePin, InCPPType, InCPPTypeObjectPath, bSetupUndoRedo, bSetupOrphanPins);
-				RemoveUnusedOrphanedPins(ReferenceNode);
+				RemoveUnusedOrphanedPins(ReferenceNode, true);
 			}
         });
 	}
@@ -11239,7 +11350,7 @@ bool URigVMController::ChangeExposedPinType(const FName& InPinName, const FStrin
 				if (ValuePin)
 				{
 					ChangePinType(ValuePin, InCPPType, InCPPTypeObjectPath, bSetupUndoRedo, bSetupOrphanPins);
-					RemoveUnusedOrphanedPins(VariableNode);
+					RemoveUnusedOrphanedPins(VariableNode, true);
 				}
 			}
 		}
@@ -11324,8 +11435,8 @@ bool URigVMController::SetExposedPinIndex(const FName& InPinName, int32 InNewInd
 		Notify(ERigVMGraphNotifType::PinIndexChanged, Pin);
 	}
 
-	RefreshFunctionPins(LibraryNode->GetEntryNode());
-	RefreshFunctionPins(LibraryNode->GetReturnNode());
+	RefreshFunctionPins(LibraryNode->GetEntryNode(), true);
+	RefreshFunctionPins(LibraryNode->GetReturnNode(), true);
 	RefreshFunctionReferences(LibraryNode, false);
 	
 	if (bSetupUndoRedo)
@@ -11353,73 +11464,62 @@ URigVMFunctionReferenceNode* URigVMController::AddFunctionReferenceNode(URigVMLi
 	{
 		return nullptr;
 	}
-	
-	if (!InFunctionDefinition)
-	{
-		return nullptr;
-	}
 
-	if (URigVMFunctionReferenceNode* ReferenceNode = AddFunctionReferenceNodeFromDescription(InFunctionDefinition->GetFunctionHeader(), InNodePosition, InNodeName, bSetupUndoRedo, bPrintPythonCommand))
+	if (!bIsTransacting && !IsGraphEditable())
 	{
-		return ReferenceNode;
+		return nullptr;
 	}
 
-	return nullptr;
-}
+	URigVMGraph* Graph = GetGraph();
+	check(Graph);
 
-URigVMFunctionReferenceNode* URigVMController::AddFunctionReferenceNodeFromDescription(const FRigVMGraphFunctionHeader& InFunctionDefinition, const FVector2D& InNodePosition, const FString& InNodeName, bool bSetupUndoRedo, bool bPrintPythonCommand, bool bAllowPrivateFunctions)
-{
-	if (!IsValidGraph())
+	if (Graph->IsA<URigVMFunctionLibrary>())
 	{
+		ReportError(TEXT("Cannot add function reference nodes to function library graphs."));
 		return nullptr;
 	}
 
-	if (!bIsTransacting && !IsGraphEditable())
+	if (InFunctionDefinition == nullptr)
 	{
+		ReportError(TEXT("Cannot add a function reference node without a valid function definition."));
 		return nullptr;
 	}
 
-	URigVMGraph* Graph = GetGraph();
-	check(Graph);
-
-	if (Graph->IsA<URigVMFunctionLibrary>())
+	if (!InFunctionDefinition->GetGraph()->IsA<URigVMFunctionLibrary>())
 	{
-		ReportError(TEXT("Cannot add function reference nodes to function library graphs."));
+		ReportAndNotifyError(TEXT("Cannot use the function definition for a function reference node."));
 		return nullptr;
 	}
 
-	if(!CanAddFunctionRefForDefinition(InFunctionDefinition, true, bAllowPrivateFunctions))
+	if(!CanAddFunctionRefForDefinition(InFunctionDefinition, true))
 	{
 		return nullptr;
 	}
 
-	FString NodeName = GetValidNodeName(InNodeName.IsEmpty() ? InFunctionDefinition.Name.ToString() : InNodeName);
+	FString NodeName = GetValidNodeName(InNodeName.IsEmpty() ? InFunctionDefinition->GetName() : InNodeName);
 	URigVMFunctionReferenceNode* FunctionRefNode = NewObject<URigVMFunctionReferenceNode>(Graph, *NodeName);
 	FunctionRefNode->Position = InNodePosition;
-	FunctionRefNode->ReferencedFunctionHeader = InFunctionDefinition;
+	FunctionRefNode->SetReferencedNode(InFunctionDefinition);
 	Graph->Nodes.Add(FunctionRefNode);
 	FunctionRefNode->InitializeFilteredPermutations();
 
 	FRigVMControllerCompileBracketScope CompileScope(this);
 	
-	{
-		TGuardValue<bool> SuspendNotifications(bSuspendNotifications, true);
-		RepopulatePinsOnNode(FunctionRefNode, false);
-	}
-	//UpdateTemplateNodePinTypes(FunctionRefNode, false); // We will need to add this if we ever allow functions to be templates
+	RepopulatePinsOnNode(FunctionRefNode, false, false, false);
+	UpdateTemplateNodePinTypes(FunctionRefNode, false);
 
 	Notify(ERigVMGraphNotifType::NodeAdded, FunctionRefNode);
 
-	if (URigVMBuildData* BuildData = URigVMBuildData::Get())
+	if (URigVMBuildData* BuildData = GetBuildData())
 	{
-		BuildData->RegisterFunctionReference(FunctionRefNode->GetReferencedFunctionHeader().LibraryPointer, FunctionRefNode);
+		BuildData->RegisterFunctionReference(InFunctionDefinition, FunctionRefNode);
 	}
-	
-	for (const FRigVMGraphFunctionArgument& Argument : InFunctionDefinition.Arguments)
+
+	for (URigVMPin* SourcePin : InFunctionDefinition->Pins)
 	{
-		if (URigVMPin* TargetPin = FunctionRefNode->FindPin(Argument.Name.ToString()))
+		if (URigVMPin* TargetPin = FunctionRefNode->FindPin(SourcePin->GetName()))
 		{
-			const FString& DefaultValue = Argument.DefaultValue;
+			FString DefaultValue = SourcePin->GetDefaultValue();
 			if (!DefaultValue.IsEmpty())
 			{
 				SetPinDefaultValue(TargetPin, DefaultValue, true, false, false);
@@ -11440,83 +11540,36 @@ URigVMFunctionReferenceNode* URigVMController::AddFunctionReferenceNodeFromDescr
 	if (bPrintPythonCommand)
 	{
 		const FString GraphName = GetSanitizedGraphName(GetGraph()->GetGraphName());
-		const FString FunctionDefinitionName = GetSanitizedNodeName(InFunctionDefinition.Name.ToString());
+		const FString FunctionDefinitionName = GetSanitizedNodeName(InFunctionDefinition->GetName());
 
-		bool bLocal = false;
-		if(IRigVMClientHost* ClientHost = GetImplementingOuter<IRigVMClientHost>())
+		if (InFunctionDefinition->GetLibrary() == GetGraph()->GetDefaultFunctionLibrary())
 		{
-			if (InFunctionDefinition.LibraryPointer.HostObject == Cast<UObject>(ClientHost->GetRigVMGraphFunctionHost()))
-			{
-				bLocal = true;
-				RigVMPythonUtils::Print(GetGraphOuterName(), 
-					FString::Printf(TEXT("blueprint.get_controller_by_name('%s').add_function_reference_node(library.find_function('%s'), %s, '%s')"),
-							*GraphName,
-							*FunctionDefinitionName,
-							*RigVMPythonUtils::Vector2DToPythonString(InNodePosition),
-							*NodeName));
-			}
+
+			RigVMPythonUtils::Print(GetGraphOuterName(), 
+				FString::Printf(TEXT("blueprint.get_controller_by_name('%s').add_function_reference_node(library.find_function('%s'), %s, '%s')"),
+						*GraphName,
+						*FunctionDefinitionName,
+						*RigVMPythonUtils::Vector2DToPythonString(InNodePosition),
+						*NodeName));
 		}
-		
-		if (!bLocal)
+		else
 		{
 			RigVMPythonUtils::Print(GetGraphOuterName(), 
-				FString::Printf(TEXT("blueprint.get_controller_by_name('%s').add_external_function_reference_node('%s', '%s', %s, '%s')"),
+				FString::Printf(TEXT("function_blueprint = unreal.load_object(name = '%s', outer = None)"),
+				*InFunctionDefinition->GetLibrary()->GetOuter()->GetPathName()));
+			RigVMPythonUtils::Print(GetGraphOuterName(), 
+				FString::Printf(TEXT("blueprint.get_controller_by_name('%s').add_function_reference_node(function_blueprint.get_local_function_library().find_function('%s'), %s, '%s')"),
 						*GraphName,
-						*InFunctionDefinition.LibraryPointer.HostObject.ToString(),
 						*FunctionDefinitionName,
-						*RigVMPythonUtils::Vector2DToPythonString(InNodePosition), 
-						*NodeName));
+						*RigVMPythonUtils::Vector2DToPythonString(InFunctionDefinition->GetPosition()), 
+						*FunctionDefinitionName));
 		}
+		
 	}
 
 	return FunctionRefNode;
 }
 
-URigVMFunctionReferenceNode* URigVMController::AddExternalFunctionReferenceNode(const FString& InHostPath, const FName& InFunctionName, const FVector2D& InNodePosition, const FString& InNodeName, bool bSetupUndoRedo, bool bPrintPythonCommand)
-{
-	if (!IsValidGraph())
-	{
-		return nullptr;
-	}
-
-	if (!bIsTransacting && !IsGraphEditable())
-	{
-		return nullptr;
-	}
-
-	URigVMGraph* Graph = GetGraph();
-	check(Graph);
-
-	if (Graph->IsA<URigVMFunctionLibrary>())
-	{
-		ReportError(TEXT("Cannot add function reference nodes to function library graphs."));
-		return nullptr;
-	}
-
-	UObject* HostObject = StaticLoadObject(UObject::StaticClass(), NULL, *InHostPath, NULL, LOAD_None, NULL);
-	if (!HostObject)
-	{
-		ReportErrorf(TEXT("Failed to load the Host object %s."), *InHostPath);
-		return nullptr;
-	}
-
-	IRigVMGraphFunctionHost* FunctionHost = Cast<IRigVMGraphFunctionHost>(HostObject);
-	if (!FunctionHost)
-	{
-		ReportError(TEXT("Host object is not a IRigVMGraphFunctionHost."));
-		return nullptr;
-	}
-
-	FRigVMGraphFunctionData* Data = FunctionHost->GetRigVMGraphFunctionStore()->FindFunctionByName(InFunctionName);
-	if (!Data)
-	{
-		ReportErrorf(TEXT("Function %s not found in host %s."), *InFunctionName.ToString(), *InHostPath);
-		return nullptr;
-	}
-
-	return AddFunctionReferenceNodeFromDescription(Data->Header, InNodePosition, InNodeName, bSetupUndoRedo, bPrintPythonCommand);
-}
-
 bool URigVMController::SetRemappedVariable(URigVMFunctionReferenceNode* InFunctionRefNode,
 	const FName& InInnerVariableName, const FName& InOuterVariableName, bool bSetupUndoRedo)
 {
@@ -11556,13 +11609,8 @@ bool URigVMController::SetRemappedVariable(URigVMFunctionReferenceNode* InFuncti
 
 	FRigVMExternalVariable InnerExternalVariable;
 	{
-		if (FRigVMExternalVariable* Variable = InFunctionRefNode->ReferencedFunctionHeader.ExternalVariables.FindByPredicate([InInnerVariableName](const FRigVMExternalVariable& Variable)
-		{
-			return Variable.Name == InInnerVariableName;
-		}))
-		{
-			InnerExternalVariable = *Variable;
-		}
+		FRigVMControllerGraphGuard GraphGuard(this, InFunctionRefNode->GetContainedGraph());
+		InnerExternalVariable = GetVariableByName(InInnerVariableName);
 	}
 
 	if(!InnerExternalVariable.IsValid(true))
@@ -11644,8 +11692,6 @@ URigVMLibraryNode* URigVMController::AddFunctionToLibrary(const FName& InFunctio
 
 	FRigVMControllerCompileBracketScope CompileScope(this);
 	
-	Notify(ERigVMGraphNotifType::NodeAdded, CollapseNode);
-	
 	if (bMutable)
 	{
 		const UScriptStruct* ExecuteContextStruct = FRigVMExecuteContext::StaticStruct();
@@ -11659,6 +11705,8 @@ URigVMLibraryNode* URigVMController::AddFunctionToLibrary(const FName& InFunctio
 			bSetupUndoRedo);
 	}
 
+	Notify(ERigVMGraphNotifType::NodeAdded, CollapseNode);
+
 	{
 		FRigVMControllerGraphGuard GraphGuard(this, CollapseNode->GetContainedGraph(), false);
 		TGuardValue<bool> GuardEditGraph(CollapseNode->ContainedGraph->bEditable, true);
@@ -11670,11 +11718,7 @@ URigVMLibraryNode* URigVMController::AddFunctionToLibrary(const FName& InFunctio
 		{
 			EntryNode = NewObject<URigVMFunctionEntryNode>(CollapseNode->ContainedGraph, TEXT("Entry"));
 			CollapseNode->ContainedGraph->Nodes.Add(EntryNode);
-			{
-				TGuardValue<bool> SuspendNotifications(bSuspendNotifications, true);
-				RefreshFunctionPins(EntryNode);
-			}
-			EntryNode->Position = FVector2D(-250.f, 0.f);
+			RefreshFunctionPins(EntryNode, false);
 			Notify(ERigVMGraphNotifType::NodeAdded, EntryNode);
 		}
 
@@ -11682,14 +11726,12 @@ URigVMLibraryNode* URigVMController::AddFunctionToLibrary(const FName& InFunctio
 		{
 			ReturnNode = NewObject<URigVMFunctionReturnNode>(CollapseNode->ContainedGraph, TEXT("Return"));
 			CollapseNode->ContainedGraph->Nodes.Add(ReturnNode);
-			{
-				TGuardValue<bool> SuspendNotifications(bSuspendNotifications, true);
-				RefreshFunctionPins(ReturnNode);
-			}
-			ReturnNode->Position = FVector2D(250.f, 0.f);
+			RefreshFunctionPins(ReturnNode, false);
 			Notify(ERigVMGraphNotifType::NodeAdded, ReturnNode);
 		}
 
+		EntryNode->Position = FVector2D(-250.f, 0.f);
+		ReturnNode->Position = FVector2D(250.f, 0.f);
 
 		if (bMutable)
 		{
@@ -11779,134 +11821,23 @@ bool URigVMController::RenameFunction(const FName& InOldFunctionName, const FNam
 	return RenameNode(Node, InNewFunctionName, bSetupUndoRedo);
 }
 
-bool URigVMController::MarkFunctionAsPublic(const FName& InFunctionName, bool bInIsPublic, bool bSetupUndoRedo,	bool bPrintPythonCommand)
+FRigVMGraphVariableDescription URigVMController::AddLocalVariable(const FName& InVariableName, const FString& InCPPType, UObject* InCPPTypeObject, const FString& InDefaultValue, bool bSetupUndoRedo, bool bPrintPythonCommand)
 {
+	FRigVMGraphVariableDescription NewVariable;
 	if (!IsValidGraph())
 	{
-		return false;
+		return NewVariable;
 	}
 
 	if (!bIsTransacting && !IsGraphEditable())
 	{
-		return false;
+		return NewVariable;
 	}
-
+	
 	URigVMGraph* Graph = GetGraph();
 	check(Graph);
 
-
-	if (!Graph->IsA<URigVMFunctionLibrary>())
-	{
-		ReportError(TEXT("Can only change function definitions from function library graphs."));
-		return false;
-	}
-
-	URigVMNode* Node = Graph->FindNode(InFunctionName.ToString());
-	if (!Node)
-	{
-		ReportErrorf(TEXT("Could not find function called '%s'."), *InFunctionName.ToString());
-		return false;
-	}
-
-	if (URigVMFunctionLibrary* FunctionLibrary = Cast<URigVMFunctionLibrary>(Graph))
-	{
-		bool bOldIsPublic = FunctionLibrary->PublicFunctionNames.Contains(InFunctionName); 
-		if ((bInIsPublic && bOldIsPublic) || (!bInIsPublic && !bOldIsPublic))
-		{
-			return true;
-		}
-
-		if (bSetupUndoRedo)
-		{
-			FRigVMBaseAction BaseAction;
-			BaseAction.Title = FString::Printf(TEXT("Mark function %s as %s"), *InFunctionName.ToString(), (bInIsPublic) ? TEXT("Public") : TEXT("Private"));
-			ActionStack->BeginAction(BaseAction);
-			ActionStack->AddAction(FRigVMMarkFunctionPublicAction(InFunctionName, bInIsPublic));
-			ActionStack->EndAction(BaseAction);
-		}
-
-		if (bInIsPublic)
-		{
-			FunctionLibrary->PublicFunctionNames.Add(InFunctionName);
-		}
-		else
-		{
-			FunctionLibrary->PublicFunctionNames.Remove(InFunctionName);
-		}
-	}
-
-	Notify(ERigVMGraphNotifType::FunctionAccessChanged, Node);
-
-	if (!bSuspendNotifications)
-	{
-		Graph->MarkPackageDirty();
-	}
-
-	if (bPrintPythonCommand)
-	{
-		RigVMPythonUtils::Print(GetGraphOuterName(), 
-			FString::Printf(TEXT("library_controller.mark_function_as_public('%s', %s)"),
-				*GetSanitizedNodeName(InFunctionName.ToString()),
-				(bInIsPublic) ? TEXT("True") : TEXT("False")));
-	}
-
-	return true;
-}
-
-bool URigVMController::IsFunctionPublic(const FName& InFunctionName)
-{
-	if (!IsValidGraph())
-	{
-		return false;
-	}
-
-	if (!bIsTransacting && !IsGraphEditable())
-	{
-		return false;
-	}
-
-	URigVMGraph* Graph = GetGraph();
-	check(Graph);
-
-
-	if (!Graph->IsA<URigVMFunctionLibrary>())
-	{
-		ReportError(TEXT("Can only check function definitions from function library graphs."));
-		return false;
-	}
-
-	URigVMNode* Node = Graph->FindNode(InFunctionName.ToString());
-	if (!Node)
-	{
-		ReportErrorf(TEXT("Could not find function called '%s'."), *InFunctionName.ToString());
-		return false;
-	}
-
-	if (URigVMFunctionLibrary* FunctionLibrary = Cast<URigVMFunctionLibrary>(Graph))
-	{
-		return FunctionLibrary->PublicFunctionNames.Contains(InFunctionName);
-	}
-
-	return false;
-}
-
-FRigVMGraphVariableDescription URigVMController::AddLocalVariable(const FName& InVariableName, const FString& InCPPType, UObject* InCPPTypeObject, const FString& InDefaultValue, bool bSetupUndoRedo, bool bPrintPythonCommand)
-{
-	FRigVMGraphVariableDescription NewVariable;
-	if (!IsValidGraph())
-	{
-		return NewVariable;
-	}
-
-	if (!bIsTransacting && !IsGraphEditable())
-	{
-		return NewVariable;
-	}
-	
-	URigVMGraph* Graph = GetGraph();
-	check(Graph);
-
-	// Check this is the main graph of a function
+	// Check this is the main graph of a function
 	{
 		if (URigVMLibraryNode* LibraryNode = Cast<URigVMLibraryNode>(Graph->GetOuter()))
 		{
@@ -11999,7 +11930,7 @@ FRigVMGraphVariableDescription URigVMController::AddLocalVariableFromObjectPath(
 	UObject* CPPTypeObject = nullptr;
 	if (!InCPPTypeObjectPath.IsEmpty())
 	{
-		CPPTypeObject = RigVMTypeUtils::FindObjectFromCPPTypeObjectPath<UObject>(InCPPTypeObjectPath);
+		CPPTypeObject = URigVMPin::FindObjectFromCPPTypeObjectPath<UObject>(InCPPTypeObjectPath);
 		if (CPPTypeObject == nullptr)
 		{
 			ReportErrorf(TEXT("Cannot find cpp type object for path '%s'."), *InCPPTypeObjectPath);
@@ -12355,7 +12286,7 @@ bool URigVMController::SetLocalVariableTypeFromObjectPath(const FName& InVariabl
 	UObject* CPPTypeObject = nullptr;
 	if (!InCPPTypeObjectPath.IsEmpty())
 	{
-		CPPTypeObject = RigVMTypeUtils::FindObjectFromCPPTypeObjectPath<UObject>(InCPPTypeObjectPath);
+		CPPTypeObject = URigVMPin::FindObjectFromCPPTypeObjectPath<UObject>(InCPPTypeObjectPath);
 		if (CPPTypeObject == nullptr)
 		{
 			ReportErrorf(TEXT("Cannot find cpp type object for path '%s'."), *InCPPTypeObjectPath);
@@ -12366,7 +12297,7 @@ bool URigVMController::SetLocalVariableTypeFromObjectPath(const FName& InVariabl
 	return SetLocalVariableType(InVariableName, InCPPType, CPPTypeObject, bSetupUndoRedo, bPrintPythonCommand);
 }
 
-bool URigVMController::SetLocalVariableDefaultValue(const FName& InVariableName, const FString& InDefaultValue, bool bSetupUndoRedo, bool bPrintPythonCommand)
+bool URigVMController::SetLocalVariableDefaultValue(const FName& InVariableName, const FString& InDefaultValue, bool bSetupUndoRedo, bool bPrintPythonCommand, bool bNotify)
 {
 	if (!IsValidGraph())
 	{
@@ -12509,7 +12440,7 @@ bool URigVMController::PerformUserWorkflow(const FRigVMUserWorkflow& InWorkflow,
 	return bSuccess;
 }
 
-TArray<TSoftObjectPtr<URigVMFunctionReferenceNode>> URigVMController::GetAffectedReferences(ERigVMControllerBulkEditType InEditType, bool bForceLoad)
+TArray<TSoftObjectPtr<URigVMFunctionReferenceNode>> URigVMController::GetAffectedReferences(ERigVMControllerBulkEditType InEditType, bool bForceLoad, bool bNotify)
 {
 	TArray<TSoftObjectPtr<URigVMFunctionReferenceNode>> FunctionReferencePtrs;
 	
@@ -12545,7 +12476,7 @@ TArray<TSoftObjectPtr<URigVMFunctionReferenceNode>> URigVMController::GetAffecte
 
 		if(bForceLoad)
 		{
-			if(OnBulkEditProgressDelegate.IsBound() && !bSuspendNotifications)
+			if(OnBulkEditProgressDelegate.IsBound() && bNotify)
 			{
 				OnBulkEditProgressDelegate.Execute(FunctionReferencePtr, InEditType, ERigVMControllerBulkEditProgress::BeginLoad, FunctionReferenceIndex, FunctionReferencePtrs.Num());
 			}
@@ -12555,7 +12486,7 @@ TArray<TSoftObjectPtr<URigVMFunctionReferenceNode>> URigVMController::GetAffecte
 				FunctionReferencePtr.LoadSynchronous();
 			}
 
-			if(OnBulkEditProgressDelegate.IsBound() && !bSuspendNotifications)
+			if(OnBulkEditProgressDelegate.IsBound() && bNotify)
 			{
 				OnBulkEditProgressDelegate.Execute(FunctionReferencePtr, InEditType, ERigVMControllerBulkEditProgress::FinishedLoad, FunctionReferenceIndex, FunctionReferencePtrs.Num());
 			}
@@ -12576,19 +12507,21 @@ TArray<TSoftObjectPtr<URigVMFunctionReferenceNode>> URigVMController::GetAffecte
 		{
 			if(URigVMFunctionReferenceNode* AffectedFunctionReferenceNode = FunctionReferencePtr.Get())
 			{
-				if(URigVMLibraryNode* AffectedFunction = AffectedFunctionReferenceNode->FindFunctionForNode())
+				if(URigVMFunctionLibrary* AffectedFunctionLibrary = AffectedFunctionReferenceNode->GetTypedOuter<URigVMFunctionLibrary>())
 				{
-					FRigVMControllerGraphGuard GraphGuard(this, AffectedFunction->GetContainedGraph(), false);
-					TGuardValue<bool> SuspendNotifications(bSuspendNotifications, true);
-					TArray<TSoftObjectPtr<URigVMFunctionReferenceNode>> AffectedFunctionReferencePtrs = GetAffectedReferences(InEditType, bForceLoad);
-					for(TSoftObjectPtr<URigVMFunctionReferenceNode> AffectedFunctionReferencePtr : AffectedFunctionReferencePtrs)
+					if(URigVMLibraryNode* AffectedFunction = AffectedFunctionLibrary->FindFunctionForNode(AffectedFunctionReferenceNode))
 					{
-						const FString Key = AffectedFunctionReferencePtr.ToSoftObjectPath().ToString();
-						if(VisitedPaths.Contains(Key))
+						FRigVMControllerGraphGuard GraphGuard(this, AffectedFunction->GetContainedGraph(), false);
+						TArray<TSoftObjectPtr<URigVMFunctionReferenceNode>> AffectedFunctionReferencePtrs = GetAffectedReferences(InEditType, bForceLoad, false);
+						for(TSoftObjectPtr<URigVMFunctionReferenceNode> AffectedFunctionReferencePtr : AffectedFunctionReferencePtrs)
 						{
-							continue;
+							const FString Key = AffectedFunctionReferencePtr.ToSoftObjectPath().ToString();
+							if(VisitedPaths.Contains(Key))
+							{
+								continue;
+							}
+							VisitedPaths.Add(Key, FunctionReferencePtrs.Add(AffectedFunctionReferencePtr));
 						}
-						VisitedPaths.Add(Key, FunctionReferencePtrs.Add(AffectedFunctionReferencePtr));
 					}
 				}
 			}
@@ -12600,7 +12533,7 @@ TArray<TSoftObjectPtr<URigVMFunctionReferenceNode>> URigVMController::GetAffecte
 	return FunctionReferencePtrs;
 }
 
-TArray<FAssetData> URigVMController::GetAffectedAssets(ERigVMControllerBulkEditType InEditType, bool bForceLoad)
+TArray<FAssetData> URigVMController::GetAffectedAssets(ERigVMControllerBulkEditType InEditType, bool bForceLoad, bool bNotify)
 {
 	TArray<FAssetData> Assets;
 
@@ -12611,7 +12544,7 @@ TArray<FAssetData> URigVMController::GetAffectedAssets(ERigVMControllerBulkEditT
 		return Assets;
 	}
 
-	TArray<TSoftObjectPtr<URigVMFunctionReferenceNode>> FunctionReferencePtrs = GetAffectedReferences(InEditType, bForceLoad);
+	TArray<TSoftObjectPtr<URigVMFunctionReferenceNode>> FunctionReferencePtrs = GetAffectedReferences(InEditType, bForceLoad, bNotify);
 	TMap<FString, int32> VisitedAssets;
 
 	URigVMGraph* Graph = GetGraph();
@@ -12800,10 +12733,7 @@ URigVMRerouteNode* URigVMController::AddFreeRerouteNode(bool bShowAsFullNode, co
 	{
 		FString DefaultValue = InDefaultValue;
 		CreateDefaultValueForStructIfRequired(ValuePin->GetScriptStruct(), DefaultValue);
-		{
-			TGuardValue<bool> SuspendNotifications(bSuspendNotifications, true);
-			AddPinsForStruct(ValuePin->GetScriptStruct(), Node, ValuePin, ValuePin->Direction, DefaultValue, false);
-		}
+		AddPinsForStruct(ValuePin->GetScriptStruct(), Node, ValuePin, ValuePin->Direction, DefaultValue, false);
 	}
 	else if (!InDefaultValue.IsEmpty() && InDefaultValue != TEXT("()"))
 	{
@@ -12841,12 +12771,72 @@ URigVMRerouteNode* URigVMController::AddFreeRerouteNode(bool bShowAsFullNode, co
 	return Node;
 }
 
-URigVMNode* URigVMController::AddBranchNode(const FVector2D& InPosition, const FString& InNodeName, bool bSetupUndoRedo, bool bPrintPythonCommand)
+URigVMBranchNode* URigVMController::AddBranchNode(const FVector2D& InPosition, const FString& InNodeName, bool bSetupUndoRedo, bool bPrintPythonCommand)
 {
-	return AddUnitNode(FRigVMFunction_ControlFlowBranch::StaticStruct(), FRigVMStruct::ExecuteName, InPosition, InNodeName, bSetupUndoRedo, bPrintPythonCommand);
+	if(!IsValidGraph())
+	{
+		return nullptr;
+	}
+
+	if (!bIsTransacting && !IsGraphEditable())
+	{
+		return nullptr;
+	}
+
+	URigVMGraph* Graph = GetGraph();
+	check(Graph);
+
+	FString Name = GetValidNodeName(InNodeName.IsEmpty() ? FString(TEXT("BranchNode")) : InNodeName);
+	URigVMBranchNode* Node = NewObject<URigVMBranchNode>(Graph, *Name);
+	Node->Position = InPosition;
+
+	URigVMPin* ExecutePin = MakeExecutePin(Node, FRigVMStruct::ExecuteContextName);
+	ExecutePin->Direction = ERigVMPinDirection::Input;
+	AddNodePin(Node, ExecutePin);
+
+	URigVMPin* ConditionPin = NewObject<URigVMPin>(Node, *URigVMBranchNode::ConditionName);
+	ConditionPin->CPPType = RigVMTypeUtils::BoolType;
+	ConditionPin->Direction = ERigVMPinDirection::Input;
+	AddNodePin(Node, ConditionPin);
+
+	URigVMPin* TruePin = NewObject<URigVMPin>(Node, *URigVMBranchNode::TrueName);
+	TruePin->CPPType = ExecutePin->CPPType;
+	TruePin->CPPTypeObject = ExecutePin->CPPTypeObject;
+	TruePin->CPPTypeObjectPath = ExecutePin->CPPTypeObjectPath;
+	TruePin->Direction = ERigVMPinDirection::Output;
+	AddNodePin(Node, TruePin);
+
+	URigVMPin* FalsePin = NewObject<URigVMPin>(Node, *URigVMBranchNode::FalseName);
+	FalsePin->CPPType = ExecutePin->CPPType;
+	FalsePin->CPPTypeObject = ExecutePin->CPPTypeObject;
+	FalsePin->CPPTypeObjectPath = ExecutePin->CPPTypeObjectPath;
+	FalsePin->Direction = ERigVMPinDirection::Output;
+	AddNodePin(Node, FalsePin);
+
+	Graph->Nodes.Add(Node);
+
+	Notify(ERigVMGraphNotifType::NodeAdded, Node);
+
+	FRigVMControllerCompileBracketScope CompileScope(this);
+	if (bSetupUndoRedo)
+	{
+		ActionStack->AddAction(FRigVMAddBranchNodeAction(Node));
+	}
+
+	if (bPrintPythonCommand)
+	{
+		TArray<FString> Commands = GetAddNodePythonCommands(Node);
+		for (const FString& Command : Commands)
+		{
+			RigVMPythonUtils::Print(GetGraphOuterName(), 
+				FString::Printf(TEXT("%s"), *Command));
+		}
+	}
+
+	return Node;
 }
 
-URigVMNode* URigVMController::AddIfNode(const FString& InCPPType, const FName& InCPPTypeObjectPath, const FVector2D& InPosition, const FString& InNodeName, bool  bSetupUndoRedo, bool bPrintPythonCommand)
+URigVMIfNode* URigVMController::AddIfNode(const FString& InCPPType, const FName& InCPPTypeObjectPath, const FVector2D& InPosition, const FString& InNodeName, bool  bSetupUndoRedo, bool bPrintPythonCommand)
 {
 	if(!IsValidGraph())
 	{
@@ -12866,7 +12856,7 @@ URigVMNode* URigVMController::AddIfNode(const FString& InCPPType, const FName& I
 	UObject* CPPTypeObject = nullptr;
 	if(!InCPPTypeObjectPath.IsNone())
 	{
-		CPPTypeObject = RigVMTypeUtils::FindObjectFromCPPTypeObjectPath<UObject>(InCPPTypeObjectPath.ToString());
+		CPPTypeObject = URigVMPin::FindObjectFromCPPTypeObjectPath<UObject>(InCPPTypeObjectPath.ToString());
 		if (CPPTypeObject == nullptr)
 		{
 			ReportErrorf(TEXT("Cannot find cpp type object for path '%s'."), *InCPPTypeObjectPath.ToString());
@@ -12874,31 +12864,106 @@ URigVMNode* URigVMController::AddIfNode(const FString& InCPPType, const FName& I
 		}
 	}
 
-	if(bSetupUndoRedo)
+	FString CPPType = RigVMTypeUtils::PostProcessCPPType(InCPPType, CPPTypeObject);
+
+	FString DefaultValue;
+	if(UScriptStruct* ScriptStruct = Cast<UScriptStruct>(CPPTypeObject))
 	{
-		OpenUndoBracket(TEXT("Add If Node"));
+		if (ScriptStruct->IsChildOf(FRigVMExecuteContext::StaticStruct()))
+		{
+			ReportErrorf(TEXT("Cannot create an if node for this type '%s'."), *InCPPTypeObjectPath.ToString());
+			return nullptr;
+		}
+		CreateDefaultValueForStructIfRequired(ScriptStruct, DefaultValue);
 	}
-
-	const FString CPPType = RigVMTypeUtils::PostProcessCPPType(InCPPType, CPPTypeObject);
-	const FString Name = GetValidNodeName(InNodeName.IsEmpty() ? FString(TEXT("IfNode")) : InNodeName);
-	const TRigVMTypeIndex& TypeIndex = FRigVMRegistry::Get().FindOrAddType({*CPPType, CPPTypeObject});
 	
-	const FRigVMDispatchFactory* Factory = FRigVMRegistry::Get().FindOrAddDispatchFactory(FRigVMDispatch_If::StaticStruct());
-	URigVMNode* Node = AddTemplateNode(Factory->GetTemplate()->GetNotation(), InPosition, Name, bSetupUndoRedo, bPrintPythonCommand);
-	if(Node)
+	FString Name = GetValidNodeName(InNodeName.IsEmpty() ? FString(TEXT("IfNode")) : InNodeName);
+	URigVMIfNode* Node = NewObject<URigVMIfNode>(Graph, *Name);
+	Node->Position = InPosition;
+
+	URigVMPin* ConditionPin = NewObject<URigVMPin>(Node, *URigVMIfNode::ConditionName);
+	ConditionPin->CPPType = RigVMTypeUtils::BoolType;
+	ConditionPin->Direction = ERigVMPinDirection::Input;
+	AddNodePin(Node, ConditionPin);
+
+	URigVMPin* TruePin = NewObject<URigVMPin>(Node, *URigVMIfNode::TrueName);
+	TruePin->CPPType = CPPType;
+	TruePin->CPPTypeObject = CPPTypeObject;
+	TruePin->CPPTypeObjectPath = InCPPTypeObjectPath;
+	TruePin->Direction = ERigVMPinDirection::Input;
+	TruePin->DefaultValue = DefaultValue;
+	AddNodePin(Node, TruePin);
+
+	if (TruePin->IsStruct())
+	{
+		AddPinsForStruct(TruePin->GetScriptStruct(), Node, TruePin, TruePin->Direction, FString(), false);
+	}
+
+	URigVMPin* FalsePin = NewObject<URigVMPin>(Node, *URigVMIfNode::FalseName);
+	FalsePin->CPPType = CPPType;
+	FalsePin->CPPTypeObject = CPPTypeObject;
+	FalsePin->CPPTypeObjectPath = InCPPTypeObjectPath;
+	FalsePin->Direction = ERigVMPinDirection::Input;
+	FalsePin->DefaultValue = DefaultValue;
+	AddNodePin(Node, FalsePin);
+
+	if (FalsePin->IsStruct())
+	{
+		AddPinsForStruct(FalsePin->GetScriptStruct(), Node, FalsePin, FalsePin->Direction, FString(), false);
+	}
+
+	URigVMPin* ResultPin = NewObject<URigVMPin>(Node, *URigVMIfNode::ResultName);
+	ResultPin->CPPType = CPPType;
+	ResultPin->CPPTypeObject = CPPTypeObject;
+	ResultPin->CPPTypeObjectPath = InCPPTypeObjectPath;
+	ResultPin->Direction = ERigVMPinDirection::Output;
+	AddNodePin(Node, ResultPin);
+
+	if (ResultPin->IsStruct())
+	{
+		AddPinsForStruct(ResultPin->GetScriptStruct(), Node, ResultPin, ResultPin->Direction, FString(), false);
+	}
+
+	Graph->Nodes.Add(Node);
+
+	Notify(ERigVMGraphNotifType::NodeAdded, Node);
+
+	FRigVMControllerCompileBracketScope CompileScope(this);
+	if (bSetupUndoRedo)
+	{
+		ActionStack->AddAction(FRigVMAddIfNodeAction(Node));
+	}
+
+	Node->InitializeFilteredPermutations();
+	if (InCPPType != RigVMTypeUtils::GetWildCardCPPType() && InCPPType != RigVMTypeUtils::GetWildCardArrayCPPType())
 	{
-		ResolveWildCardPin(Node->GetPins().Last(), TypeIndex, bSetupUndoRedo, bPrintPythonCommand);
+		PrepareTemplatePinForType(TruePin, {TruePin->GetTypeIndex()}, bSetupUndoRedo);
+		TArray<int32> FilterPermutations = Node->GetFilteredPermutationsIndices();
+		if (FilterPermutations.Num() == 1)
+		{
+			const TArray<FRigVMTemplatePreferredType> NewPreferredPermutationTypes = Node->GetPreferredTypesForPermutation(FilterPermutations[0]);
+			if (bSetupUndoRedo)
+			{
+				ActionStack->AddAction(FRigVMSetPreferredTemplatePermutationsAction(Node, NewPreferredPermutationTypes));
+			}
+			Node->PreferredPermutationPairs = NewPreferredPermutationTypes;
+		}
 	}
 	
-	if(bSetupUndoRedo)
+	if (bPrintPythonCommand)
 	{
-		CloseUndoBracket();
+		TArray<FString> Commands = GetAddNodePythonCommands(Node);
+		for (const FString& Command : Commands)
+		{
+			RigVMPythonUtils::Print(GetGraphOuterName(), 
+				FString::Printf(TEXT("%s"), *Command));
+		}
 	}
 
 	return Node;
 }
 
-URigVMNode* URigVMController::AddIfNodeFromStruct(UScriptStruct* InScriptStruct, const FVector2D& InPosition, const FString& InNodeName, bool bSetupUndoRedo)
+URigVMIfNode* URigVMController::AddIfNodeFromStruct(UScriptStruct* InScriptStruct, const FVector2D& InPosition, const FString& InNodeName, bool bSetupUndoRedo)
 {
 	if (!InScriptStruct)
 	{
@@ -12908,7 +12973,7 @@ URigVMNode* URigVMController::AddIfNodeFromStruct(UScriptStruct* InScriptStruct,
 	return AddIfNode(RigVMTypeUtils::GetUniqueStructTypeName(InScriptStruct), FName(InScriptStruct->GetPathName()), InPosition, InNodeName, bSetupUndoRedo);
 }
 
-URigVMNode* URigVMController::AddSelectNode(const FString& InCPPType, const FName& InCPPTypeObjectPath, const FVector2D& InPosition, const FString& InNodeName, bool bSetupUndoRedo, bool bPrintPythonCommand)
+URigVMSelectNode* URigVMController::AddSelectNode(const FString& InCPPType, const FName& InCPPTypeObjectPath, const FVector2D& InPosition, const FString& InNodeName, bool bSetupUndoRedo, bool bPrintPythonCommand)
 {
 	if (!IsValidGraph())
 	{
@@ -12928,7 +12993,7 @@ URigVMNode* URigVMController::AddSelectNode(const FString& InCPPType, const FNam
 	UObject* CPPTypeObject = nullptr;
 	if (!InCPPTypeObjectPath.IsNone())
 	{
-		CPPTypeObject = RigVMTypeUtils::FindObjectFromCPPTypeObjectPath<UObject>(InCPPTypeObjectPath.ToString());
+		CPPTypeObject = URigVMPin::FindObjectFromCPPTypeObjectPath<UObject>(InCPPTypeObjectPath.ToString());
 		if (CPPTypeObject == nullptr)
 		{
 			ReportErrorf(TEXT("Cannot find cpp type object for path '%s'."), *InCPPTypeObjectPath.ToString());
@@ -12936,30 +13001,90 @@ URigVMNode* URigVMController::AddSelectNode(const FString& InCPPType, const FNam
 		}
 	}
 
-	if(bSetupUndoRedo)
+	FString CPPType = RigVMTypeUtils::PostProcessCPPType(InCPPType, CPPTypeObject);
+
+	FString DefaultValue;
+	if (UScriptStruct* ScriptStruct = Cast<UScriptStruct>(CPPTypeObject))
 	{
-		OpenUndoBracket(TEXT("Add Select Node"));
+		if (ScriptStruct->IsChildOf(FRigVMExecuteContext::StaticStruct()))
+		{
+			ReportErrorf(TEXT("Cannot create a select node for this type '%s'."), *InCPPTypeObjectPath.ToString());
+			return nullptr;
+		}
+		CreateDefaultValueForStructIfRequired(ScriptStruct, DefaultValue);
 	}
 
-	const FString CPPType = RigVMTypeUtils::PostProcessCPPType(InCPPType, CPPTypeObject);
-	const FString Name = GetValidNodeName(InNodeName.IsEmpty() ? FString(TEXT("SelectNode")) : InNodeName);
-	const TRigVMTypeIndex& TypeIndex = FRigVMRegistry::Get().FindOrAddType({*CPPType, CPPTypeObject});
-	
-	const FRigVMDispatchFactory* Factory = FRigVMRegistry::Get().FindOrAddDispatchFactory(FRigVMDispatch_SelectInt32::StaticStruct());
-	URigVMNode* Node = AddTemplateNode(Factory->GetTemplate()->GetNotation(), InPosition, Name, bSetupUndoRedo, bPrintPythonCommand);
-	if(Node)
+	FString Name = GetValidNodeName(InNodeName.IsEmpty() ? FString(TEXT("IfNode")) : InNodeName);
+	URigVMSelectNode* Node = NewObject<URigVMSelectNode>(Graph, *Name);
+	Node->Position = InPosition;
+
+	URigVMPin* IndexPin = NewObject<URigVMPin>(Node, *URigVMSelectNode::IndexName);
+	IndexPin->CPPType = RigVMTypeUtils::Int32Type;
+	IndexPin->Direction = ERigVMPinDirection::Input;
+	AddNodePin(Node, IndexPin);
+
+	URigVMPin* ValuePin = NewObject<URigVMPin>(Node, *URigVMSelectNode::ValueName);
+	ValuePin->CPPType = RigVMTypeUtils::ArrayTypeFromBaseType(CPPType);
+	ValuePin->CPPTypeObject = CPPTypeObject;
+	ValuePin->CPPTypeObjectPath = InCPPTypeObjectPath;
+	ValuePin->Direction = ERigVMPinDirection::Input;
+	ValuePin->bIsExpanded = true;
+	AddNodePin(Node, ValuePin);
+
+	URigVMPin* ResultPin = NewObject<URigVMPin>(Node, *URigVMSelectNode::ResultName);
+	ResultPin->CPPType = CPPType;
+	ResultPin->CPPTypeObject = CPPTypeObject;
+	ResultPin->CPPTypeObjectPath = InCPPTypeObjectPath;
+	ResultPin->Direction = ERigVMPinDirection::Output;
+	AddNodePin(Node, ResultPin);
+
+	if (ResultPin->IsStruct())
 	{
-		ResolveWildCardPin(Node->GetPins().Last(), TypeIndex, bSetupUndoRedo, bPrintPythonCommand);
+		AddPinsForStruct(ResultPin->GetScriptStruct(), Node, ResultPin, ResultPin->Direction, FString(), false);
 	}
 
-	if(bSetupUndoRedo)
+	Graph->Nodes.Add(Node);
+
+	Notify(ERigVMGraphNotifType::NodeAdded, Node);
+
+	SetArrayPinSize(ValuePin->GetPinPath(), 2, DefaultValue, false);
+
+	FRigVMControllerCompileBracketScope CompileScope(this);
+	if (bSetupUndoRedo)
 	{
-		CloseUndoBracket();
+		ActionStack->AddAction(FRigVMAddSelectNodeAction(Node));
+	}
+
+	Node->InitializeFilteredPermutations();
+	if (InCPPType != RigVMTypeUtils::GetWildCardCPPType() && InCPPType != RigVMTypeUtils::GetWildCardArrayCPPType())
+	{
+		PrepareTemplatePinForType(ResultPin, {ResultPin->GetTypeIndex()}, bSetupUndoRedo);
+		TArray<int32> FilterPermutations = Node->GetFilteredPermutationsIndices();
+		if (FilterPermutations.Num() == 1)
+		{
+			const TArray<FRigVMTemplatePreferredType> NewPreferredPermutationTypes = Node->GetPreferredTypesForPermutation(FilterPermutations[0]);
+			if (bSetupUndoRedo)
+			{
+				ActionStack->AddAction(FRigVMSetPreferredTemplatePermutationsAction(Node, NewPreferredPermutationTypes));
+			}
+			Node->PreferredPermutationPairs = NewPreferredPermutationTypes;
+		}
+	}
+	
+	if (bPrintPythonCommand)
+	{
+		TArray<FString> Commands = GetAddNodePythonCommands(Node);
+		for (const FString& Command : Commands)
+		{
+			RigVMPythonUtils::Print(GetGraphOuterName(), 
+				FString::Printf(TEXT("%s"), *Command));
+		}
 	}
+
 	return Node;
 }
 
-URigVMNode* URigVMController::AddSelectNodeFromStruct(UScriptStruct* InScriptStruct, const FVector2D& InPosition,
+URigVMSelectNode* URigVMController::AddSelectNodeFromStruct(UScriptStruct* InScriptStruct, const FVector2D& InPosition,
 	const FString& InNodeName, bool bSetupUndoRedo)
 {
 	if (!InScriptStruct)
@@ -12994,21 +13119,6 @@ URigVMTemplateNode* URigVMController::AddTemplateNode(const FName& InNotation, c
 		return nullptr;
 	}
 
-	if(IRigVMClientHost* ClientHost = GetImplementingOuter<IRigVMClientHost>())
-	{
-		if(const FRigVMClient* Client = ClientHost->GetRigVMClient())
-		{
-			if(!Template->SupportsExecuteContextStruct(Client->GetExecuteContextStruct()))
-			{
-				ReportErrorf(TEXT("Cannot add node for template '%s' - incompatible execute context: '%s' vs '%s'."),
-						*Template->GetNotation().ToString(),
-						*Template->GetExecuteContextStruct()->GetStructCPPName(),
-						*Client->GetExecuteContextStruct()->GetStructCPPName());
-				return nullptr;
-			}
-		}
-	}
-
 	FString Name = GetValidNodeName(InNodeName.IsEmpty() ? Template->GetName().ToString() : InNodeName);
 	URigVMTemplateNode* Node = nullptr;
 
@@ -13050,23 +13160,46 @@ URigVMTemplateNode* URigVMController::AddTemplateNode(const FName& InNotation, c
 	Node->InitializeFilteredPermutations();
 
 	FRigVMRegistry& Registry = FRigVMRegistry::Get();
-	AddPinsForTemplate(Template, Types, Node);
-
-	if (Node->HasWildCardPin())
-	{
-		UpdateTemplateNodePinTypes(Node, false);
-	}
-	else
+	
+	for (int32 ArgIndex = 0; ArgIndex < Template->NumArguments(); ArgIndex++)
 	{
-		if (!Node->IsA<URigVMFunctionEntryNode>() && !Node->IsA<URigVMFunctionReturnNode>())
+		const FRigVMTemplateArgument* Arg = Template->GetArgument(ArgIndex);
+
+		URigVMPin* Pin = NewObject<URigVMPin>(Node, Arg->GetName());
+		const TRigVMTypeIndex& TypeIndex = Types.FindChecked(Arg->GetName());
+		const FRigVMTemplateArgumentType Type = Registry.GetType(TypeIndex);
+		
+		Pin->CPPType = Type.CPPType.ToString();
+		Pin->CPPTypeObject = Type.CPPTypeObject;
+		if (Pin->CPPTypeObject)
 		{
-			FullyResolveTemplateNode(Node, INDEX_NONE, false);
+			Pin->CPPTypeObjectPath = *Pin->CPPTypeObject->GetPathName();
 		}
-	}
-
-	Graph->Nodes.Add(Node);
+		Pin->Direction = Arg->GetDirection();
+		Pin->LastKnownTypeIndex = TypeIndex;
+		Pin->LastKnownCPPType = Pin->CPPType;
 
-	Notify(ERigVMGraphNotifType::NodeAdded, Node);
+		AddNodePin(Node, Pin);
+
+		if(!Pin->IsWildCard())
+		{
+			const FString DefaultValue = Node->GetInitialDefaultValueForPin(Pin->GetFName());
+			if(UScriptStruct* ScriptStruct = Cast<UScriptStruct>(Pin->CPPTypeObject))
+			{
+				AddPinsForStruct(ScriptStruct, Pin->GetNode(), Pin, Pin->Direction, DefaultValue, false, false);
+			}
+			else if(!DefaultValue.IsEmpty())
+			{
+				SetPinDefaultValue(Pin, DefaultValue, true, false, false, false);
+			}
+		}
+	}
+
+	UpdateTemplateNodePinTypes(Node, false);
+
+	Graph->Nodes.Add(Node);
+
+	Notify(ERigVMGraphNotifType::NodeAdded, Node);
 
 	FRigVMAddTemplateNodeAction Action;
 	FRigVMControllerCompileBracketScope CompileScope(this);
@@ -13181,7 +13314,7 @@ URigVMEnumNode* URigVMController::AddEnumNode(const FName& InCPPTypeObjectPath,
 	URigVMGraph* Graph = GetGraph();
 	check(Graph);
 
-    UObject* CPPTypeObject = RigVMTypeUtils::FindObjectFromCPPTypeObjectPath<UObject>(InCPPTypeObjectPath.ToString());
+    UObject* CPPTypeObject = URigVMPin::FindObjectFromCPPTypeObjectPath<UObject>(InCPPTypeObjectPath.ToString());
 	if (CPPTypeObject == nullptr)
 	{
 		ReportErrorf(TEXT("Cannot find cpp type object for path '%s'."), *InCPPTypeObjectPath.ToString());
@@ -13236,9 +13369,9 @@ URigVMEnumNode* URigVMController::AddEnumNode(const FName& InCPPTypeObjectPath,
 	return Node;
 }
 
-URigVMNode* URigVMController::AddArrayNode(ERigVMOpCode InOpCode, const FString& InCPPType,
+URigVMArrayNode* URigVMController::AddArrayNode(ERigVMOpCode InOpCode, const FString& InCPPType,
 	UObject* InCPPTypeObject, const FVector2D& InPosition, const FString& InNodeName, bool bSetupUndoRedo,
-	bool bPrintPythonCommand, bool bIsPatching)
+	bool bPrintPythonCommand)
 {
 	if (!IsValidGraph())
 	{
@@ -13250,109 +13383,310 @@ URigVMNode* URigVMController::AddArrayNode(ERigVMOpCode InOpCode, const FString&
 		return nullptr;
 	}
 
-	const FString CPPType = RigVMTypeUtils::IsArrayType(InCPPType) ? RigVMTypeUtils::BaseTypeFromArrayType(InCPPType) : InCPPType;
-	const TRigVMTypeIndex& ElementTypeIndex = FRigVMRegistry::Get().FindOrAddType({*CPPType, InCPPTypeObject});
-	if(ElementTypeIndex == INDEX_NONE)
-	{
-		return nullptr;
+	// validate the op code
+	bool bIsMutable = false;
+	switch(InOpCode)
+	{
+		case ERigVMOpCode::ArrayReset:
+		case ERigVMOpCode::ArrayGetNum: 
+		case ERigVMOpCode::ArraySetNum:
+		case ERigVMOpCode::ArrayGetAtIndex:  
+		case ERigVMOpCode::ArraySetAtIndex:
+		case ERigVMOpCode::ArrayAdd:
+		case ERigVMOpCode::ArrayInsert:
+		case ERigVMOpCode::ArrayRemove:
+		case ERigVMOpCode::ArrayFind:
+		case ERigVMOpCode::ArrayAppend:
+		case ERigVMOpCode::ArrayClone:
+		case ERigVMOpCode::ArrayIterator:
+		case ERigVMOpCode::ArrayUnion:
+		case ERigVMOpCode::ArrayDifference:
+		case ERigVMOpCode::ArrayIntersection:
+		case ERigVMOpCode::ArrayReverse:
+		{
+			break;
+		}
+		default:
+		{
+			ReportErrorf(TEXT("OpCode '%s' is not valid for Array Node."), *StaticEnum<ERigVMOpCode>()->GetNameStringByValue((int64)InOpCode));
+			return nullptr;
+		}
 	}
-	const TRigVMTypeIndex& ArrayTypeIndex = FRigVMRegistry::Get().GetArrayTypeFromBaseTypeIndex(ElementTypeIndex);
 
-	const FName FactoryName = FRigVMDispatch_ArrayBase::GetFactoryNameForOpCode(InOpCode);
-	if(FactoryName.IsNone())
+	URigVMGraph* Graph = GetGraph();
+	check(Graph);
+
+	if (Graph->IsA<URigVMFunctionLibrary>())
 	{
-		ReportErrorf(TEXT("OpCode '%s' is not valid for Array Node."), *StaticEnum<ERigVMOpCode>()->GetNameStringByValue((int64)InOpCode));
+		ReportError(TEXT("Cannot add array nodes to function library graphs."));
 		return nullptr;
 	}
 
-	const FRigVMDispatchFactory* Factory = FRigVMRegistry::Get().FindDispatchFactory(FactoryName);
-	if(Factory == nullptr)
+	FString CPPType = InCPPType;
+	if(RigVMTypeUtils::IsArrayType(CPPType))
 	{
-		ReportErrorf(TEXT("Cannot find array dispatch '%s'."), *FactoryName.ToString());
-		return nullptr;
+		CPPType = RigVMTypeUtils::BaseTypeFromArrayType(CPPType);
 	}
 
-	if(bSetupUndoRedo)
+	if (InCPPTypeObject == nullptr)
 	{
-		OpenUndoBracket(TEXT("Add Array Node"));
+		InCPPTypeObject = URigVMCompiler::GetScriptStructForCPPType(CPPType);
 	}
+	if (InCPPTypeObject == nullptr)
+	{
+		InCPPTypeObject = URigVMPin::FindObjectFromCPPTypeObjectPath<UObject>(CPPType);
+	}
+
+	CPPType = RigVMTypeUtils::PostProcessCPPType(CPPType, InCPPTypeObject);
 	
-	const FRigVMTemplate* Template = Factory->GetTemplate();
-	URigVMTemplateNode* Node = AddTemplateNode(
-		Template->GetNotation(),
-		InPosition, 
-		InNodeName, 
-		bSetupUndoRedo, 
-		bPrintPythonCommand);
+	const FString Name = GetValidNodeName(InNodeName.IsEmpty() ? FString(TEXT("ArrayNode")) : InNodeName);
+	URigVMArrayNode* Node = NewObject<URigVMArrayNode>(Graph, *Name);
+	Node->Position = InPosition;
+	Node->OpCode = InOpCode;
 
-	if(!FRigVMRegistry::Get().IsWildCardType(ElementTypeIndex))
+	struct Local
 	{
-		FName ArgumentNameToResolve = NAME_None;
-		TRigVMTypeIndex TypeIndex = INDEX_NONE;
-		for(int32 Index = 0; Index < Template->NumArguments(); Index++)
+		static URigVMPin* AddPin(URigVMController* InController, URigVMNode* InNode, const FName& InName, ERigVMPinDirection InDirection, bool bIsArray, const FString& InCPPType, UObject* InCPPTypeObject)
 		{
-			const FRigVMTemplateArgument* Argument = Template->GetArgument(Index);
-			if(Argument->IsSingleton())
+			URigVMPin* Pin = NewObject<URigVMPin>(InNode, InName);
+			Pin->CPPType = InCPPType;
+			Pin->CPPTypeObject = InCPPTypeObject;
+			if(Pin->CPPTypeObject)
 			{
-				continue;
+				Pin->CPPTypeObjectPath = *Pin->CPPTypeObject->GetPathName();
 			}
-			if(Argument->GetArrayType() == FRigVMTemplateArgument::EArrayType_SingleValue)
+			if(bIsArray && !RigVMTypeUtils::IsArrayType(Pin->CPPType))
 			{
-				ArgumentNameToResolve = Argument->GetName();
-				TypeIndex = ElementTypeIndex;
-				break;
+				Pin->CPPType = RigVMTypeUtils::ArrayTypeFromBaseType(*Pin->CPPType);
 			}
-			if(Argument->GetArrayType() == FRigVMTemplateArgument::EArrayType_ArrayValue)
+			Pin->Direction = InDirection;
+			Pin->bIsDynamicArray = bIsArray;
+			AddNodePin(InNode, Pin);
+
+			if(Pin->Direction != ERigVMPinDirection::Hidden && !bIsArray && !Pin->IsExecuteContext())
 			{
-				ArgumentNameToResolve = Argument->GetName();
-				TypeIndex = ArrayTypeIndex;
-				break;
+				if(UScriptStruct* Struct = Cast<UScriptStruct>(Pin->CPPTypeObject))
+				{
+					FString DefaultValue;
+					InController->CreateDefaultValueForStructIfRequired(Pin->GetScriptStruct(), DefaultValue);
+					InController->AddPinsForStruct(Struct, InNode, Pin, InDirection, DefaultValue, true, false);
+				}
 			}
-		}
 
-		if(!ArgumentNameToResolve.IsNone() && TypeIndex != INDEX_NONE)
+			return Pin;
+		}
+		
+		static URigVMPin* AddExecutePin(URigVMController* InController, URigVMNode* InNode, ERigVMPinDirection InDirection = ERigVMPinDirection::IO, const FName& InName = NAME_None)
 		{
-			if(bIsPatching)
+			const FName PinName = InName.IsNone() ? FRigVMStruct::ExecuteContextName : InName;
+			UScriptStruct* ExecuteContextStruct = FRigVMExecuteContext::StaticStruct();
+			URigVMPin* Pin = AddPin(InController, InNode, PinName, InDirection, false, FString::Printf(TEXT("F%s"), *ExecuteContextStruct->GetName()), ExecuteContextStruct);
+			if(PinName == FRigVMStruct::ExecuteContextName)
 			{
-				FRigVMTemplate::FTypeMap TypeMap;
-				TypeMap.Add(ArgumentNameToResolve, TypeIndex);
+				Pin->DisplayName = FRigVMStruct::ExecuteName;
+			}
+			return Pin;
+		}
 
-				TArray<int32> Permutations;
-				Template->Resolve(TypeMap, Permutations, true);
-				check(Permutations.Num() == 1);
+		static URigVMPin* AddArrayPin(URigVMController* InController, URigVMNode* InNode, ERigVMPinDirection InDirection, const FString& InCPPType, UObject* InCPPTypeObject, const FName& InName = NAME_None)
+		{
+			const FName PinName = InName.IsNone() ? *URigVMArrayNode::ArrayName : InName;
+			return AddPin(InController, InNode, PinName, InDirection, true, InCPPType, InCPPTypeObject);  
+		}
 
-				for(const FRigVMTemplate::FTypePair& Pair : TypeMap)
-				{
-					if(!FRigVMRegistry::Get().IsWildCardType(Pair.Value))
-					{
-						if(URigVMPin* Pin = Node->FindPin(Pair.Key.ToString()))
-						{
-							ChangePinType(Pin, Pair.Value, false, false);
-							Node->UpdateFilteredPermutations(Pin, {Pair.Value});
-						}
-					}
-				}
+		static URigVMPin* AddElementPin(URigVMController* InController, URigVMNode* InNode, ERigVMPinDirection InDirection, const FString& InCPPType, UObject* InCPPTypeObject)
+		{
+			return AddPin(InController, InNode, *URigVMArrayNode::ElementName, InDirection, false, InCPPType, InCPPTypeObject);  
+		}
 
-				FullyResolveTemplateNode(Node, Permutations[0], false);
-			}
-			else if(const URigVMPin* Pin = Node->FindPin(ArgumentNameToResolve.ToString()))
+		static URigVMPin* AddIndexPin(URigVMController* InController, URigVMNode* InNode, ERigVMPinDirection InDirection)
+		{
+			return AddPin(InController, InNode, *URigVMArrayNode::IndexName, InDirection, false, RigVMTypeUtils::Int32Type, nullptr);  
+		}
+
+		static URigVMPin* AddNumPin(URigVMController* InController, URigVMNode* InNode, ERigVMPinDirection InDirection)
+		{
+			return AddPin(InController, InNode, *URigVMArrayNode::NumName, InDirection, false, RigVMTypeUtils::Int32Type, nullptr);  
+		}
+
+		static URigVMPin* AddCountPin(URigVMController* InController, URigVMNode* InNode, ERigVMPinDirection InDirection)
+		{
+			return AddPin(InController, InNode, *URigVMArrayNode::CountName, InDirection, false, RigVMTypeUtils::Int32Type, nullptr);  
+		}
+
+		static URigVMPin* AddRatioPin(URigVMController* InController, URigVMNode* InNode)
+		{
+			return AddPin(InController, InNode, *URigVMArrayNode::RatioName, ERigVMPinDirection::Output, false, RigVMTypeUtils::FloatType, nullptr);  
+		}
+
+		static URigVMPin* AddContinuePin(URigVMController* InController, URigVMNode* InNode)
+		{
+			return AddPin(InController, InNode, *URigVMArrayNode::ContinueName, ERigVMPinDirection::Hidden, false, RigVMTypeUtils::BoolType, nullptr);
+		}
+
+		static URigVMPin* AddSuccessPin(URigVMController* InController, URigVMNode* InNode)
+		{
+			return AddPin(InController, InNode, *URigVMArrayNode::SuccessName, ERigVMPinDirection::Output, false, RigVMTypeUtils::BoolType, nullptr);  
+		}
+	};
+
+	switch(InOpCode)
+	{
+		case ERigVMOpCode::ArrayReset:
+		case ERigVMOpCode::ArrayReverse:
+		{
+			Local::AddExecutePin(this, Node);
+			Local::AddArrayPin(this, Node, ERigVMPinDirection::IO, CPPType, InCPPTypeObject);
+			break;
+		}
+		case ERigVMOpCode::ArrayGetNum:
+		{
+			Local::AddArrayPin(this, Node, ERigVMPinDirection::Input, CPPType, InCPPTypeObject);
+			Local::AddNumPin(this, Node, ERigVMPinDirection::Output);
+			break;
+		} 
+		case ERigVMOpCode::ArraySetNum:
+		{
+			Local::AddExecutePin(this, Node);
+			Local::AddArrayPin(this, Node, ERigVMPinDirection::IO, CPPType, InCPPTypeObject);
+			Local::AddNumPin(this, Node, ERigVMPinDirection::Input);
+			break;
+		}
+		case ERigVMOpCode::ArrayGetAtIndex:
+		{
+			Local::AddArrayPin(this, Node, ERigVMPinDirection::Input, CPPType, InCPPTypeObject);
+			Local::AddIndexPin(this, Node, ERigVMPinDirection::Input);
+			Local::AddElementPin(this, Node, ERigVMPinDirection::Output, CPPType, InCPPTypeObject);
+			break;
+		}
+		case ERigVMOpCode::ArraySetAtIndex:
+		case ERigVMOpCode::ArrayInsert:
+		{
+			Local::AddExecutePin(this, Node);
+			Local::AddArrayPin(this, Node, ERigVMPinDirection::IO, CPPType, InCPPTypeObject);
+			Local::AddIndexPin(this, Node, ERigVMPinDirection::Input);
+			Local::AddElementPin(this, Node, ERigVMPinDirection::Input, CPPType, InCPPTypeObject);
+			break;
+		}
+		case ERigVMOpCode::ArrayAdd:
+		{
+			Local::AddExecutePin(this, Node);
+			Local::AddArrayPin(this, Node, ERigVMPinDirection::IO, CPPType, InCPPTypeObject);
+			Local::AddElementPin(this, Node, ERigVMPinDirection::Input, CPPType, InCPPTypeObject);
+			Local::AddIndexPin(this, Node, ERigVMPinDirection::Output);
+			break;
+		}
+		case ERigVMOpCode::ArrayFind:
+		{
+			Local::AddArrayPin(this, Node, ERigVMPinDirection::Input, CPPType, InCPPTypeObject);
+			Local::AddElementPin(this, Node, ERigVMPinDirection::Input, CPPType, InCPPTypeObject);
+			Local::AddIndexPin(this, Node, ERigVMPinDirection::Output);
+			Local::AddSuccessPin(this, Node);
+			break;
+		}
+		case ERigVMOpCode::ArrayRemove:
+		{
+			Local::AddExecutePin(this, Node);
+			Local::AddArrayPin(this, Node, ERigVMPinDirection::IO, CPPType, InCPPTypeObject);
+			Local::AddIndexPin(this, Node, ERigVMPinDirection::Input);
+			break;
+		}
+		case ERigVMOpCode::ArrayAppend:
+		case ERigVMOpCode::ArrayUnion:
+		{
+			Local::AddExecutePin(this, Node);
+			Local::AddArrayPin(this, Node, ERigVMPinDirection::IO, CPPType, InCPPTypeObject);
+			Local::AddArrayPin(this, Node, ERigVMPinDirection::Input, CPPType, InCPPTypeObject, *URigVMArrayNode::OtherName);
+			break;
+		}
+		case ERigVMOpCode::ArrayClone:
+		{
+			Local::AddArrayPin(this, Node, ERigVMPinDirection::Input, CPPType, InCPPTypeObject);
+			Local::AddArrayPin(this, Node, ERigVMPinDirection::Output, CPPType, InCPPTypeObject, *URigVMArrayNode::CloneName);
+			break;
+		}
+		case ERigVMOpCode::ArrayIterator:
+		{
+			Local::AddExecutePin(this, Node);
+			Local::AddArrayPin(this, Node, ERigVMPinDirection::Input, CPPType, InCPPTypeObject);
+			Local::AddElementPin(this, Node, ERigVMPinDirection::Output, CPPType, InCPPTypeObject);
+			Local::AddIndexPin(this, Node, ERigVMPinDirection::Output);
+			Local::AddCountPin(this, Node, ERigVMPinDirection::Output);
+			Local::AddRatioPin(this, Node);
+			Local::AddContinuePin(this, Node);
+			Local::AddExecutePin(this, Node, ERigVMPinDirection::Output, *URigVMArrayNode::CompletedName);
+			break;
+		}
+		case ERigVMOpCode::ArrayDifference:
+		case ERigVMOpCode::ArrayIntersection:
+		{
+			Local::AddArrayPin(this, Node, ERigVMPinDirection::Input, CPPType, InCPPTypeObject);
+			Local::AddArrayPin(this, Node, ERigVMPinDirection::Input, CPPType, InCPPTypeObject, *URigVMArrayNode::OtherName);
+			Local::AddArrayPin(this, Node, ERigVMPinDirection::Output, CPPType, InCPPTypeObject, *URigVMArrayNode::ResultName);
+			break;
+		}
+		default:
+		{
+			checkNoEntry();
+		}
+	}
+
+	Graph->Nodes.Add(Node);
+
+	if (!bSuspendNotifications)
+	{
+		Graph->MarkPackageDirty();
+	}
+
+	FRigVMControllerCompileBracketScope CompileScope(this);
+	FRigVMAddArrayNodeAction Action;
+	if (bSetupUndoRedo)
+	{
+		Action = FRigVMAddArrayNodeAction(Node);
+		Action.Title = FString::Printf(TEXT("Add %s Array Node"), *Node->GetNodeTitle());
+		ActionStack->BeginAction(Action);
+	}
+
+	Notify(ERigVMGraphNotifType::NodeAdded, Node);
+
+	if (bSetupUndoRedo)
+	{
+		ActionStack->EndAction(Action);
+	}
+
+	Node->InitializeFilteredPermutations();
+	if (InCPPType != RigVMTypeUtils::GetWildCardCPPType() && InCPPType != RigVMTypeUtils::GetWildCardArrayCPPType())
+	{
+		URigVMPin* ArrayPin = Node->FindPin(URigVMArrayNode::ArrayName);
+		PrepareTemplatePinForType(ArrayPin, {ArrayPin->GetTypeIndex()}, bSetupUndoRedo);
+		TArray<int32> FilterPermutations = Node->GetFilteredPermutationsIndices();
+		if (FilterPermutations.Num() == 1)
+		{
+			const TArray<FRigVMTemplatePreferredType> NewPreferredPermutationTypes = Node->GetPreferredTypesForPermutation(FilterPermutations[0]);
+			if (bSetupUndoRedo)
 			{
-				ResolveWildCardPin(Pin->GetPinPath(), TypeIndex, bSetupUndoRedo, bPrintPythonCommand);
+				ActionStack->AddAction(FRigVMSetPreferredTemplatePermutationsAction(Node, NewPreferredPermutationTypes));
 			}
+			Node->PreferredPermutationPairs = NewPreferredPermutationTypes;
 		}
 	}
 
-	if(bSetupUndoRedo)
+	if (bPrintPythonCommand)
 	{
-		CloseUndoBracket();
+		TArray<FString> Commands = GetAddNodePythonCommands(Node);
+		for (const FString& Command : Commands)
+		{
+			RigVMPythonUtils::Print(GetGraphOuterName(), 
+				FString::Printf(TEXT("%s"), *Command));
+		}
 	}
 
 	return Node;
 }
 
-URigVMNode* URigVMController::AddArrayNodeFromObjectPath(ERigVMOpCode InOpCode, const FString& InCPPType,
+URigVMArrayNode* URigVMController::AddArrayNodeFromObjectPath(ERigVMOpCode InOpCode, const FString& InCPPType,
 	const FString& InCPPTypeObjectPath, const FVector2D& InPosition, const FString& InNodeName, bool bSetupUndoRedo,
-	bool bPrintPythonCommand, bool bIsPatching)
+	bool bPrintPythonCommand)
 {
 	if (!IsValidGraph())
 	{
@@ -13367,7 +13701,7 @@ URigVMNode* URigVMController::AddArrayNodeFromObjectPath(ERigVMOpCode InOpCode,
 	UObject* CPPTypeObject = nullptr;
 	if (!InCPPTypeObjectPath.IsEmpty())
 	{
-		CPPTypeObject = RigVMTypeUtils::FindObjectFromCPPTypeObjectPath<UObject>(InCPPTypeObjectPath);
+		CPPTypeObject = URigVMPin::FindObjectFromCPPTypeObjectPath<UObject>(InCPPTypeObjectPath);
 		if (CPPTypeObject == nullptr)
 		{
 			ReportErrorf(TEXT("Cannot find cpp type object for path '%s'."), *InCPPTypeObjectPath);
@@ -13375,7 +13709,7 @@ URigVMNode* URigVMController::AddArrayNodeFromObjectPath(ERigVMOpCode InOpCode,
 		}
 	}
 
-	return AddArrayNode(InOpCode, InCPPType, CPPTypeObject, InPosition, InNodeName, bSetupUndoRedo, bPrintPythonCommand, bIsPatching);
+	return AddArrayNode(InOpCode, InCPPType, CPPTypeObject, InPosition, InNodeName, bSetupUndoRedo, bPrintPythonCommand);
 }
 
 URigVMInvokeEntryNode* URigVMController::AddInvokeEntryNode(const FName& InEntryName, const FVector2D& InPosition,
@@ -13639,45 +13973,39 @@ bool URigVMController::CanAddNode(URigVMNode* InNode, bool bReportErrors, bool b
 	URigVMGraph* Graph = GetGraph();
 	check(Graph);
 
-	if (Graph->IsA<URigVMFunctionLibrary>())
-	{
-		if (!InNode->IsA<URigVMCollapseNode>())
-		{
-			return false;
-		}
-	}
-
 	if (URigVMFunctionReferenceNode* FunctionRefNode = Cast<URigVMFunctionReferenceNode>(InNode))
 	{
 		if (URigVMFunctionLibrary* FunctionLibrary = FunctionRefNode->GetLibrary())
-		{			
-			FRigVMGraphFunctionHeader FunctionDefinition = FunctionRefNode->GetReferencedFunctionHeader();
-			if(!CanAddFunctionRefForDefinition(FunctionDefinition, false))
+		{
+			if (URigVMLibraryNode* FunctionDefinition = FunctionRefNode->GetReferencedNode())
 			{
-				URigVMFunctionLibrary* TargetLibrary = Graph->GetDefaultFunctionLibrary();
-				URigVMLibraryNode* NewFunctionDefinition = TargetLibrary->FindPreviouslyLocalizedFunction(FunctionDefinition.LibraryPointer);
-			
-				if((NewFunctionDefinition == nullptr) && RequestLocalizeFunctionDelegate.IsBound())
+				if(!CanAddFunctionRefForDefinition(FunctionDefinition, false))
 				{
-					if(RequestLocalizeFunctionDelegate.Execute(FunctionDefinition.LibraryPointer))
+					URigVMFunctionLibrary* TargetLibrary = Graph->GetDefaultFunctionLibrary();
+					URigVMLibraryNode* NewFunctionDefinition = TargetLibrary->FindPreviouslyLocalizedFunction(FunctionDefinition);
+					
+					if((NewFunctionDefinition == nullptr) && RequestLocalizeFunctionDelegate.IsBound())
 					{
-						NewFunctionDefinition = TargetLibrary->FindPreviouslyLocalizedFunction(FunctionDefinition.LibraryPointer);
+						if(RequestLocalizeFunctionDelegate.Execute(FunctionDefinition))
+						{
+							NewFunctionDefinition = TargetLibrary->FindPreviouslyLocalizedFunction(FunctionDefinition);
+						}
 					}
-				}
 
-				if(NewFunctionDefinition == nullptr)
+					if(NewFunctionDefinition == nullptr)
+					{
+						return false;
+					}
+					
+					SetReferencedFunction(FunctionRefNode, NewFunctionDefinition, false);
+					FunctionDefinition = NewFunctionDefinition;
+				}
+				
+				if(!CanAddFunctionRefForDefinition(FunctionDefinition, bReportErrors))
 				{
+					DestroyObject(InNode);
 					return false;
 				}
-			
-				SetReferencedFunction(FunctionRefNode, NewFunctionDefinition, false);
-				FunctionDefinition = FunctionRefNode->GetReferencedFunctionHeader();
-			}
-			
-			if(!CanAddFunctionRefForDefinition(FunctionDefinition, bReportErrors))
-			{
-				DestroyObject(InNode);
-				return false;
 			}
 		}			
 	}
@@ -13760,6 +14088,7 @@ TObjectPtr<URigVMNode> URigVMController::FindEventNode(const UScriptStruct* InSc
 
 	// construct equivalent default struct
 	FStructOnScope InDefaultStructScope(InScriptStruct);
+	InScriptStruct->InitializeDefaultValue((uint8*)InDefaultStructScope.GetStructMemory());
 	
 	if (URigVMGraph* Graph = GetGraph())
 	{
@@ -13828,7 +14157,7 @@ bool URigVMController::CanAddEventNode(UScriptStruct* InScriptStruct, const bool
 	return !bHasEventNode;
 }
 
-bool URigVMController::CanAddFunctionRefForDefinition(const FRigVMGraphFunctionHeader& InFunctionDefinition, bool bReportErrors, bool bAllowPrivateFunctions)
+bool URigVMController::CanAddFunctionRefForDefinition(URigVMLibraryNode* InFunctionDefinition, bool bReportErrors)
 {
 	if(!IsValidGraph())
 	{
@@ -13840,52 +14169,39 @@ bool URigVMController::CanAddFunctionRefForDefinition(const FRigVMGraphFunctionH
 		return false;
 	}
 
+	check(InFunctionDefinition);
+
 	URigVMGraph* Graph = GetGraph();
 	check(Graph);
 
-	if (!bAllowPrivateFunctions)
+	if(IsFunctionAvailableDelegate.IsBound())
 	{
-		if(IRigVMClientHost* ClientHost = GetImplementingOuter<IRigVMClientHost>())
+		if(!IsFunctionAvailableDelegate.Execute(InFunctionDefinition))
 		{
-			bool bIsAvailable = ClientHost->GetRigVMClient()->GetFunctionLibrary()->GetFunctionHostObjectPath() ==
-				InFunctionDefinition.LibraryPointer.HostObject;
-			if (!bIsAvailable)
-			{
-				if (IRigVMGraphFunctionHost* Host = Cast<IRigVMGraphFunctionHost>(InFunctionDefinition.LibraryPointer.HostObject.TryLoad()))
-				{
-					bIsAvailable = Host->GetRigVMGraphFunctionStore()->IsFunctionPublic(InFunctionDefinition.LibraryPointer);		
-				}
-			}
-			if (!bIsAvailable)
+			if(bReportErrors)
 			{
-				if(bReportErrors)
-				{
-					ReportAndNotifyError(TEXT("Function is not available for placement in another graph host."));
-				}
-				return false;
+				ReportAndNotifyError(TEXT("Function is not available for placement in another graph host."));
 			}
+			return false;
 		}
 	}
 
-	if (URigVMNode* Node = Cast<URigVMNode>(Graph->GetOuter()))
+	if(IsDependencyCyclicDelegate.IsBound())
 	{
-		if (URigVMLibraryNode* LibraryNode = Node->FindFunctionForNode())
+		if(IsDependencyCyclicDelegate.Execute(Graph, InFunctionDefinition))
 		{
-			if (InFunctionDefinition.Dependencies.Contains(LibraryNode->GetFunctionIdentifier()))
+			if(bReportErrors)
 			{
-				if(bReportErrors)
-				{
-					ReportAndNotifyError(TEXT("Function is not available for placement in this graph host due to dependency cycles."));
-				}
-				return false;
+				ReportAndNotifyError(TEXT("Function is not available for placement in this graph host due to dependency cycles."));
 			}
+			return false;
 		}
 	}
-	
+
 	URigVMLibraryNode* ParentLibraryNode = Cast<URigVMLibraryNode>(Graph->GetOuter());
 	while (ParentLibraryNode)
 	{
-		if (TSoftObjectPtr<UObject>(ParentLibraryNode).ToSoftObjectPath() == InFunctionDefinition.LibraryPointer.LibraryNode)
+		if (ParentLibraryNode == InFunctionDefinition)
 		{
 			if(bReportErrors)
 			{
@@ -13899,12 +14215,14 @@ bool URigVMController::CanAddFunctionRefForDefinition(const FRigVMGraphFunctionH
 	return true;
 }
 
-void URigVMController::AddPinsForStruct(UStruct* InStruct, URigVMNode* InNode, URigVMPin* InParentPin, ERigVMPinDirection InPinDirection, const FString& InDefaultValue, bool bAutoExpandArrays)
+void URigVMController::AddPinsForStruct(UStruct* InStruct, URigVMNode* InNode, URigVMPin* InParentPin, ERigVMPinDirection InPinDirection, const FString& InDefaultValue, bool bAutoExpandArrays, bool bNotify, const FPinInfoArray* PreviousPins)
 {
 	if(!ShouldStructBeUnfolded(InStruct))
 	{
 		return;
 	}
+
+	// todo: reuse the default values when creating the pins
 	
 	TArray<FString> MemberNameValuePairs = URigVMPin::SplitDefaultValue(InDefaultValue);
 	TMap<FName, FString> MemberValues;
@@ -13951,10 +14269,8 @@ void URigVMController::AddPinsForStruct(UStruct* InStruct, URigVMNode* InNode, U
 						DefaultValue = *DefaultValuePtr;
 					}
 					CreateDefaultValueForStructIfRequired(StructProperty->Struct, DefaultValue);
-					{
-						TGuardValue<bool> SuspendNotifications(bSuspendNotifications, true);
-						AddPinsForStruct(StructProperty->Struct, InNode, Pin, Pin->GetDirection(), DefaultValue, bAutoExpandArrays);
-					}
+
+					AddPinsForStruct(StructProperty->Struct, InNode, Pin, Pin->GetDirection(), DefaultValue, bAutoExpandArrays);
 				}
 				else if(DefaultValuePtr != nullptr)
 				{
@@ -13990,7 +14306,7 @@ void URigVMController::AddPinsForStruct(UStruct* InStruct, URigVMNode* InNode, U
 				Pin->DefaultValue = DefaultValue;
 			}
 
-			if (!bSuspendNotifications)
+			if (bNotify)
 			{
 				Notify(ERigVMGraphNotifType::PinAdded, Pin);
 			}
@@ -14034,10 +14350,7 @@ void URigVMController::AddPinsForArray(FArrayProperty* InArrayProperty, URigVMNo
 				{
 					CreateDefaultValueForStructIfRequired(ScriptStruct, DefaultValue);
 				}
-				{
-					TGuardValue<bool> SuspendNotifications(bSuspendNotifications, true);
-					AddPinsForStruct(StructProperty->Struct, InNode, Pin, Pin->Direction, DefaultValue, bAutoExpandArrays);
-				}
+				AddPinsForStruct(StructProperty->Struct, InNode, Pin, Pin->Direction, DefaultValue, bAutoExpandArrays);
 			}
 			else if (!DefaultValue.IsEmpty())
 			{
@@ -14069,134 +14382,25 @@ void URigVMController::AddPinsForArray(FArrayProperty* InArrayProperty, URigVMNo
 	}
 }
 
-void URigVMController::AddPinsForTemplate(const FRigVMTemplate* InTemplate, const FRigVMTemplateTypeMap& InPinTypeMap, URigVMNode* InNode)
+void URigVMController::ConfigurePinFromProperty(FProperty* InProperty, URigVMPin* InOutPin, ERigVMPinDirection InPinDirection)
 {
-	const FRigVMRegistry& Registry = FRigVMRegistry::Get();
-
-	FRigVMDispatchContext DispatchContext;
-	if(const URigVMDispatchNode* DispatchNode = Cast<URigVMDispatchNode>(InNode))
+	if (InPinDirection == ERigVMPinDirection::Invalid)
 	{
-		DispatchContext = DispatchNode->GetDispatchContext();
+		InOutPin->Direction = FRigVMStruct::GetPinDirectionFromProperty(InProperty);
 	}
-
-	auto AddExecutePins = [InTemplate, InNode, &Registry, &DispatchContext, this](ERigVMPinDirection InPinDirection)
+	else
 	{
-		for (int32 ArgIndex = 0; ArgIndex < InTemplate->NumExecuteArguments(DispatchContext); ArgIndex++)
-		{
-			const FRigVMExecuteArgument* Arg = InTemplate->GetExecuteArgument(ArgIndex, DispatchContext);
-			if(Arg->Direction != InPinDirection)
-			{
-				continue;
-			}
-			
-			URigVMPin* Pin = NewObject<URigVMPin>(InNode, Arg->Name);
-			const FRigVMTemplateArgumentType Type = Registry.GetType(Arg->TypeIndex);
-	        
-			Pin->CPPType = Type.CPPType.ToString();
-			Pin->CPPTypeObject = Type.CPPTypeObject;
-			if (Pin->CPPTypeObject)
-			{
-				Pin->CPPTypeObjectPath = *Pin->CPPTypeObject->GetPathName();
-			}
-			Pin->Direction = Arg->Direction;
-			Pin->LastKnownTypeIndex = Arg->TypeIndex;
-			Pin->LastKnownCPPType = Pin->CPPType;
-
-			AddNodePin(InNode, Pin);
-
-			if(Registry.IsArrayType(Arg->TypeIndex))
-			{
-				if(const URigVMDispatchNode* DispatchNode = Cast<URigVMDispatchNode>(Pin->GetNode()))
-				{
-					if(const FRigVMDispatchFactory* Factory = DispatchNode->GetFactory())
-					{
-						const FString DefaultValue =  Factory->GetArgumentDefaultValue(Pin->GetFName(), Arg->TypeIndex);
-						if(!DefaultValue.IsEmpty())
-						{
-							SetPinDefaultValue(Pin, DefaultValue, true, false, false);
-						}
-					}
-				}
-			}
-		}
-	};
+		InOutPin->Direction = InPinDirection;
+	}
 
-	AddExecutePins(ERigVMPinDirection::IO);
-	AddExecutePins(ERigVMPinDirection::Input);
+#if WITH_EDITOR
 
-	for (int32 ArgIndex = 0; ArgIndex < InTemplate->NumArguments(); ArgIndex++)
+	if (!InOutPin->IsArrayElement())
 	{
-		const FRigVMTemplateArgument* Arg = InTemplate->GetArgument(ArgIndex);
-
-		URigVMPin* Pin = NewObject<URigVMPin>(InNode, Arg->GetName());
-		const TRigVMTypeIndex& TypeIndex = InPinTypeMap.FindChecked(Arg->GetName());
-		const FRigVMTemplateArgumentType Type = FRigVMRegistry::Get().GetType(TypeIndex);
-		Pin->CPPType = Type.CPPType.ToString();
-		Pin->CPPTypeObject = Type.CPPTypeObject;
-		if (Pin->CPPTypeObject)
+		FString DisplayNameText = InProperty->GetDisplayNameText().ToString();
+		if (!DisplayNameText.IsEmpty())
 		{
-			Pin->CPPTypeObjectPath = *Pin->CPPTypeObject->GetPathName();
-		}
-		Pin->Direction = Arg->GetDirection();
-
-		AddNodePin(InNode, Pin);
-
-		if(!Pin->IsWildCard() && !Pin->IsArray())
-		{
-			FString DefaultValue;
-			if(const URigVMTemplateNode* TemplateNode = Cast<URigVMTemplateNode>(InNode))
-			{
-				DefaultValue = TemplateNode->GetInitialDefaultValueForPin(Pin->GetFName());
-			}
-
-			TGuardValue<bool> SuspendNotifications(bSuspendNotifications, true);
-			if(UScriptStruct* ScriptStruct = Cast<UScriptStruct>(Pin->CPPTypeObject))
-			{
-				AddPinsForStruct(ScriptStruct, Pin->GetNode(), Pin, Pin->Direction, DefaultValue, false);
-			}
-			else if(!DefaultValue.IsEmpty())
-			{
-				SetPinDefaultValue(Pin, DefaultValue, true, false, false);
-			}
-		}
-		else if(Pin->IsFixedSizeArray())
-		{
-			if(const URigVMDispatchNode* DispatchNode = Cast<URigVMDispatchNode>(Pin->GetNode()))
-			{
-				if(const FRigVMDispatchFactory* Factory = DispatchNode->GetFactory())
-				{
-					const FString DefaultValue =  Factory->GetArgumentDefaultValue(Pin->GetFName(), RigVMTypeUtils::TypeIndex::WildCardArray);
-					if(!DefaultValue.IsEmpty())
-					{
-						SetPinDefaultValue(Pin, DefaultValue, true, false, false);
-					}
-				}
-			}
-		}
-	}
-
-	AddExecutePins(ERigVMPinDirection::Output);
-}
-
-void URigVMController::ConfigurePinFromProperty(FProperty* InProperty, URigVMPin* InOutPin, ERigVMPinDirection InPinDirection)
-{
-	if (InPinDirection == ERigVMPinDirection::Invalid)
-	{
-		InOutPin->Direction = FRigVMStruct::GetPinDirectionFromProperty(InProperty);
-	}
-	else
-	{
-		InOutPin->Direction = InPinDirection;
-	}
-
-#if WITH_EDITOR
-
-	if (!InOutPin->IsArrayElement())
-	{
-		FString DisplayNameText = InProperty->GetDisplayNameText().ToString();
-		if (!DisplayNameText.IsEmpty())
-		{
-			InOutPin->DisplayName = *DisplayNameText;
+			InOutPin->DisplayName = *DisplayNameText;
 		}
 		else
 		{
@@ -14319,29 +14523,6 @@ void URigVMController::ConfigurePinFromPin(URigVMPin* InOutPin, URigVMPin* InPin
 	}
 }
 
-void URigVMController::ConfigurePinFromArgument(URigVMPin* InOutPin, const FRigVMGraphFunctionArgument& InArgument, bool bCopyDisplayName)
-{
-	// it is important we copy things that define the identity of the pin
-	// things that defines the state of the pin is copied during GetPinState()
-	// though addmittedly these two functions have overlaps currently
-	InOutPin->bIsConstant = InArgument.bIsConst;
-	InOutPin->Direction = InArgument.Direction;
-	InOutPin->CPPType = InArgument.CPPType.ToString();
-	InOutPin->CPPTypeObjectPath = *InArgument.CPPTypeObject.ToSoftObjectPath().ToString();
-	InOutPin->CPPTypeObject = InArgument.CPPTypeObject.Get();
-	InOutPin->DefaultValue = InArgument.DefaultValue;
-	InOutPin->bIsDynamicArray = InArgument.bIsArray;
-	if(bCopyDisplayName)
-	{
-		InOutPin->SetDisplayName(InArgument.DisplayName);
-	}
-
-	if(InOutPin->IsExecuteContext() && InOutPin->CPPTypeObject != FRigVMExecuteContext::StaticStruct())
-	{
-		MakeExecutePin(InOutPin);
-	}
-}
-
 bool URigVMController::ShouldStructBeUnfolded(const UStruct* Struct)
 {
 	if (Struct == nullptr)
@@ -14376,11 +14557,10 @@ bool URigVMController::ShouldPinBeUnfolded(URigVMPin* InPin)
 	{
 		return ShouldStructBeUnfolded(InPin->GetScriptStruct());
 	}
-	if (InPin->IsArray())
+	else if (InPin->IsArray())
 	{
 		return InPin->GetDirection() == ERigVMPinDirection::Input ||
-			InPin->GetDirection() == ERigVMPinDirection::IO ||
-			InPin->IsFixedSizeArray();
+			InPin->GetDirection() == ERigVMPinDirection::IO;
 	}
 	return false;
 }
@@ -14446,10 +14626,25 @@ FProperty* URigVMController::FindPropertyForPin(const FString& InPinPath)
 	return nullptr;
 }
 
-int32 URigVMController::DetachLinksFromPinObjects(const TArray<URigVMLink*>* InLinks)
+URigVMBuildData* URigVMController::GetBuildData(bool bCreateIfNeeded)
+{
+	static TStrongObjectPtr<URigVMBuildData> sBuildData;
+	if(!sBuildData.IsValid() && bCreateIfNeeded && IsInGameThread())
+	{
+		sBuildData = TStrongObjectPtr<URigVMBuildData>(
+			NewObject<URigVMBuildData>(
+				GetTransientPackage(), 
+				TEXT("RigVMBuildData"), 
+				RF_Transient));
+	}
+	return sBuildData.Get();
+}
+
+int32 URigVMController::DetachLinksFromPinObjects(const TArray<URigVMLink*>* InLinks, bool bNotify)
 {
 	URigVMGraph* Graph = GetGraph();
 	check(Graph);
+	TGuardValue<bool> EventuallySuspendNotifs(bSuspendNotifications, !bNotify);
 
 	TArray<URigVMLink*> Links;
 	if (InLinks)
@@ -14492,7 +14687,7 @@ int32 URigVMController::DetachLinksFromPinObjects(const TArray<URigVMLink*>* InL
 			{
 				FRigVMControllerGraphGuard GraphGuard(this, CollapseNode->GetContainedGraph(), false);
 				TGuardValue<bool> GuardEditGraph(CollapseNode->ContainedGraph->bEditable, true);
-				DetachLinksFromPinObjects(InLinks);
+				DetachLinksFromPinObjects(InLinks, bNotify);
 			}
 		}
 	}
@@ -14500,10 +14695,11 @@ int32 URigVMController::DetachLinksFromPinObjects(const TArray<URigVMLink*>* InL
 	return Links.Num();
 }
 
-int32 URigVMController::ReattachLinksToPinObjects(bool bFollowCoreRedirectors, const TArray<URigVMLink*>* InLinks, bool bSetupOrphanedPins, bool bAllowNonArgumentLinks)
+int32 URigVMController::ReattachLinksToPinObjects(bool bFollowCoreRedirectors, const TArray<URigVMLink*>* InLinks, bool bNotify, bool bSetupOrphanedPins, bool bAllowNonArgumentLinks)
 {
 	URigVMGraph* Graph = GetGraph();
 	check(Graph);
+	TGuardValue<bool> EventuallySuspendNotifs(bSuspendNotifications, !bNotify);
 	FScopeLock Lock(&PinPathCoreRedirectorsLock);
 
 	bool bReplacingAllLinks = false;
@@ -14576,7 +14772,7 @@ int32 URigVMController::ReattachLinksToPinObjects(bool bFollowCoreRedirectors, c
 				}
 				else
 				{
-					ReportWarningf(TEXT("Unable to re-create link %s"), *URigVMLink::GetPinPathRepresentation(Link->SourcePinPath, Link->TargetPinPath));
+					ReportWarningf(TEXT("Unable to re-create link %s -> %s"), *Link->SourcePinPath, *Link->TargetPinPath);
 					TargetPin->Links.Remove(Link);
 					SourcePin->Links.Remove(Link);
 					continue;
@@ -14629,7 +14825,7 @@ int32 URigVMController::ReattachLinksToPinObjects(bool bFollowCoreRedirectors, c
 		
 		if (SourcePin == nullptr)
 		{
-			ReportWarningf(TEXT("Unable to re-create link %s"), *URigVMLink::GetPinPathRepresentation(Link->SourcePinPath, Link->TargetPinPath));
+			ReportWarningf(TEXT("Unable to re-create link %s -> %s"), *Link->SourcePinPath, *Link->TargetPinPath);
 			if (TargetPin != nullptr)
 			{
 				TargetPin->Links.Remove(Link);
@@ -14638,7 +14834,7 @@ int32 URigVMController::ReattachLinksToPinObjects(bool bFollowCoreRedirectors, c
 		}
 		if (TargetPin == nullptr)
 		{
-			ReportWarningf(TEXT("Unable to re-create link %s"), *URigVMLink::GetPinPathRepresentation(Link->SourcePinPath, Link->TargetPinPath));
+			ReportWarningf(TEXT("Unable to re-create link %s -> %s"), *Link->SourcePinPath, *Link->TargetPinPath);
 			if (SourcePin != nullptr)
 			{
 				SourcePin->Links.Remove(Link);
@@ -14687,7 +14883,7 @@ int32 URigVMController::ReattachLinksToPinObjects(bool bFollowCoreRedirectors, c
 			{
 				FRigVMControllerGraphGuard GraphGuard(this, CollapseNode->GetContainedGraph(), false);
 				TGuardValue<bool> GuardEditGraph(CollapseNode->ContainedGraph->bEditable, true);
-				ReattachLinksToPinObjects(bFollowCoreRedirectors, nullptr, bSetupOrphanedPins, bAllowNonArgumentLinks);
+				ReattachLinksToPinObjects(bFollowCoreRedirectors, nullptr, bNotify, bSetupOrphanedPins, bAllowNonArgumentLinks);
 			}
 		}
 	}
@@ -14887,7 +15083,7 @@ bool URigVMController::ShouldRedirectPin(const FString& InOldPinPath, FString& I
 	return false;
 }
 
-void URigVMController::RepopulatePinsOnNode(URigVMNode* InNode, bool bFollowCoreRedirectors, bool bSetupOrphanedPins)
+void URigVMController::RepopulatePinsOnNode(URigVMNode* InNode, bool bFollowCoreRedirectors, bool bNotify, bool bSetupOrphanedPins, bool bDetachAndReattachLinks)
 {
 	if (InNode == nullptr)
 	{
@@ -14904,9 +15100,14 @@ void URigVMController::RepopulatePinsOnNode(URigVMNode* InNode, bool bFollowCore
 	URigVMCollapseNode* CollapseNode = Cast<URigVMCollapseNode>(InNode);
 	URigVMFunctionReferenceNode* FunctionRefNode = Cast<URigVMFunctionReferenceNode>(InNode);
 	URigVMVariableNode* VariableNode = Cast<URigVMVariableNode>(InNode);
+	URigVMIfNode* IfNode = Cast<URigVMIfNode>(InNode);
+	URigVMSelectNode* SelectNode = Cast<URigVMSelectNode>(InNode);
+	URigVMArrayNode* ArrayNode = Cast<URigVMArrayNode>(InNode);
 	URigVMDispatchNode* DispatchNode = Cast<URigVMDispatchNode>(InNode);
 
+	TGuardValue<bool> EventuallySuspendNotifs(bSuspendNotifications, !bNotify);
 	FScopeLock Lock(&PinPathCoreRedirectorsLock);
+	const FRigVMRegistry& Registry = FRigVMRegistry::Get();
 
 	URigVMGraph* Graph = GetGraph();
 	check(Graph);
@@ -14920,22 +15121,18 @@ void URigVMController::RepopulatePinsOnNode(URigVMNode* InNode, bool bFollowCore
 		}
 	}
 
-	// step 1/3: keep a record of the current state of the node's pins
-	TMap<FString, FString> RedirectedPinPaths;
-	if (bFollowCoreRedirectors)
+	bool bSetupOrphanPinsForThisNode = bSetupOrphanedPins;
+	if(CollapseNode)
 	{
-		RedirectedPinPaths = GetRedirectedPinPaths(InNode);
+		if(CollapseNode->GetOuter()->IsA<URigVMFunctionLibrary>())
+		{
+			bSetupOrphanPinsForThisNode = false;
+		}
 	}
-	TMap<FString, FPinState> PinStates = GetPinStates(InNode);
 
-	// also in case this node is part of an injection
-	FName InjectionInputPinName = NAME_None;
-	FName InjectionOutputPinName = NAME_None;
-	if (URigVMInjectionInfo* InjectionInfo = InNode->GetInjectionInfo())
-	{
-		InjectionInputPinName = InjectionInfo->InputPin ? InjectionInfo->InputPin->GetFName() : NAME_None;
-		InjectionOutputPinName = InjectionInfo->OutputPin ? InjectionInfo->OutputPin->GetFName() : NAME_None;
-	}
+	const FPinInfoArray PreviousPinInfos(InNode);
+	const uint32 PreviousPinHash = GetTypeHash(PreviousPinInfos);
+	FPinInfoArray NewPinInfos;
 
 	// step 2/3: clear pins on the node and repopulate the node with new pins
 	if (UnitNode != nullptr)
@@ -14948,20 +15145,6 @@ void URigVMController::RepopulatePinsOnNode(URigVMNode* InNode, bool bFollowCore
 			return;
 		}
 
-		RemovePinsDuringRepopulate(UnitNode, UnitNode->Pins, bSetupOrphanedPins);
-
-		if (ScriptStruct == nullptr)
-		{
-			ReportWarningf(
-				TEXT("Control Rig '%s', Node '%s' has no struct assigned. Do you have a broken redirect?"),
-				*UnitNode->GetOutermost()->GetPathName(),
-				*UnitNode->GetName()
-				);
-
-			RemoveNode(UnitNode, false, true);
-			return;
-		}
-
 		FString NodeColorMetadata;
 		ScriptStruct->GetStringMetaDataHierarchical(*URigVMNode::NodeColorName, &NodeColorMetadata);
 		if (!NodeColorMetadata.IsEmpty())
@@ -14969,24 +15152,80 @@ void URigVMController::RepopulatePinsOnNode(URigVMNode* InNode, bool bFollowCore
 			UnitNode->NodeColor = GetColorFromMetadata(NodeColorMetadata);
 		}
 
-		FString ExportedDefaultValue;
-		CreateDefaultValueForStructIfRequired(ScriptStruct, ExportedDefaultValue);
-		AddPinsForStruct(ScriptStruct, UnitNode, nullptr, ERigVMPinDirection::Invalid, ExportedDefaultValue, false);
+		const TSharedPtr<FStructOnScope> DefaultValueContent = UnitNode->ConstructStructInstance(false);
+		NewPinInfos.AddPins(ScriptStruct, this, ERigVMPinDirection::Invalid, INDEX_NONE, DefaultValueContent->GetStructMemory());
 	}
 	else if (DispatchNode)
 	{
-		FRigVMTemplateTypeMap PinTypeMap;
+		TMap<FName, TRigVMTypeIndex> PinTypeMap;
 		for (const URigVMPin* Pin : DispatchNode->Pins)
 		{
 			PinTypeMap.Add(Pin->GetFName(), Pin->GetTypeIndex());
 		}
 		
-		RemovePinsDuringRepopulate(DispatchNode, DispatchNode->Pins, bSetupOrphanedPins);
-
 		const FRigVMTemplate* Template = DispatchNode->GetTemplate();
-		AddPinsForTemplate(Template, PinTypeMap, DispatchNode);
 		
-		ResolveTemplateNodeMetaData(DispatchNode, false);
+		for (int32 ArgIndex = 0; ArgIndex < Template->NumArguments(); ArgIndex++)
+		{
+			const FRigVMTemplateArgument* Arg = Template->GetArgument(ArgIndex);
+
+			TRigVMTypeIndex TypeIndex = INDEX_NONE;
+			if(const TRigVMTypeIndex* ExistingTypeIndex = PinTypeMap.Find(Arg->GetName()))
+			{
+				TypeIndex = *ExistingTypeIndex;
+				if(!Arg->SupportsTypeIndex(TypeIndex))
+				{
+					TypeIndex = INDEX_NONE;
+				}
+			}
+
+			if(TypeIndex == INDEX_NONE)
+			{
+				if(Arg->IsSingleton())
+				{
+					TypeIndex = Arg->GetSupportedTypeIndices()[0];
+				}
+				else if(Arg->GetArrayType() == FRigVMTemplateArgument::EArrayType_ArrayValue)
+				{
+					TypeIndex = RigVMTypeUtils::TypeIndex::WildCardArray;
+				}
+				else
+				{
+					TypeIndex = RigVMTypeUtils::TypeIndex::WildCard;
+				}
+			}
+
+			FString DefaultValue;
+			UScriptStruct* ArgumentScriptStruct = nullptr;
+			const uint8* DefaultValueMemory = nullptr;
+
+			if(const URigVMPin* ArgumentPin = DispatchNode->FindPin(Arg->Name.ToString()))
+			{
+				DefaultValue = ArgumentPin->GetDefaultValue();
+				ArgumentScriptStruct = Cast<UScriptStruct>(ArgumentPin->GetCPPTypeObject());
+			}
+			else if(const FRigVMDispatchFactory* Factory = DispatchNode->GetFactory())
+			{
+				if(Arg->IsSingleton())
+				{
+					DefaultValue = Factory->GetArgumentDefaultValue(Arg->Name, Arg->GetTypeIndices()[0]);
+					const FRigVMTemplateArgumentType& Type = Registry.GetType(Arg->GetTypeIndices()[0]);
+					ArgumentScriptStruct = Cast<UScriptStruct>(Type.CPPTypeObject);
+				}
+			}
+
+			FStructOnScope DefaultValueMemoryScope; // has to be in this scope so that DefaultValueMemory is valid
+			if(ArgumentScriptStruct && !DefaultValue.IsEmpty())
+			{
+				DefaultValueMemoryScope = FStructOnScope(ArgumentScriptStruct);
+
+				FRigVMPinDefaultValueImportErrorContext ErrorPipe;
+				ArgumentScriptStruct->ImportText(*DefaultValue, DefaultValueMemoryScope.GetStructMemory(), nullptr, PPF_None, &ErrorPipe, FString());
+				DefaultValueMemory = DefaultValueMemoryScope.GetStructMemory();
+			}
+			
+			(void)NewPinInfos.AddPin(this, INDEX_NONE, Arg->Name, Arg->GetDirection(), TypeIndex, DefaultValue, DefaultValueMemory);
+		}
 	}
 	else if ((RerouteNode != nullptr) || (VariableNode != nullptr))
 	{
@@ -15048,7 +15287,7 @@ void URigVMController::RepopulatePinsOnNode(URigVMNode* InNode, bool bFollowCore
 				
 					if(RigVMTypeUtils::CPPTypeFromExternalVariable(Variable, CPPType, &CPPTypeObject))
 					{
-						RefreshVariableNode(VariableNode->GetFName(), Variable.Name, CPPType, Variable.TypeObject, false, bSetupOrphanedPins);
+						RefreshVariableNode(VariableNode->GetFName(), Variable.Name, CPPType, Variable.TypeObject, false, bSetupOrphanPinsForThisNode);
 					}
 					else
 					{
@@ -15069,38 +15308,14 @@ void URigVMController::RepopulatePinsOnNode(URigVMNode* InNode, bool bFollowCore
 				);
 			}
 		}
-		
-		RemovePinsDuringRepopulate(InNode, ValuePin->SubPins, bSetupOrphanedPins);
 
-		if (ValuePin->IsStruct())
-		{
-			UScriptStruct* ScriptStruct = ValuePin->GetScriptStruct();
-			if (ScriptStruct == nullptr)
-			{
-				ReportErrorf(
-					TEXT("Control Rig '%s', Node '%s' has no struct assigned. Do you have a broken redirect?"),
-					*InNode->GetOutermost()->GetPathName(),
-					*InNode->GetName()
-				);
-
-				RemoveNode(InNode, false, true);
-				return;
-			}
-
-			FString ExportedDefaultValue;
-			CreateDefaultValueForStructIfRequired(ScriptStruct, ExportedDefaultValue);
-			{
-				TGuardValue<bool> SuspendNotifications(bSuspendNotifications, true);
-				AddPinsForStruct(ScriptStruct, InNode, ValuePin, ValuePin->Direction, ExportedDefaultValue, false);
-			}
-		}
+		NewPinInfos = FPinInfoArray(InNode, this);
 	}
 	else if (EntryNode || ReturnNode)
 	{
 		if (URigVMLibraryNode* LibraryNode = Cast<URigVMLibraryNode>(InNode->GetGraph()->GetOuter()))
 		{
 			bool bIsEntryNode = EntryNode != nullptr;
-			RemovePinsDuringRepopulate(InNode, InNode->Pins, bSetupOrphanedPins);
 
 			TArray<URigVMPin*> SortedLibraryPins;
 
@@ -15141,139 +15356,516 @@ void URigVMController::RepopulatePinsOnNode(URigVMNode* InNode, bool bFollowCore
 					}
 				}
 
-				URigVMPin* ExposedPin = NewObject<URigVMPin>(InNode, LibraryPin->GetFName());
-				ConfigurePinFromPin(ExposedPin, LibraryPin, true);
+				ERigVMPinDirection Direction = bIsEntryNode ? ERigVMPinDirection::Output : ERigVMPinDirection::Input;  
+				(void)NewPinInfos.AddPin(LibraryPin, INDEX_NONE, Direction);
+			}
+		}
+		else
+		{
+			// due to earlier bugs with copy and paste we can find entry and return nodes under the top level
+			// graph. we'll ignore these for now.
+		}
+	}
+	else if (CollapseNode)
+	{
+		NewPinInfos = FPinInfoArray(InNode, this);
+	}
+	else if (FunctionRefNode)
+	{
+		if (URigVMLibraryNode* ReferencedNode = FunctionRefNode->GetReferencedNode())
+		{
+			NewPinInfos = FPinInfoArray(ReferencedNode, this);
+		}
+	}
+	else if (IfNode || SelectNode || ArrayNode)
+	{
+		NewPinInfos = FPinInfoArray(InNode, this);
+	}
+	else
+	{
+		return;
+	}
 
-				if (bIsEntryNode)
-				{
-					ExposedPin->Direction = ERigVMPinDirection::Output;
-				}
-				else
+	auto RecursivelyRepopulatePinsOnCollapseNode = [this, bFollowCoreRedirectors, bNotify, bSetupOrphanedPins]
+		(URigVMCollapseNode* InCollapseNode)
+		{
+			if(InCollapseNode)
+			{
+				FRigVMControllerGraphGuard GraphGuard(this, InCollapseNode->GetContainedGraph(), false);
+				TGuardValue<bool> GuardEditGraph(InCollapseNode->ContainedGraph->bEditable, true);
+				// need to get a copy of the node array since the following function could remove nodes from the graph
+				// we don't want to remove elements from the array we are iterating over.
+				TArray<URigVMNode*> ContainedNodes = InCollapseNode->GetContainedNodes();
+				for (URigVMNode* ContainedNode : ContainedNodes)
 				{
-					ExposedPin->Direction = ERigVMPinDirection::Input;
+					RepopulatePinsOnNode(ContainedNode, bFollowCoreRedirectors, bNotify, bSetupOrphanedPins);
 				}
+			}
+		};
 
-				AddNodePin(InNode, ExposedPin);
+	// if the nodes match in structure - nothing to do here
+	if(GetTypeHash(NewPinInfos) == PreviousPinHash)
+	{
+		// we at least need to recurse into a collapse node
+		RecursivelyRepopulatePinsOnCollapseNode(CollapseNode);
+		return;
+	}
+	
+#if UE_RIGVMCONTROLLER_VERBOSE_REPOPULATE
+	UE_LOG(LogRigVMDeveloper, Display, TEXT("Repopulating pins on node %s"), *InNode->GetPathName());
+#endif
+	
+	bool bRequireDetachLinks = false;
+	bool bRequirePinStates = false;
 
-				if (ExposedPin->IsStruct())
+	TArray<int32> NewPinsToAdd, PreviousPinsToRemove, PreviousPinsToOrphan, PreviousPinsToUpdate;
+	for(int32 Index = 0; Index < PreviousPinInfos.Num(); Index++)
+	{
+		const FString PinPath = PreviousPinInfos.GetPinPath(Index);
+		const int32 NewIndex = NewPinInfos.GetIndexFromPinPath(PinPath);
+		if(NewIndex == INDEX_NONE)
+		{
+			const int32 RootIndex = PreviousPinInfos.GetRootIndex(Index);
+			
+			if(PreviousPinInfos[Index].Direction != ERigVMPinDirection::Hidden)
+			{
+				if(URigVMPin* Pin = InNode->FindPin(PinPath))
 				{
-					TGuardValue<bool> SuspendNotifications(bSuspendNotifications, true);
-					AddPinsForStruct(ExposedPin->GetScriptStruct(), InNode, ExposedPin, ExposedPin->GetDirection(), FString(), false);
+					if(!Pin->GetLinks().IsEmpty())
+					{
+						bRequireDetachLinks = true;
+						bRequirePinStates = true;
+					}
+					
+					if(!PreviousPinsToOrphan.Contains(RootIndex))
+					{
+						URigVMPin* RootPin = Pin->GetRootPin();
+
+						bool bPinShouldBeOrphaned = bSetupOrphanPinsForThisNode; 
+						if(!bPinShouldBeOrphaned)
+						{
+							if(RootPin->GetSourceLinks(true).Num() > 0 ||
+								RootPin->GetTargetLinks(true).Num() > 0)
+							{
+								bPinShouldBeOrphaned = true;
+							}
+						}
+		
+						if(bPinShouldBeOrphaned)
+						{
+							PreviousPinsToOrphan.Add(RootIndex);
+							bRequireDetachLinks = true;
+							bRequirePinStates = true;
+#if UE_RIGVMCONTROLLER_VERBOSE_REPOPULATE
+							UE_LOG(LogRigVMDeveloper, Display, TEXT("Previously existing pin '%s' needs to be orphaned."), *RootPin->GetPinPath());
+#endif
+						}
+					}
 				}
+			}
 
-				Notify(ERigVMGraphNotifType::PinAdded, ExposedPin);
+			if(!PreviousPinsToOrphan.Contains(RootIndex))
+			{
+				PreviousPinsToRemove.Add(Index);
+#if UE_RIGVMCONTROLLER_VERBOSE_REPOPULATE
+				UE_LOG(LogRigVMDeveloper, Display, TEXT("Previously existing pin '%s' is now obsolete."), *PinPath);
+#endif
 			}
 		}
+		else if(GetTypeHash(PreviousPinInfos[Index]) != GetTypeHash(NewPinInfos[NewIndex]))
+		{
+			const bool bTypesDiffer = !Registry.CanMatchTypes(PreviousPinInfos[Index].TypeIndex, NewPinInfos[NewIndex].TypeIndex, true);
+			if(PreviousPinInfos[Index].Direction != NewPinInfos[NewIndex].Direction)
+			{
+				bRequireDetachLinks = true;
+			}
+			else if(PreviousPinInfos[Index].Direction != ERigVMPinDirection::Hidden)
+			{
+				bRequireDetachLinks |= bTypesDiffer;
+			}
+
+			if(Registry.CanMatchTypes(PreviousPinInfos[Index].TypeIndex, NewPinInfos[NewIndex].TypeIndex, true))
+			{
+				PreviousPinsToUpdate.Add(Index);
+			}
+			
+#if UE_RIGVMCONTROLLER_VERBOSE_REPOPULATE
+			const FName& PreviousCPPType = Registry.GetType(PreviousPinInfos[Index].TypeIndex).CPPType;
+			const FName& NewCPPType = Registry.GetType(NewPinInfos[NewIndex].TypeIndex).CPPType;
+
+			const FString PreviousDirection = StaticEnum<ERigVMPinDirection>()->GetDisplayNameTextByValue((int64)PreviousPinInfos[Index].Direction).ToString();
+			const FString NewDirection = StaticEnum<ERigVMPinDirection>()->GetDisplayNameTextByValue((int64)NewPinInfos[NewIndex].Direction).ToString();
+
+			UE_LOG(LogRigVMDeveloper, Display,
+				TEXT("Previous pin '%s' (Index %d, %s, %s) differs with new pin (Index %d, %s, %s)."),
+				*PinPath,
+				Index,
+				*PreviousCPPType.ToString(),
+				*PreviousDirection,
+				NewIndex,
+				*NewCPPType.ToString(),
+				*NewDirection
+			);
+#endif
+		}
+	}
+	for(int32 Index = 0; Index < NewPinInfos.Num(); Index++)
+	{
+		const FString PinPath = NewPinInfos.GetPinPath(Index);
+		const int32 PreviousIndex = PreviousPinInfos.GetIndexFromPinPath(PinPath);
+		if(PreviousIndex == INDEX_NONE)
+		{
+			NewPinsToAdd.Add(Index);
+
+#if UE_RIGVMCONTROLLER_VERBOSE_REPOPULATE
+			UE_LOG(LogRigVMDeveloper, Display, TEXT("Newly required pin '%s' needs to be added."), *PinPath);
+#endif
+		}
+	}
+
+	auto CreatePinFromPinInfo = [this, &Registry, &PreviousPinInfos, bNotify](const FPinInfo& InPinInfo, const FString& InPinPath, UObject* InOuter) -> URigVMPin*
+	{
+		URigVMPin* Pin = NewObject<URigVMPin>(InOuter, InPinInfo.Name);
+		if(InPinInfo.Property)
+		{
+			ConfigurePinFromProperty(InPinInfo.Property, Pin, InPinInfo.Direction);
+		}
 		else
 		{
-			// due to earlier bugs with copy and paste we can find entry and return nodes under the top level
-			// graph. we'll ignore these for now.
+			const FRigVMTemplateArgumentType& Type = Registry.GetType(InPinInfo.TypeIndex);
+			Pin->CPPType = Type.CPPType.ToString();
+			Pin->CPPTypeObject = Type.CPPTypeObject;
+			if(Pin->CPPTypeObject)
+			{
+				Pin->CPPTypeObjectPath = *Pin->CPPTypeObject->GetPathName();
+			}
+			if(Registry.IsExecuteType(InPinInfo.TypeIndex))
+			{
+				MakeExecutePin(Pin);
+			}
+				
+			Pin->Direction = InPinInfo.Direction;
+			Pin->DisplayName = InPinInfo.DisplayName.IsEmpty() ? NAME_None : FName(*InPinInfo.DisplayName);
+			Pin->bIsConstant = InPinInfo.bIsConstant;
+			Pin->bIsDynamicArray = InPinInfo.bIsDynamicArray;
+			Pin->CustomWidgetName = InPinInfo.CustomWidgetName.IsEmpty() ? NAME_None : FName(*InPinInfo.CustomWidgetName);
+		}
+
+		Pin->bIsExpanded = InPinInfo.bIsExpanded;
+		Pin->DefaultValue = InPinInfo.DefaultValue;
+
+		// reuse expansion state and default value
+		if(const FPinInfo* PreviousPin = PreviousPinInfos.GetPinFromPinPath(InPinPath))
+		{
+			if(PreviousPin->TypeIndex == InPinInfo.TypeIndex)
+			{
+				Pin->bIsExpanded = PreviousPin->bIsExpanded;
+				Pin->DefaultValue = PreviousPin->DefaultValue;
+			}
+		}
+
+		if (URigVMPin* ParentPin = Cast<URigVMPin>(InOuter))
+		{
+			AddSubPin(ParentPin, Pin);
 		}
+		else
+		{
+			AddNodePin(CastChecked<URigVMNode>(InOuter), Pin);
+		}
+
+		if(bNotify)
+		{
+			Notify(ERigVMGraphNotifType::PinAdded, Pin);
+		}
+
+		return Pin;
+	};
+
+	// step 1/3: keep a record of the current state of the node's pins
+	TMap<FString, FString> RedirectedPinPaths;
+	if (bFollowCoreRedirectors)
+	{
+		RedirectedPinPaths = GetRedirectedPinPaths(InNode);
 	}
-	else if (CollapseNode)
+
+	// also in case this node is part of an injection
+	FName InjectionInputPinName = NAME_None;
+	FName InjectionOutputPinName = NAME_None;
+	if (URigVMInjectionInfo* InjectionInfo = InNode->GetInjectionInfo())
+	{
+		InjectionInputPinName = InjectionInfo->InputPin ? InjectionInfo->InputPin->GetFName() : NAME_None;
+		InjectionOutputPinName = InjectionInfo->OutputPin ? InjectionInfo->OutputPin->GetFName() : NAME_None;
+	}
+
+	TMap<FString, FPinState> PinStates;
+	TArray<URigVMLink*> DetachedLinks;
+
+	if(bRequirePinStates)
 	{
-		// PinStates were already saved earlier, and will be applied later, no need to do it here 
+		PinStates = GetPinStates(InNode);
+	}
+
+	if(bDetachAndReattachLinks && bRequireDetachLinks)
+	{
+#if UE_RIGVMCONTROLLER_VERBOSE_REPOPULATE
+		UE_LOG(LogRigVMDeveloper, Display, TEXT("Detaching links of node %s."), *InNode->GetPathName());
+#endif
+		DetachedLinks = InNode->GetLinks();
+		DetachLinksFromPinObjects(&DetachedLinks, bNotify);
+	}
+
+	// we can do a simpler version of the update and simply
+	// add new pins, remove obsolete pins and change types as needed 
+#if UE_RIGVMCONTROLLER_VERBOSE_REPOPULATE
+	UE_LOG(LogRigVMDeveloper, Display, TEXT("Performing fast update of node %s ..."), *InNode->GetPathName());
+#endif
 
-		// since we are replacing existing pins, their names have to be saved
-		// and assigned to the new pins only after we removed the old ones
-		TArray<TTuple<URigVMPin*, FName>> NewRootPinInfos;
-		for (URigVMPin* RootPin : InNode->Pins)
+	// orphan pins
+	for(int32 Index = 0; Index < PreviousPinsToOrphan.Num(); Index++)
+	{
+		const FString& PinPath = PreviousPinInfos.GetPinPath(PreviousPinsToOrphan[Index]);
+		if(URigVMPin* Pin = InNode->FindPin(PinPath))
 		{
-			URigVMPin* NewRootPin = NewObject<URigVMPin>(InNode);
-			ConfigurePinFromPin(NewRootPin, RootPin, true);
-			EnsurePinValidity(NewRootPin, false);
-		
-			NewRootPinInfos.Add(TTuple<URigVMPin*, FName>(NewRootPin, RootPin->GetFName()));
+#if UE_RIGVMCONTROLLER_VERBOSE_REPOPULATE
+			UE_LOG(LogRigVMDeveloper, Display, TEXT("Orphaning pin '%s'."), *PinPath);
+#endif
+			check(Pin->IsRootPin());
+
+			const FString OrphanedName = FString::Printf(TEXT("%s%s"), *URigVMPin::OrphanPinPrefix, *Pin->GetName());
+			if(InNode->FindPin(OrphanedName) == nullptr)
+			{
+				Pin->DisplayName = Pin->GetFName();
+				RenameObject(Pin, *OrphanedName, nullptr);
+				InNode->Pins.Remove(Pin);
+
+				if(bNotify)
+				{
+					Notify(ERigVMGraphNotifType::PinRemoved, Pin);
+				}
+
+				InNode->OrphanedPins.Add(Pin);
+
+				if(bNotify)
+				{
+					Notify(ERigVMGraphNotifType::PinAdded, Pin);
+				}
+			}
+			else
+			{
+				RemovePin(Pin, false, bNotify);
+			}
+		}
+	}
+
+	// remove obsolete pins
+	for(int32 Index = PreviousPinsToRemove.Num() - 1; Index >= 0; Index--)
+	{
+		const FString& PinPath = PreviousPinInfos.GetPinPath(PreviousPinsToRemove[Index]);
+		if(URigVMPin* Pin = InNode->FindPin(PinPath))
+		{
+#if UE_RIGVMCONTROLLER_VERBOSE_REPOPULATE
+			UE_LOG(LogRigVMDeveloper, Display, TEXT("Removing pin '%s'."), *PinPath);
+#endif
+			RemovePin(Pin, false, bNotify);
+		}
+	}
+
+	// add missing pins
+	for(int32 Index = 0; Index < NewPinsToAdd.Num(); Index++)
+	{
+		const FString& PinPath = NewPinInfos.GetPinPath(NewPinsToAdd[Index]);
+		FString ParentPinPath, PinName;
+		UObject* OuterForPin = InNode;
+		if(URigVMPin::SplitPinPathAtEnd(PinPath, ParentPinPath, PinName))
+		{
+			OuterForPin = InNode->FindPin(ParentPinPath);
 		}
 		
-		RemovePinsDuringRepopulate(InNode, InNode->Pins, bSetupOrphanedPins);
+		(void)CreatePinFromPinInfo(NewPinInfos[NewPinsToAdd[Index]], PinPath, OuterForPin);
+#if UE_RIGVMCONTROLLER_VERBOSE_REPOPULATE
+		UE_LOG(LogRigVMDeveloper, Display, TEXT("Adding new pin '%s'."), *PinPath);
+#endif
+	}
+
+	// update existing pins
+	for(int32 Index = 0; Index < PreviousPinsToUpdate.Num(); Index++)
+	{
+		const FString& PinPath = PreviousPinInfos.GetPinPath(PreviousPinsToUpdate[Index]);
+		const FPinInfo* NewPinInfo = NewPinInfos.GetPinFromPinPath(PinPath);
+		check(NewPinInfo);
 		
-		for (TTuple<URigVMPin*, FName> NewRootPinInfo : NewRootPinInfos)
+		if(URigVMPin* Pin = InNode->FindPin(PinPath))
 		{
-			RenameObject(NewRootPinInfo.Key, *NewRootPinInfo.Value.ToString(), InNode);
-			AddNodePin(InNode, NewRootPinInfo.Key);
+			if(Pin->IsExecuteContext())
+			{
+				MakeExecutePin(Pin);
+			}
+			
+			if(Pin->GetTypeIndex() != NewPinInfo->TypeIndex)
+			{
+				// we expect these changes to only apply to float and double pins.
+				check((NewPinInfo->TypeIndex == RigVMTypeUtils::TypeIndex::Float) ||
+					(NewPinInfo->TypeIndex == RigVMTypeUtils::TypeIndex::FloatArray) ||
+					(NewPinInfo->TypeIndex == RigVMTypeUtils::TypeIndex::Double) ||
+					(NewPinInfo->TypeIndex == RigVMTypeUtils::TypeIndex::DoubleArray));
+
+				const FRigVMTemplateArgumentType& NewType = Registry.GetType(NewPinInfo->TypeIndex);
+#if UE_RIGVMCONTROLLER_VERBOSE_REPOPULATE
+				UE_LOG(LogRigVMDeveloper, Display, TEXT("Changing pin '%s' type from %s to %s."),
+					*PinPath,
+					*Pin->CPPType,
+					*NewType.CPPType.ToString()
+				);
+#endif
+				Pin->CPPType = NewType.CPPType.ToString();
+				Pin->CPPTypeObject = NewType.CPPTypeObject;
+				check(Pin->CPPTypeObject == nullptr);
+				Pin->CPPTypeObjectPath = NAME_None;
+				Pin->LastKnownTypeIndex = NewPinInfo->TypeIndex;
+
+				if(bNotify)
+				{
+					Notify(ERigVMGraphNotifType::PinTypeChanged, Pin);
+				}
+			}
 		}
-		
-		for (URigVMPin* Pin : InNode->Pins)
+	}
+
+	// create a map representing the order of expected pins
+	TMap<FString, TArray<FName>> PinOrder;
+	for(int32 Index = 0; Index < NewPinInfos.Num(); Index++)
+	{
+		const FPinInfo& NewPin = NewPinInfos[Index];
+		FString ParentPinPath;
+		if(NewPin.ParentIndex != INDEX_NONE)
 		{
-			if (Pin->IsStruct())
+			ParentPinPath = NewPinInfos.GetPinPath(NewPin.ParentIndex);
+			if(NewPinInfos[NewPin.ParentIndex].bIsArray)
 			{
-				TGuardValue<bool> SuspendNotifications(bSuspendNotifications, true);
-				AddPinsForStruct(Pin->GetScriptStruct(), InNode, Pin, Pin->GetDirection(), FString(), false);
+				continue;
 			}
-			Notify(ERigVMGraphNotifType::PinAdded, Pin);
 		}
 
-		if (CollapseNode->GetOuter()->IsA<URigVMFunctionLibrary>())
+		TArray<FName>& SubPinOrder = PinOrder.FindOrAdd(ParentPinPath);
+		SubPinOrder.Add(NewPin.Name);
+	}
+
+	auto SortPinArray = [this, bNotify](TArray<TObjectPtr<URigVMPin>>& Pins, const TArray<FName>* PinOrder)
+	{
+		if(PinOrder == nullptr)
 		{
-			// no need to notify since the function library graph is invisible anyway
-			TGuardValue<bool> SuspendNotifications(bSuspendNotifications, true);
-			RemoveUnusedOrphanedPins(CollapseNode);
+			return;
 		}
 
-		FRigVMControllerGraphGuard GraphGuard(this, CollapseNode->GetContainedGraph(), false);
-		TGuardValue<bool> GuardEditGraph(CollapseNode->ContainedGraph->bEditable, true);
-		// need to get a copy of the node array since the following function could remove nodes from the graph
-		// we don't want to remove elements from the array we are iterating over.
-		TArray<URigVMNode*> ContainedNodes = CollapseNode->GetContainedNodes();
-		for (URigVMNode* ContainedNode : ContainedNodes)
+		if(Pins.Num() < 2)
+		{
+			return;
+		}
+
+		const TArray<TObjectPtr<URigVMPin>> PreviousPins = Pins;
+
+		if(Pins[0]->IsArrayElement())
+		{
+			Algo::Sort(Pins, [PinOrder](const TObjectPtr<URigVMPin>& A, const TObjectPtr<URigVMPin>& B) -> bool
+			{
+				return A->GetFName().Compare(B->GetFName()) < 0;
+			});
+		}
+		else
 		{
-			RepopulatePinsOnNode(ContainedNode, bFollowCoreRedirectors, bSetupOrphanedPins);
+			Algo::Sort(Pins, [PinOrder](const TObjectPtr<URigVMPin>& A, const TObjectPtr<URigVMPin>& B) -> bool
+			{
+				const int32 IndexA = PinOrder->Find(A->GetFName());
+				const int32 IndexB = PinOrder->Find(B->GetFName());
+				return IndexA < IndexB;
+			});
+		}
+
+
+#if !UE_RIGVMCONTROLLER_VERBOSE_REPOPULATE
+		if(bNotify)
+#endif
+		{
+			for(int32 Index=0;Index<Pins.Num();Index++)
+			{
+				if(PreviousPins[Index] != Pins[Index])
+				{
+#if UE_RIGVMCONTROLLER_VERBOSE_REPOPULATE
+					if(bNotify)
+					{
+#endif
+						Notify(ERigVMGraphNotifType::PinIndexChanged, Pins[Index]);
+#if UE_RIGVMCONTROLLER_VERBOSE_REPOPULATE
+					}
+					UE_LOG(LogRigVMDeveloper, Display, TEXT("Pin '%s' changed index from %d to %d."),
+						*Pins[Index]->GetPinPath(),
+						PreviousPins.Find(Pins[Index]),
+						Index
+					);
+#endif
+				}
+			}
 		}
+	};
+
+	SortPinArray(InNode->Pins, PinOrder.Find(FString()));
+	for(URigVMPin* Pin : InNode->Pins)
+	{
+		SortPinArray(Pin->SubPins, PinOrder.Find(Pin->GetPinPath()));
+	}
+	
+	if(DispatchNode)
+	{
+		ResolveTemplateNodeMetaData(DispatchNode, false);
 	}
-	else if (FunctionRefNode)
+	else if(CollapseNode)
 	{
-		FRigVMGraphFunctionHeader& Header = FunctionRefNode->ReferencedFunctionHeader;
-		
-		// we want to make sure notify the graph of a potential name change
-		// when repopulating the function ref node
-		Notify(ERigVMGraphNotifType::NodeRenamed, FunctionRefNode);
-		RemovePinsDuringRepopulate(InNode, InNode->Pins, bSetupOrphanedPins);
-
-		TMap<FString, FPinState> ReferencedPinStates;
-		if (Header.LibraryPointer.LibraryNode.ResolveObject())
-		{
-			URigVMLibraryNode* LibraryNode = Cast<URigVMLibraryNode>(Header.LibraryPointer.LibraryNode.ResolveObject());
-			ReferencedPinStates = GetPinStates(LibraryNode);
-		}
-
-		for (const FRigVMGraphFunctionArgument& ReferencedArgument : Header.Arguments)
+		if (CollapseNode->GetOuter()->IsA<URigVMFunctionLibrary>())
 		{
-			URigVMPin* NewPin = NewObject<URigVMPin>(InNode, ReferencedArgument.Name);
-			ConfigurePinFromArgument(NewPin, ReferencedArgument, true);
-			EnsurePinValidity(NewPin, false);
-			
-			AddNodePin(InNode, NewPin);
-
-			if (NewPin->IsStruct())
-			{
-				TGuardValue<bool> SuspendNotifications(bSuspendNotifications, true);
-				AddPinsForStruct(NewPin->GetScriptStruct(), InNode, NewPin, NewPin->GetDirection(), FString(), false);
-			}
-
-			Notify(ERigVMGraphNotifType::PinAdded, NewPin);
+			// no need to notify since the function library graph is invisible anyway
+			RemoveUnusedOrphanedPins(CollapseNode, false);
 		}
 
-		if (ReferencedPinStates.Num() > 0)
+		RecursivelyRepopulatePinsOnCollapseNode(CollapseNode);
+	}
+	else if (FunctionRefNode)
+	{
+		if (URigVMLibraryNode* ReferencedNode = FunctionRefNode->GetReferencedNode())
 		{
-			ApplyPinStates(InNode, ReferencedPinStates);
+			// we want to make sure notify the graph of a potential name change
+			// when repopulating the function ref node
+			Notify(ERigVMGraphNotifType::NodeRenamed, FunctionRefNode);
 		}
 	}
-	else
+
+	if(!PinStates.IsEmpty())
 	{
-		return;
+#if UE_RIGVMCONTROLLER_VERBOSE_REPOPULATE
+		UE_LOG(LogRigVMDeveloper, Display, TEXT("Reapplying pin-states of node %s..."), *InNode->GetPathName());
+#endif
+		ApplyPinStates(InNode, PinStates, RedirectedPinPaths);
 	}
 
-	ApplyPinStates(InNode, PinStates, RedirectedPinPaths);
+	if(!DetachedLinks.IsEmpty())
+	{
+#if UE_RIGVMCONTROLLER_VERBOSE_REPOPULATE
+		UE_LOG(LogRigVMDeveloper, Display, TEXT("Reattaching links of node %s..."), *InNode->GetPathName());
+#endif
+		ReattachLinksToPinObjects(bFollowCoreRedirectors, &DetachedLinks, bNotify, bSetupOrphanPinsForThisNode);
+	}
 
 	if (URigVMInjectionInfo* InjectionInfo = InNode->GetInjectionInfo())
 	{
 		InjectionInfo->InputPin = InNode->FindPin(InjectionInputPinName.ToString());
 		InjectionInfo->OutputPin = InNode->FindPin(InjectionOutputPinName.ToString());
 	}
+
+#if UE_RIGVMCONTROLLER_VERBOSE_REPOPULATE
+	UE_LOG(LogRigVMDeveloper, Display, TEXT("Repopulate of node %s is completed.\n"), *InNode->GetPathName());
+#endif
 }
 
-void URigVMController::RemovePinsDuringRepopulate(URigVMNode* InNode, TArray<URigVMPin*>& InPins, bool bSetupOrphanedPins)
+void URigVMController::RemovePinsDuringRepopulate(URigVMNode* InNode, TArray<URigVMPin*>& InPins, bool bNotify, bool bSetupOrphanedPins)
 {
 	TArray<URigVMPin*> Pins = InPins;
 	for (URigVMPin* Pin : Pins)
@@ -15302,14 +15894,14 @@ void URigVMController::RemovePinsDuringRepopulate(URigVMNode* InNode, TArray<URi
 					RenameObject(RootPin, *OrphanedName, nullptr);
 					InNode->Pins.Remove(RootPin);
 
-					if(!bSuspendNotifications)
+					if(bNotify)
 					{
 						Notify(ERigVMGraphNotifType::PinRemoved, RootPin);
 					}
 
 					InNode->OrphanedPins.Add(RootPin);
 
-					if(!bSuspendNotifications)
+					if(bNotify)
 					{
 						Notify(ERigVMGraphNotifType::PinAdded, RootPin);
 					}
@@ -15322,7 +15914,7 @@ void URigVMController::RemovePinsDuringRepopulate(URigVMNode* InNode, TArray<URi
 				
 					OrphanedRootPin->GetNode()->OrphanedPins.Add(OrphanedRootPin);
 
-					if(!bSuspendNotifications)
+					if(bNotify)
 					{
 						Notify(ERigVMGraphNotifType::PinAdded, OrphanedRootPin);
 					}
@@ -15343,13 +15935,13 @@ void URigVMController::RemovePinsDuringRepopulate(URigVMNode* InNode, TArray<URi
 	{
 		if(!Pin->IsOrphanPin())
 		{
-			RemovePin(Pin, false);
+			RemovePin(Pin, false, bNotify);
 		}
 	}
 	InPins.Reset();
 }
 
-bool URigVMController::RemoveUnusedOrphanedPins(URigVMNode* InNode)
+bool URigVMController::RemoveUnusedOrphanedPins(URigVMNode* InNode, bool bNotify)
 {
 	if(!InNode->HasOrphanedPins())
 	{
@@ -15366,7 +15958,7 @@ bool URigVMController::RemoveUnusedOrphanedPins(URigVMNode* InNode)
 
 		if(NumSourceLinks + NumTargetLinks == 0)
 		{
-			RemovePin(OrphanedPin, false);
+			RemovePin(OrphanedPin, false, bNotify);
 		}
 		else
 		{
@@ -15560,14 +16152,8 @@ void URigVMController::ApplyPinState(URigVMPin* InPin, const FPinState& InPinSta
 	{
 		for (const URigVMInjectionInfo::FWeakInfo& InjectionInfo : InPinState.WeakInjectionInfos)
 		{
-			if(URigVMNode* FormerlyInjectedNode = InjectionInfo.Node.Get())
+			if(const URigVMNode* FormerlyInjectedNode = InjectionInfo.Node.Get())
 			{
-				if (FormerlyInjectedNode->IsInjected())
-				{
-					URigVMInjectionInfo* Injection = Cast<URigVMInjectionInfo>(FormerlyInjectedNode->GetOuter());
-					RenameObject(FormerlyInjectedNode, nullptr, InPin->GetGraph());
-					DestroyObject(Injection);
-				}
 				if(InjectionInfo.bInjectedAsInput)
 				{
 					const FString OutputPinPath = URigVMPin::JoinPinPath(FormerlyInjectedNode->GetNodePath(), InjectionInfo.OutputPinName.ToString()); 
@@ -15708,19 +16294,6 @@ void URigVMController::ReportAndNotifyError(const FString& InMessage) const
 	SendUserFacingNotification(InMessage, 0.f, nullptr, TEXT("MessageLog.Error"));
 }
 
-void URigVMController::ReportPinTypeChange(URigVMPin* InPin, const FString& InNewCPPType)
-{
-	/*
-	UE_LOG(LogRigVMDeveloper,
-		Warning,
-		TEXT("Pin '%s' is about to change type from '%s' to '%s'."),
-		*InPin->GetPinPath(),
-		*InPin->GetCPPType(),
-		*InNewCPPType
-	);
-	*/
-}
-
 void URigVMController::SendUserFacingNotification(const FString& InMessage, float InDuration, const UObject* InSubject, const FName& InBrushName) const
 {
 #if WITH_EDITOR
@@ -15909,7 +16482,7 @@ void URigVMController::ResolveTemplateNodeMetaData(URigVMTemplateNode* InNode, b
 				const FString NewDefaultValue = InNode->GetInitialDefaultValueForPin(Pin->GetFName(), FilteredPermutationIndices);
 				if(!NewDefaultValue.IsEmpty())
 				{
-					SetPinDefaultValue(Pin, NewDefaultValue, true, bSetupUndoRedo, false);
+					SetPinDefaultValue(Pin, NewDefaultValue, true, bSetupUndoRedo, false, true);
 				}
 			}
 		}
@@ -16044,27 +16617,14 @@ bool URigVMController::FullyResolveTemplateNode(URigVMTemplateNode* InNode, int3
 	{
 		if(ResolvedFunction->Struct == nullptr)
 		{
-			FRigVMDispatchContext DispatchContext;
-			if(const URigVMDispatchNode* DispatchNode = Cast<URigVMDispatchNode>(InNode))
-			{
-				DispatchContext = DispatchNode->GetDispatchContext();
-			}
-
 			const TArray<FRigVMTemplateArgument>& Arguments = Template->Arguments;
 			for(URigVMPin* Pin : InNode->GetPins())
 			{
-				bool bFound = Arguments.ContainsByPredicate([Pin](const FRigVMTemplateArgument& Argument)
+				const bool bFound = Arguments.ContainsByPredicate([Pin](const FRigVMTemplateArgument& Argument)
 				{
 					return Pin->GetFName() == Argument.GetName();
 				});
 				if(!bFound)
-				{
-					bFound = Template->GetExecuteArguments(DispatchContext).ContainsByPredicate([Pin](const FRigVMExecuteArgument& Argument)
-					{
-						return Pin->GetFName() == Argument.Name;
-					});
-				}
-				if(!bFound)
 				{
 					PinsToRemove.Add(Pin);
 				}
@@ -16158,7 +16718,7 @@ bool URigVMController::FullyResolveTemplateNode(URigVMTemplateNode* InNode, int3
 	// remove obsolete pins
 	for(URigVMPin* PinToRemove : PinsToRemove)
 	{
-		RemovePin(PinToRemove, false);
+		RemovePin(PinToRemove, false, true);
 	}
 
 	// add missing pins
@@ -16183,8 +16743,8 @@ bool URigVMController::FullyResolveTemplateNode(URigVMTemplateNode* InNode, int3
 
 				if(Factory)
 				{
-					const FName DisplayNameText = Factory->GetDisplayNameForArgument(MissingPin.GetName());
-					if(!DisplayNameText.IsNone())
+					const FText DisplayNameText = Factory->GetDisplayNameForArgument(MissingPin.GetName());
+					if(!DisplayNameText.IsEmpty())
 					{
 						Pin->DisplayName = *DisplayNameText.ToString();
 					}
@@ -16262,19 +16822,9 @@ bool URigVMController::PrepareTemplatePinForType(URigVMPin* InPin, const TArray<
 
 	FRigVMRegistry& Registry = FRigVMRegistry::Get();
 
-	/*
-	if (InTypeIndices.Num() == 1)
-	{
-		if(Registry.CanMatchTypes(InPin->GetTypeIndex(), InTypeIndices[0], true))
-		{
-			return true;
-		}
-	}
-	*/
-
 	// It might be a subpin of a struct, in that case we just want to make sure
 	// the type required matches the one in the pin
-	if (InPin->IsStructMember() || InPin->IsOrphanPin())
+	if (InPin->IsStructMember())
 	{
 		TRigVMTypeIndex PinTypeIndex = InPin->GetTypeIndex();
 		
@@ -16403,7 +16953,7 @@ bool URigVMController::PrepareTemplatePinForType(URigVMPin* InPin, const TArray<
 		for (int32 i=0; i<InconsistentLinks.Num(); ++i)
 		{
 			TSoftObjectPtr<URigVMGraph> GraphPtr(InconsistentLinksGraph[i]);
-			LinksToBreak.AddUnique(GraphPtr.Get()->FindLink(URigVMLink::GetPinPathRepresentation(InconsistentLinks[i].Key, InconsistentLinks[i].Value)));
+			LinksToBreak.AddUnique(GraphPtr.Get()->FindLink(FString::Printf(TEXT("%s -> %s"), *InconsistentLinks[i].Key, *InconsistentLinks[i].Value)));
 		}
 		
 		// Warn the user the links that would be broken
@@ -16452,7 +17002,6 @@ bool URigVMController::PrepareTemplatePinForType(URigVMPin* InPin, const TArray<
 			TGuardValue<FRigVMController_RequestBreakLinksDialogDelegate> GuardDelegate(RequestBreakLinksDialogDelegate, nullptr);
 			TGuardValue<bool> GuardNotifications(bSuspendNotifications, true);
 			TGuardValue<bool> SuspendRecomputeTemplates(bSuspendRecomputingTemplateFilters, true);
-			TGuardValue<bool> DisableCasting(bEnableTypeCasting, false);
 			OpenUndoBracket(FString::Printf(TEXT("Resolve wildcard pin %s"), *InPin->GetPinPath()));
 
 			bool bBrokenLinks = false;
@@ -16493,7 +17042,7 @@ bool URigVMController::PrepareTemplatePinForType(URigVMPin* InPin, const TArray<
 
 			for (TPair<FString, FString>& Pair : InconsistentLinks)
 			{
-				LinksToBreak.AddUnique(GetGraph()->FindLink(URigVMLink::GetPinPathRepresentation(Pair.Key, Pair.Value)));
+				LinksToBreak.AddUnique(GetGraph()->FindLink(FString::Printf(TEXT("%s -> %s"), *Pair.Key, *Pair.Value)));
 			}
 		}
 
@@ -16600,7 +17149,7 @@ bool URigVMController::PrepareTemplatePinForType(URigVMPin* InPin, const TArray<
 
 TArray<TRigVMTypeIndex> URigVMController::GetFilteredTypes(URigVMPin* InPin)
 {
-	if (!InPin->IsStructMember() && !InPin->IsOrphanPin())
+	if (!InPin->IsStructMember())
 	{
 		if (URigVMTemplateNode* TemplateNode = Cast<URigVMTemplateNode>(InPin->GetNode()))
 		{
@@ -16635,7 +17184,7 @@ bool URigVMController::ResolveWildCardPin(const FString& InPinPath, const FStrin
 	UObject* CPPTypeObject = nullptr;
 	if (!InCPPTypeObjectPath.IsNone())
 	{
-		CPPTypeObject = RigVMTypeUtils::FindObjectFromCPPTypeObjectPath<UObject>(InCPPTypeObjectPath.ToString());
+		CPPTypeObject = URigVMPin::FindObjectFromCPPTypeObjectPath<UObject>(InCPPTypeObjectPath.ToString());
 		if (CPPTypeObject == nullptr)
 		{
 			ReportErrorf(TEXT("Cannot find cpp type object for path '%s'."), *InCPPTypeObjectPath.ToString());
@@ -16715,11 +17264,6 @@ bool URigVMController::ResolveWildCardPin(URigVMPin* InPin, TRigVMTypeIndex InTy
 		return false;
 	}
 
-	if(FRigVMRegistry::Get().IsWildCardType(InTypeIndex))
-	{
-		return false;
-	}
-
 	URigVMPin* RootPin = InPin;
 	TRigVMTypeIndex Type = InTypeIndex;	
 	while (RootPin->IsArrayElement())
@@ -17132,15 +17676,10 @@ bool URigVMController::UpdateTemplateNodePinTypes(URigVMTemplateNode* InNode, bo
 				int32 PinIndex = Pins.Find(ArgumentPin);
 				WasReduced[PinIndex] = true;
 				FinalPinTypes[PinIndex] = Preferred.TypeIndex;
-
-				// If the argument is hidden, the available PinTypes will all be INDEX_NONE
-				if (ArgumentPin->Direction != ERigVMPinDirection::Hidden)
-				{					
-					ReducedTypes = ReducedTypes.FilterByPredicate([Preferred, PinIndex](const TPair<int32, TArray<TRigVMTypeIndex>>& Permutation)
-					{
-						return Permutation.Value[PinIndex] == Preferred.TypeIndex;
-					});
-				}
+				ReducedTypes = ReducedTypes.FilterByPredicate([Preferred, PinIndex](const TPair<int32, TArray<TRigVMTypeIndex>>& Permutation)
+				{
+					return Permutation.Value[PinIndex] == Preferred.TypeIndex;
+				});
 			}
 		}
 	}
@@ -17213,7 +17752,7 @@ bool URigVMController::UpdateTemplateNodePinTypes(URigVMTemplateNode* InNode, bo
 			{
 				bShouldUnresolve = true;
 			}
-
+				
 			if (bShouldUnresolve)
 			{
 				// Unresolve
@@ -17244,21 +17783,10 @@ bool URigVMController::UpdateTemplateNodePinTypes(URigVMTemplateNode* InNode, bo
 				{				
 					CPPType = Pin->IsArray() ? RigVMTypeUtils::GetWildCardArrayCPPType() : RigVMTypeUtils::GetWildCardCPPType();
 				}
-
-				// execute pins are no longer part of the template, avoid changing the type
-				if(Argument == nullptr && Pin->IsExecuteContext())
-				{
-					CPPType = Pin->GetCPPType();
-					CPPObjectType = Pin->GetCPPTypeObject();
-				}
 			
 				if (Pin->GetCPPType() != CPPType || Pin->GetCPPTypeObject() != CPPObjectType)
 				{
 					bAnyTypeChanged = !Pin->IsExecuteContext();
-					if(bAnyTypeChanged)
-					{
-						ReportPinTypeChange(Pin, CPPType);
-					}
 					ChangePinType(Pin, CPPType, CPPObjectType, bSetupUndoRedo, false, false, false, bInitializeDefaultValue);
 				}
 			}
@@ -17268,11 +17796,6 @@ bool URigVMController::UpdateTemplateNodePinTypes(URigVMTemplateNode* InNode, bo
 				if (Pin->GetTypeIndex() != FinalPinTypes[PinIndex])
 				{
 					bAnyTypeChanged = !Pin->IsExecuteContext();
-					if(bAnyTypeChanged)
-					{
-						const FString CPPType = FRigVMRegistry::Get().GetType(FinalPinTypes[PinIndex]).CPPType.ToString();
-						ReportPinTypeChange(Pin, CPPType);
-					}
 					ChangePinType(Pin, FinalPinTypes[PinIndex], bSetupUndoRedo, false, false, false, bInitializeDefaultValue);
 				}
 
@@ -17377,10 +17900,10 @@ bool URigVMController::PropagateTemplateFilteredTypes(URigVMTemplateNode* InNode
 {
 	auto UpdateAndPropagate = [&](URigVMPin* Pin)
 	{
-		const TArray<URigVMLink*>& Links = Pin->GetLinks();
-		for (URigVMLink* Link : Links)
+		TArray<URigVMPin*> OtherPins = Pin->GetLinkedSourcePins();
+		OtherPins.Append(Pin->GetLinkedTargetPins());
+		for (URigVMPin* OtherPin : OtherPins)
 		{
-			URigVMPin* OtherPin = Link->GetSourcePin() == Pin ? Link->GetTargetPin() : Link->GetSourcePin();
 			bool bPropagate = false;
 			bool bIsTemplate = false;
 			if (URigVMTemplateNode* OtherTemplate = Cast<URigVMTemplateNode>(OtherPin->GetNode()))
@@ -17412,6 +17935,7 @@ bool URigVMController::PropagateTemplateFilteredTypes(URigVMTemplateNode* InNode
 						}
 						else
 						{
+							URigVMLink* Link = Pin->FindLinkForPin(OtherPin);
 							ensureMsgf(!ActionStack->BracketActions.IsEmpty(), TEXT("Unexpected link broken %s in package %s"), *Link->GetPinPathRepresentation(), *GetPackage()->GetPathName());
 							BreakLink(Link->GetSourcePin(), Link->GetTargetPin(), bSetupUndoRedo);
 							return false;
@@ -17424,6 +17948,7 @@ bool URigVMController::PropagateTemplateFilteredTypes(URigVMTemplateNode* InNode
 			{
 				if (!InNode->FilteredSupportsType(Pin, OtherPin->GetTypeIndex()))
 				{
+					URigVMLink* Link = Pin->FindLinkForPin(OtherPin);
 					ensureMsgf(!ActionStack->BracketActions.IsEmpty(), TEXT("Unexpected link broken %s in package %s"), *Link->GetPinPathRepresentation(), *GetPackage()->GetPathName());
 					BreakLink(Link->GetSourcePin(), Link->GetTargetPin(), bSetupUndoRedo);
 					return false;
@@ -18020,21 +18545,10 @@ bool URigVMController::RecomputeAllTemplateFilteredPermutations(bool bSetupUndoR
 		InitializeAllTemplateFiltersInGraph(bSetupUndoRedo, false);	
 
 		// Apply all links to update filtered permutations
-		TArray<URigVMLink*> SortedLinks = Graph->GetLinks().FilterByPredicate([this](URigVMLink* Link)
+		TArray<URigVMLink*> SortedLinks = Graph->GetLinks().FilterByPredicate([](URigVMLink* Link)
 		{
 			// Filter out detached links
-			if(!Link->GetSourcePin()->IsLinkedTo(Link->GetTargetPin()))
-			{
-				return false;
-			}
-
-			// Filter out links on execute pins
-			if(Link->GetSourcePin()->IsExecuteContext() || Link->GetTargetPin()->IsExecuteContext())
-			{
-				return false;
-			}
-			
-			return true;
+			return Link->GetSourcePin()->IsLinkedTo(Link->GetTargetPin());
 		});
 
 		// Solve for unit nodes first, other templates are more expensive to solve. This way we are avoiding solving
@@ -18087,10 +18601,7 @@ bool URigVMController::RecomputeAllTemplateFilteredPermutations(bool bSetupUndoR
 			}
 		});
 			
-		{
-			TGuardValue<bool> SuspendNotifications(bSuspendNotifications, true);
-			DetachLinksFromPinObjects(&SortedLinks);
-		}
+		DetachLinksFromPinObjects(&SortedLinks);
 		TArray<URigVMLink*> LinksToRemove;
 		for (int32 i=0; i<SortedLinks.Num(); ++i)
 		{
@@ -18516,7 +19027,7 @@ bool URigVMController::ChangePinType(URigVMPin* InPin, const FString& InCPPType,
 		return false;
 	}
 
-	UObject* CPPTypeObject = RigVMTypeUtils::FindObjectFromCPPTypeObjectPath(InCPPTypeObjectPath.ToString());
+	UObject* CPPTypeObject = URigVMPin::FindObjectFromCPPTypeObjectPath(InCPPTypeObjectPath.ToString());
 
 	// always refresh pin type if it is a user defined struct, whose internal layout can change at anytime
 	bool bForceRefresh = false;
@@ -18548,7 +19059,7 @@ bool URigVMController::ChangePinType(URigVMPin* InPin, const FString& InCPPType,
 		return false;
 	}
 
-	if (RigVMTypeUtils::RequiresCPPTypeObject(InCPPType) && !InCPPTypeObject)
+	if (FRigVMPropertyDescription::RequiresCPPTypeObject(InCPPType) && !InCPPTypeObject)
 	{
 		return false;
 	}
@@ -18637,7 +19148,7 @@ bool URigVMController::ChangePinType(URigVMPin* InPin, TRigVMTypeIndex InTypeInd
 	{
 		Links.Append(InPin->GetSourceLinks(true));
 		Links.Append(InPin->GetTargetLinks(true));
-		DetachLinksFromPinObjects(&Links);
+		DetachLinksFromPinObjects(&Links, true);
 
 		const FString OrphanedName = FString::Printf(TEXT("%s%s"), *URigVMPin::OrphanPinPrefix, *InPin->GetName());
 		if(InPin->GetNode()->FindPin(OrphanedName) == nullptr)
@@ -18648,7 +19159,7 @@ bool URigVMController::ChangePinType(URigVMPin* InPin, TRigVMTypeIndex InTypeInd
 
 			if(OrphanedPin->IsStruct())
 			{
-				AddPinsForStruct(OrphanedPin->GetScriptStruct(), OrphanedPin->GetNode(), OrphanedPin, OrphanedPin->Direction, OrphanedPin->GetDefaultValue(), false);
+				AddPinsForStruct(OrphanedPin->GetScriptStruct(), OrphanedPin->GetNode(), OrphanedPin, OrphanedPin->Direction, OrphanedPin->GetDefaultValue(), false, true);
 			}
 				
 			InPin->GetNode()->OrphanedPins.Add(OrphanedPin);
@@ -18660,7 +19171,7 @@ bool URigVMController::ChangePinType(URigVMPin* InPin, TRigVMTypeIndex InTypeInd
 		TArray<URigVMPin*> Pins = InPin->SubPins;
 		for (URigVMPin* Pin : Pins)
 		{
-			RemovePin(Pin, bSetupUndoRedo);
+			RemovePin(Pin, bSetupUndoRedo, true);
 		}
 		
 		InPin->SubPins.Reset();
@@ -18708,7 +19219,7 @@ bool URigVMController::ChangePinType(URigVMPin* InPin, TRigVMTypeIndex InTypeInd
 	{
 		FString DefaultValue = InPin->DefaultValue;
 		CreateDefaultValueForStructIfRequired(InPin->GetScriptStruct(), DefaultValue);
-		AddPinsForStruct(InPin->GetScriptStruct(), InPin->GetNode(), InPin, InPin->Direction, DefaultValue, false);
+		AddPinsForStruct(InPin->GetScriptStruct(), InPin->GetNode(), InPin, InPin->Direction, DefaultValue, false, true);
 	}
 
 	if (InPin->IsArray())
@@ -18838,6 +19349,7 @@ bool URigVMController::ChangePinType(URigVMPin* InPin, TRigVMTypeIndex InTypeInd
 				{
 					const URigVMPin* ParentPin = Pin->GetParentPin();
 					TSharedPtr<FStructOnScope> StructOnScope = MakeShareable(new FStructOnScope(ParentPin->GetScriptStruct()));
+					ParentPin->GetScriptStruct()->InitializeDefaultValue((uint8*)StructOnScope->GetStructMemory());
 					const FString StructDefaultValue = FRigVMStruct::ExportToFullyQualifiedText(ParentPin->GetScriptStruct(), StructOnScope->GetStructMemory());
 					Local::ApplyResolvedDefaultValue(this, Pin, Pin->GetName(), StructDefaultValue, bSetupUndoRedo);
 				}
@@ -18899,8 +19411,8 @@ bool URigVMController::ChangePinType(URigVMPin* InPin, TRigVMTypeIndex InTypeInd
 
 	if(Links.Num() > 0)
 	{
-		ReattachLinksToPinObjects(false, &Links, true);
-		RemoveUnusedOrphanedPins(InPin->GetNode());
+		ReattachLinksToPinObjects(false, &Links, true, true);
+		RemoveUnusedOrphanedPins(InPin->GetNode(), true);
 	}
 
 	return true;
@@ -19010,27 +19522,90 @@ void URigVMController::AddSubPin(URigVMPin* InParentPin, URigVMPin* InPin)
 	InParentPin->SubPins.Add(InPin);
 }
 
+
+static UObject* FindObjectGloballyWithRedirectors(const TCHAR* InObjectName)
+{
+	// Do a global search for the CPP type. Note that searching with ANY_PACKAGE _does not_
+	// apply redirectors. So only if this fails do we apply them manually below.
+	UObject* Object = FindFirstObject<UField>(InObjectName, EFindFirstObjectOptions::EnsureIfAmbiguous);
+	if(Object != nullptr)
+	{
+		return Object;
+	}
+
+	FCoreRedirectObjectName NewObjectName;
+	const bool bFoundRedirect = FCoreRedirects::RedirectNameAndValues(
+		ECoreRedirectFlags::Type_Class | ECoreRedirectFlags::Type_Struct | ECoreRedirectFlags::Type_Enum,
+		FCoreRedirectObjectName(InObjectName),
+		NewObjectName,
+		nullptr,
+		ECoreRedirectMatchFlags::None);
+
+	if (!bFoundRedirect)
+	{
+		return nullptr;
+	}
+
+	const FString RedirectedObjectName = NewObjectName.ObjectName.ToString();
+	UPackage *Package = nullptr;
+	if (!NewObjectName.PackageName.IsNone())
+	{
+		Package = FindPackage(nullptr, *NewObjectName.PackageName.ToString());
+	}
+	if (Package != nullptr)
+	{
+		Object = FindObject<UField>(Package, *RedirectedObjectName);
+	}
+	if (Package == nullptr || Object == nullptr)
+	{
+		// Hail Mary pass.
+		Object = FindFirstObject<UField>(*RedirectedObjectName, EFindFirstObjectOptions::EnsureIfAmbiguous);
+	}
+	return Object;
+}
+
 bool URigVMController::EnsurePinValidity(URigVMPin* InPin, bool bRecursive)
 {
 	check(InPin);
 	
 	// check if the CPPTypeObject is set up correctly.
-	if(RigVMTypeUtils::RequiresCPPTypeObject(InPin->GetCPPType()))
+	if(FRigVMPropertyDescription::RequiresCPPTypeObject(InPin->GetCPPType()))
 	{
 		// GetCPPTypeObject attempts to update pin type information to the latest
 		// without testing for redirector
 		if(InPin->GetCPPTypeObject() == nullptr)
 		{
-			FString CPPType = InPin->GetCPPType();
-			InPin->CPPTypeObject = RigVMTypeUtils::ObjectFromCPPType(CPPType);
-			InPin->CPPType = CPPType;
-		}
-		else
-		{
-			InPin->CPPType = RigVMTypeUtils::PostProcessCPPType(InPin->CPPType, InPin->GetCPPTypeObject());
+			// try to find the CPPTypeObject by name
+			FString CPPType = InPin->IsArray() ? InPin->GetArrayElementCppType() : InPin->GetCPPType();
+
+			UObject* CPPTypeObject = FindObjectGloballyWithRedirectors(*CPPType);
+
+			if (CPPTypeObject == nullptr)
+			{
+				// If we've mistakenly stored the struct type with the 'F', 'U', or 'A' prefixes, we need to strip them
+				// off first. Enums are always named with their prefix intact.
+				if (!CPPType.IsEmpty() && (CPPType[0] == TEXT('F') || CPPType[0] == TEXT('U') || CPPType[0] == TEXT('A')))
+				{
+					CPPType = CPPType.Mid(1);
+				}
+				CPPTypeObject = FindObjectGloballyWithRedirectors(*CPPType);
+			}
+
+			if(CPPTypeObject == nullptr)
+			{
+				const FString Message = FString::Printf(
+					TEXT("%s: Pin '%s' is missing the CPPTypeObject for CPPType '%s'."),
+					*InPin->GetPathName(), *InPin->GetPinPath(), *InPin->GetCPPType());
+				FScriptExceptionHandler::Get().HandleException(ELogVerbosity::Error, *Message, *FString());
+				return false;
+			}
+			
+			InPin->CPPTypeObject = CPPTypeObject;
 		}
 	}
-	
+
+	InPin->CPPType = RigVMTypeUtils::PostProcessCPPType(InPin->CPPType, InPin->GetCPPTypeObject());
+
 	if(bRecursive)
 	{
 		for(URigVMPin* SubPin : InPin->SubPins)
@@ -19125,12 +19700,7 @@ void URigVMController::RefreshFunctionReferences(URigVMLibraryNode* InFunctionDe
 		FunctionLibrary->ForEachReference(InFunctionDefinition->GetFName(), [this, bSetupUndoRedo](URigVMFunctionReferenceNode* ReferenceNode)
 		{
 			FRigVMControllerGraphGuard GraphGuard(this, ReferenceNode->GetGraph(), bSetupUndoRedo);
-
-			TArray<URigVMLink*> Links = ReferenceNode->GetLinks();
-			DetachLinksFromPinObjects(&Links);
-			RepopulatePinsOnNode(ReferenceNode, false);
-			TGuardValue<bool> ReportGuard(bReportWarningsAndErrors, false);
-			ReattachLinksToPinObjects(false, &Links);
+			RepopulatePinsOnNode(ReferenceNode, false, true, false, true);
 		});
 	}
 }
@@ -19391,308 +19961,6 @@ FRigVMTemplate::FTypeMap URigVMController::GetCommonlyUsedTypesForTemplate(
 	return Template->GetArgumentTypesFromString(TypesString);
 }
 
-FRigVMClientPatchResult URigVMController::PatchUnitNodesOnLoad()
-{
-	FRigVMClientPatchResult Result;
-	
-	if (const URigVMGraph* Graph = GetGraph())
-	{
-		TArray<URigVMUnitNode*> UnitNodesToTurnToDispatches;
-		
-		// check for unit nodes that should be dispatches
-		for(URigVMNode* Node : Graph->GetNodes())
-		{
-			if(URigVMUnitNode* UnitNode = Cast<URigVMUnitNode>(Node))
-			{
-				if(const FRigVMTemplate* Template = UnitNode->GetTemplate())
-				{
-					if(Template->GetDispatchFactory() != nullptr)
-					{
-						UnitNodesToTurnToDispatches.Add(UnitNode);
-					}
-				}
-			}
-		}
-
-		// convert unit nodes to dispatches
-		for(URigVMUnitNode* UnitNode : UnitNodesToTurnToDispatches)
-		{
-			TArray<TPair<FString, FString>> LinkedPaths = GetLinkedPinPaths(UnitNode);
-			const FVector2D NodePosition = UnitNode->GetPosition();
-			const FString NodeName = UnitNode->GetName();
-			TMap<FString, FPinState> PinStates = GetPinStates(UnitNode, false);
-			
-			FRigVMTemplate::FTypeMap TypeMap;
-			for(URigVMPin* Pin : UnitNode->GetPins())
-			{
-				if(Pin->GetDirection() == ERigVMPinDirection::Input ||
-					Pin->GetDirection() == ERigVMPinDirection::Visible ||
-					Pin->GetDirection() == ERigVMPinDirection::IO ||
-					Pin->GetDirection() == ERigVMPinDirection::Output)
-				{
-					TypeMap.Add(Pin->GetFName(), Pin->GetTypeIndex());
-				}
-			}
-
-			const FRigVMTemplate* Template = UnitNode->GetTemplate();
-
-			Result.RemovedNodes.Add(UnitNode->GetPathName());
-			Result.bChangedContent = true;
-
-			RemoveNode(UnitNode, false, false, false, false);
-
-			URigVMTemplateNode* NewNode = AddTemplateNode(
-				Template->GetNotation(),
-				NodePosition, 
-				NodeName, 
-				false, 
-				false);
-
-			Result.AddedNodes.Add(NewNode);
-
-			TArray<int32> Permutations;
-			Template->Resolve(TypeMap, Permutations, false);
-
-			for(URigVMPin* Pin : NewNode->GetPins())
-			{
-				if(Pin->IsWildCard())
-				{
-					if(const TRigVMTypeIndex* ResolvedTypeIndex = TypeMap.Find(Pin->GetFName()))
-					{
-						if(!FRigVMRegistry::Get().IsWildCardType(*ResolvedTypeIndex))
-						{
-							ChangePinType(Pin, *ResolvedTypeIndex, false, false);
-							NewNode->UpdateFilteredPermutations(Pin, {*ResolvedTypeIndex});
-						}
-					}
-				}
-			}
-
-			ApplyPinStates(NewNode, PinStates, {}, false);
-			RestoreLinkedPaths(LinkedPaths, {}, {}, false);
-		}
-	}
-
-	return Result;
-}
-
-FRigVMClientPatchResult URigVMController::PatchDispatchNodesOnLoad()
-{
-	FRigVMClientPatchResult Result;
-
-	if (const URigVMGraph* Graph = GetGraph())
-	{
-		for(URigVMNode* Node : Graph->GetNodes())
-		{
-			if(URigVMDispatchNode* DispatchNode = Cast<URigVMDispatchNode>(Node))
-			{
-				// find template performs backwards lookup
-				if(const FRigVMTemplate* Template = DispatchNode->GetTemplate())
-				{
-					if(Template->GetNotation() != DispatchNode->TemplateNotation)
-					{
-						Result.bChangedContent = Result.bChangedContent || DispatchNode->TemplateNotation != Template->GetNotation(); 
-						DispatchNode->TemplateNotation = Template->GetNotation();
-						
-						if(!DispatchNode->ResolvedFunctionName.IsEmpty())
-						{
-							FString FactoryName, ArgumentsString;
-							if(DispatchNode->ResolvedFunctionName.Split(TEXT("::"), &FactoryName, &ArgumentsString))
-							{
-								DispatchNode->ResolvedFunctionName.Reset();
-								DispatchNode->ResolvedPermutation = INDEX_NONE;
-								
-								const FRigVMTemplateTypeMap ArgumentTypes = Template->GetArgumentTypesFromString(ArgumentsString);
-								if(ArgumentTypes.Num() == Template->NumArguments())
-								{
-									if(const FRigVMDispatchFactory* Factory = Template->GetDispatchFactory())
-									{
-										const FString ResolvedPermutationName = Factory->GetPermutationName(ArgumentTypes);
-										if(const FRigVMFunction* Function = FRigVMRegistry::Get().FindFunction(*ResolvedPermutationName))
-										{
-											DispatchNode->ResolvedFunctionName = Function->GetName();
-											DispatchNode->ResolvedPermutation = Template->FindPermutation(Function);
-											Result.bChangedContent = true;
-										}
-									}
-								}
-							}
-						}
-					}
-				}
-			}
-		}
-	}
-
-	return Result;
-}
-
-FRigVMClientPatchResult URigVMController::PatchBranchNodesOnLoad()
-{
-	FRigVMClientPatchResult Result;
-
-	if (const URigVMGraph* Graph = GetGraph())
-	{
-		TArray<URigVMNode*> BranchNodes = Graph->GetNodes().FilterByPredicate([](URigVMNode* Node)
-		{
-			return Node->IsA<UDEPRECATED_RigVMBranchNode>();
-		});
-
-		for(URigVMNode* BranchNode : BranchNodes)
-		{
-			TArray<TPair<FString, FString>> LinkedPaths = GetLinkedPinPaths(BranchNode);
-			const FVector2D NodePosition = BranchNode->GetPosition();
-			const FString NodeName = BranchNode->GetName();
-			const URigVMPin* OldConditionPin = BranchNode->FindPin(GET_MEMBER_NAME_CHECKED(FRigVMFunction_ControlFlowBranch, Condition).ToString());
-			const FString ConditionDefault = GetPinDefaultValue(OldConditionPin->GetPinPath());
-
-			Result.RemovedNodes.Add(BranchNode->GetPathName());
-			Result.bChangedContent = true;
-
-			RemoveNode(BranchNode, false, true, false, false);
-
-			const URigVMNode* NewNode = AddUnitNode(FRigVMFunction_ControlFlowBranch::StaticStruct(), FRigVMStruct::ExecuteName, NodePosition, NodeName, false, false);
-
-			Result.AddedNodes.Add(NewNode);
-
-			if(!ConditionDefault.IsEmpty())
-			{
-				const URigVMPin* ConditionPin = NewNode->FindPin(GET_MEMBER_NAME_CHECKED(FRigVMFunction_ControlFlowBranch, Condition).ToString());
-				SetPinDefaultValue(ConditionPin->GetPinPath(), ConditionDefault, false, false, false, false);
-			}
-			RestoreLinkedPaths(LinkedPaths, {}, {}, false);
-		}
-	}
-	
-	return Result;
-}
-
-FRigVMClientPatchResult URigVMController::PatchIfSelectNodesOnLoad()
-{
-	FRigVMClientPatchResult Result;
-	
-	if (const URigVMGraph* Graph = GetGraph())
-	{
-		TArray<URigVMNode*> IfOrSelectNodes = Graph->GetNodes().FilterByPredicate([](URigVMNode* Node)
-		{
-			return Node->IsA<UDEPRECATED_RigVMIfNode>() ||
-				Node->IsA<UDEPRECATED_RigVMSelectNode>();
-		});
-
-		for(URigVMNode* IfOrSelectNode : IfOrSelectNodes)
-		{
-			const bool bIsIfNode = IfOrSelectNode->IsA<UDEPRECATED_RigVMIfNode>();
-			TArray<TPair<FString, FString>> LinkedPaths = GetLinkedPinPaths(IfOrSelectNode);
-			const FVector2D NodePosition = IfOrSelectNode->GetPosition();
-			const FString NodeName = IfOrSelectNode->GetName();
-			const TRigVMTypeIndex TypeIndex = IfOrSelectNode->GetPins().Last()->GetTypeIndex();
-			TMap<FString, FPinState> PinStates = GetPinStates(IfOrSelectNode, true);
-
-			Result.RemovedNodes.Add(IfOrSelectNode->GetPathName());
-			Result.bChangedContent = true;
-
-			RemoveNode(IfOrSelectNode, false, true, false, false);
-
-			const FRigVMDispatchFactory* Factory = FRigVMRegistry::Get().FindOrAddDispatchFactory(
-				bIsIfNode ? FRigVMDispatch_If::StaticStruct() : FRigVMDispatch_SelectInt32::StaticStruct());
-
-			const FRigVMTemplate* Template = Factory->GetTemplate();
-			URigVMTemplateNode* NewNode = AddTemplateNode(
-				Template->GetNotation(),
-				NodePosition, 
-				NodeName, 
-				false, 
-				false);
-
-			Result.AddedNodes.Add(NewNode);
-
-			if(!FRigVMRegistry::Get().IsWildCardType(TypeIndex))
-			{
-				TArray<int32> Permutations;
-				FRigVMTemplateTypeMap Types;
-				Types.Add(IfOrSelectNode->GetPins().Last()->GetFName(), TypeIndex);
-				Template->Resolve(Types, Permutations, false);
-
-				for(URigVMPin* Pin : NewNode->GetPins())
-				{
-					if(Pin->IsWildCard())
-					{
-						if(const TRigVMTypeIndex* ResolvedTypeIndex = Types.Find(Pin->GetFName()))
-						{
-							if(!FRigVMRegistry::Get().IsWildCardType(*ResolvedTypeIndex))
-							{
-								NewNode->UpdateFilteredPermutations(Pin, {*ResolvedTypeIndex});
-							}
-						}
-					}
-				}
-				UpdateTemplateNodePinTypes(NewNode, false, false);
-			}
-
-			ApplyPinStates(NewNode, PinStates, {}, false);
-			RestoreLinkedPaths(LinkedPaths, {}, {}, false);
-		}
-	}
-
-	return Result;
-}
-
-FRigVMClientPatchResult URigVMController::PatchArrayNodesOnLoad()
-{
-	FRigVMClientPatchResult Result;
-	
-	if (const URigVMGraph* Graph = GetGraph())
-	{
-		TArray<URigVMNode*> ArrayNodes = Graph->GetNodes().FilterByPredicate([](URigVMNode* Node)
-		{
-			return Node->IsA<UDEPRECATED_RigVMArrayNode>();
-		});
-
-		for(URigVMNode* ModelNode : ArrayNodes)
-		{
-			UDEPRECATED_RigVMArrayNode* ArrayNode = CastChecked<UDEPRECATED_RigVMArrayNode>(ModelNode);
-			TArray<TPair<FString, FString>> LinkedPaths = GetLinkedPinPaths(ArrayNode);
-			const FVector2D NodePosition = ArrayNode->GetPosition();
-			const FString NodeName = ArrayNode->GetName();
-			const FString CPPType = ArrayNode->GetCPPType();
-			UObject* CPPTypeObject = ArrayNode->GetCPPTypeObject();
-			const ERigVMOpCode OpCode = ArrayNode->GetOpCode();
-			TMap<FString, FPinState> PinStates = GetPinStates(ArrayNode, true);
-
-			Result.RemovedNodes.Add(ArrayNode->GetPathName());
-			Result.bChangedContent = true;
-
-			RemoveNode(ArrayNode, false, false, false, false);
-
-			URigVMNode* NewNode = AddArrayNode(OpCode, CPPType, CPPTypeObject, NodePosition, NodeName, false, false, true);
-
-			Result.AddedNodes.Add(NewNode);
-
-			ApplyPinStates(NewNode, PinStates, {}, false);
-			RestoreLinkedPaths(LinkedPaths, {}, {}, false);
-		}
-	}
-
-	return Result;
-}
-
-void URigVMController::PostDuplicateHost(const FString& InOldPathName, const FString& InNewPathName)
-{
-	if (const URigVMGraph* Graph = GetGraph())
-	{
-		TArray<URigVMNode*> FunctionRefNodes = Graph->GetNodes().FilterByPredicate([](URigVMNode* Node)
-		{
-			return Node->IsA<URigVMFunctionReferenceNode>();
-		});
-
-		for(URigVMNode* Node : FunctionRefNodes)
-		{
-			URigVMFunctionReferenceNode* FunctionReferenceNode = CastChecked<URigVMFunctionReferenceNode>(Node);
-			FunctionReferenceNode->ReferencedFunctionHeader.PostDuplicateHost(InOldPathName, InNewPathName);
-		}
-	}
-}
-
 #endif
 
 URigVMControllerSettings::URigVMControllerSettings(const FObjectInitializer& Initializer)
diff --git a/Engine/Source/Developer/RigVMDeveloper/Private/RigVMModel/RigVMPin.cpp b/Engine/Source/Developer/RigVMDeveloper/Private/RigVMModel/RigVMPin.cpp
--- a/Engine/Source/Developer/RigVMDeveloper/Private/RigVMModel/RigVMPin.cpp
+++ b/Engine/Source/Developer/RigVMDeveloper/Private/RigVMModel/RigVMPin.cpp
@@ -17,13 +17,10 @@
 #include "Misc/StringBuilder.h"
 #include "Logging/LogScopedVerbosityOverride.h"
 #include "RigVMModel/Nodes/RigVMCollapseNode.h"
-#include "RigVMModel/Nodes/RigVMDispatchNode.h"
 #include "RigVMModel/Nodes/RigVMFunctionReferenceNode.h"
 #include "RigVMModel/Nodes/RigVMFunctionEntryNode.h"
 #include "RigVMModel/Nodes/RigVMFunctionReturnNode.h"
 #include "RigVMModel/Nodes/RigVMInvokeEntryNode.h"
-#include "RigVMModel/Nodes/RigVMSelectNode.h"
-#include "RigVMStringUtils.h"
 
 #include UE_INLINE_GENERATED_CPP_BY_NAME(RigVMPin)
 
@@ -57,32 +54,116 @@ const FString URigVMPin::OrphanPinPrefix = TEXT("Orphan::");
 
 bool URigVMPin::SplitPinPathAtStart(const FString& InPinPath, FString& LeftMost, FString& Right)
 {
-	return RigVMStringUtils::SplitPinPathAtStart(InPinPath, LeftMost, Right);
+	return InPinPath.Split(TEXT("."), &LeftMost, &Right, ESearchCase::IgnoreCase, ESearchDir::FromStart);
 }
 
 bool URigVMPin::SplitPinPathAtEnd(const FString& InPinPath, FString& Left, FString& RightMost)
 {
-	return RigVMStringUtils::SplitPinPathAtEnd(InPinPath, Left, RightMost);
+	return InPinPath.Split(TEXT("."), &Left, &RightMost, ESearchCase::IgnoreCase, ESearchDir::FromEnd);
 }
 
 bool URigVMPin::SplitPinPath(const FString& InPinPath, TArray<FString>& Parts)
 {
-	return RigVMStringUtils::SplitPinPath(InPinPath, Parts);
+	int32 OriginalPartsCount = Parts.Num();
+	FString PinPathRemaining = InPinPath;
+	FString Left, Right;
+	while(SplitPinPathAtStart(PinPathRemaining, Left, Right))
+	{
+		Parts.Add(Left);
+		Left.Empty();
+		PinPathRemaining = Right;
+	}
+
+	if (!Right.IsEmpty())
+	{
+		Parts.Add(Right);
+	}
+
+	return Parts.Num() > OriginalPartsCount;
 }
 
 FString URigVMPin::JoinPinPath(const FString& Left, const FString& Right)
 {
-	return RigVMStringUtils::JoinPinPath(Left, Right);
+	ensure(!Left.IsEmpty() && !Right.IsEmpty());
+	return Left + TEXT(".") + Right;
 }
 
 FString URigVMPin::JoinPinPath(const TArray<FString>& InParts)
 {
-	return RigVMStringUtils::JoinPinPath(InParts);
+	if (InParts.Num() == 0)
+	{
+		return FString();
+	}
+
+	FString Result = InParts[0];
+	for (int32 PartIndex = 1; PartIndex < InParts.Num(); PartIndex++)
+	{
+		Result += TEXT(".") + InParts[PartIndex];
+	}
+
+	return Result;
 }
 
 TArray<FString> URigVMPin::SplitDefaultValue(const FString& InDefaultValue)
 {
-	return RigVMStringUtils::SplitDefaultValue(InDefaultValue);
+	TArray<FString> Parts;
+	if (InDefaultValue.IsEmpty())
+	{
+		return Parts;
+	}
+
+	ensure(InDefaultValue[0] == TCHAR('('));
+	ensure(InDefaultValue[InDefaultValue.Len() - 1] == TCHAR(')')); 
+
+	FString Content = InDefaultValue.Mid(1, InDefaultValue.Len() - 2);
+	int32 BraceCount = 0;
+	int32 QuoteCount = 0;
+
+	int32 LastPartStartIndex = 0;
+	for (int32 CharIndex = 0; CharIndex < Content.Len(); CharIndex++)
+	{
+		TCHAR Char = Content[CharIndex];
+		if (QuoteCount > 0)
+		{
+			if (Char == TCHAR('"'))
+			{
+				QuoteCount = 0;
+			}
+		}
+		else if (Char == TCHAR('"'))
+		{
+			QuoteCount = 1;
+		}
+
+		if (Char == TCHAR('('))
+		{
+			if (QuoteCount == 0)
+			{
+				BraceCount++;
+			}
+		}
+		else if (Char == TCHAR(')'))
+		{
+			if (QuoteCount == 0)
+			{
+				BraceCount--;
+				BraceCount = FMath::Max<int32>(BraceCount, 0);
+			}
+		}
+		else if (Char == TCHAR(',') && BraceCount == 0 && QuoteCount == 0)
+		{
+			// ignore whitespaces
+			Parts.Add(Content.Mid(LastPartStartIndex, CharIndex - LastPartStartIndex).Replace(TEXT(" "), TEXT("")));
+			LastPartStartIndex = CharIndex + 1;
+		}
+	}
+
+	if (!Content.IsEmpty())
+	{
+		// ignore whitespaces from the start and end of the string
+		Parts.Add(Content.Mid(LastPartStartIndex).TrimStartAndEnd());
+	}
+	return Parts;
 }
 
 FString URigVMPin::GetDefaultValueForArray(TConstArrayView<FString> DefaultValues)
@@ -357,14 +438,6 @@ FName URigVMPin::GetDisplayName() const
 {
 	if (DisplayName == NAME_None)
 	{
-		if(const URigVMTemplateNode* Node = Cast<URigVMTemplateNode>(GetNode()))
-		{
-			const FName DisplayNameForArgument = Node->GetDisplayNameForPin(*this->GetSegmentPath(true));
-			if(!DisplayNameForArgument.IsNone())
-			{
-				return DisplayNameForArgument;
-			}
-		}
 		return GetFName();
 	}
 
@@ -397,13 +470,6 @@ ERigVMPinDirection URigVMPin::GetDirection() const
 
 bool URigVMPin::IsExpanded() const
 {
-	if(!bIsExpanded)
-	{
-		if(IsFixedSizeArray())
-		{
-			return true;
-		}
-	}
 	return bIsExpanded;
 }
 
@@ -493,28 +559,6 @@ bool URigVMPin::IsDynamicArray() const
 	return bIsDynamicArray;
 }
 
-bool URigVMPin::IsLazy() const
-{
-	// fixed array elements are treated as lazy elements
-	// if the original argument is also marked as lazy
-	if(const URigVMPin* ParentPin = GetParentPin())
-	{
-		if(ParentPin->IsFixedSizeArray())
-		{
-			return ParentPin->IsLazy();
-		}
-	}
-	
-	if(GetDirection() == ERigVMPinDirection::Input)
-	{
-		if(const URigVMNode* Node = GetNode())
-		{
-			return Node->ShouldInputPinComputeLazily(this);
-		}
-	}
-	return false;
-}
-
 int32 URigVMPin::GetPinIndex() const
 {
 	int32 Index = INDEX_NONE;
@@ -647,40 +691,6 @@ bool URigVMPin::ContainsWildCardSubPin() const
 	return false;
 }
 
-bool URigVMPin::IsFixedSizeArray() const
-{
-#if WITH_EDITOR
-
-	if(IsArray() && IsRootPin())
-	{
-		if(const URigVMNode* Node = GetNode())
-		{
-			if(const URigVMUnitNode* UnitNode = Cast<URigVMUnitNode>(Node))
-			{
-				if(const UScriptStruct* Struct = UnitNode->GetScriptStruct())
-				{
-					if(const FProperty* Property = Struct->FindPropertyByName(GetFName()))
-					{
-						return Property->HasMetaData(FRigVMStruct::FixedSizeArrayMetaName);
-					}
-				}
-			}
-			else if(const URigVMDispatchNode* DispatchNode = Cast<URigVMDispatchNode>(Node))
-			{
-				if(const FRigVMDispatchFactory* Factory = DispatchNode->GetFactory())
-				{
-					return Factory->HasArgumentMetaData(GetFName(), FRigVMStruct::FixedSizeArrayMetaName);
-				}
-			}
-			else if(Node->IsA<UDEPRECATED_RigVMSelectNode>())
-			{
-				return GetFName().ToString() == UDEPRECATED_RigVMSelectNode::ValueName;
-			}
-		}
-	}
-#endif
-	return false;
-}
 
 FString URigVMPin::GetDefaultValue() const
 {
@@ -838,7 +848,7 @@ bool URigVMPin::IsValidDefaultValue(const FString& InDefaultValue) const
 				return true;
 			}
 			
-			UObject* Object = RigVMTypeUtils::FindObjectFromCPPTypeObjectPath(Value);
+			UObject* Object = FindObjectFromCPPTypeObjectPath(Value);
 			if(Object == nullptr)
 			{
 				return false;
@@ -1053,6 +1063,36 @@ FText URigVMPin::GetToolTipText() const
 	return FText();
 }
 
+UObject* URigVMPin::FindObjectFromCPPTypeObjectPath(const FString& InObjectPath)
+{
+	if (InObjectPath.IsEmpty())
+	{
+		return nullptr;
+	}
+
+	if (InObjectPath == FName(NAME_None).ToString())
+	{
+		return nullptr;
+	}
+
+	// we do this to avoid ambiguous searches for 
+	// common names such as "transform" or "vector"
+	UPackage* Package = nullptr;
+	FString PackageName;
+	FString CPPTypeObjectName = InObjectPath;
+	if (InObjectPath.Split(TEXT("."), &PackageName, &CPPTypeObjectName))
+	{
+		Package = FindPackage(nullptr, *PackageName);
+	}
+	
+	if (UObject* ObjectWithinPackage = FindObject<UObject>(Package, *CPPTypeObjectName))
+	{
+		return ObjectWithinPackage;
+	}
+
+	return FindFirstObject<UObject>(*InObjectPath, EFindFirstObjectOptions::NativeFirst | EFindFirstObjectOptions::EnsureIfAmbiguous);
+}
+
 URigVMVariableNode* URigVMPin::GetBoundVariableNode() const
 {
 	for (TObjectPtr<URigVMInjectionInfo> InjectionInfo : InjectionInfos)
@@ -1301,13 +1341,10 @@ void URigVMPin::UpdateTypeInformationIfRequired() const
 		if (CPPTypeObjectPath != NAME_None)
 		{
 			URigVMPin* MutableThis = (URigVMPin*)this;
-			MutableThis->CPPTypeObject = RigVMTypeUtils::FindObjectFromCPPTypeObjectPath(CPPTypeObjectPath.ToString());
+			MutableThis->CPPTypeObject = FindObjectFromCPPTypeObjectPath(CPPTypeObjectPath.ToString());
 			MutableThis->CPPType = RigVMTypeUtils::PostProcessCPPType(CPPType, CPPTypeObject);
-			if (!MutableThis->CPPType.IsEmpty())
-			{
-				MutableThis->LastKnownTypeIndex = FRigVMRegistry::Get().FindOrAddType(FRigVMTemplateArgumentType(*CPPType, CPPTypeObject));
-				MutableThis->LastKnownCPPType = MutableThis->CPPType;
-			} 
+			MutableThis->LastKnownTypeIndex = FRigVMRegistry::Get().FindOrAddType(FRigVMTemplateArgumentType(*CPPType, CPPTypeObject));
+			MutableThis->LastKnownCPPType = MutableThis->CPPType; 
 		}
 	}
 
@@ -1544,44 +1581,50 @@ TArray<URigVMPin*> URigVMPin::GetLinkedTargetPins(bool bRecursive) const
 TArray<URigVMLink*> URigVMPin::GetSourceLinks(bool bRecursive) const
 {
 	TArray<URigVMLink*> Results;
-	for (URigVMLink* Link : Links)
+	if(GetDirection() == ERigVMPinDirection::IO ||
+		GetDirection() == ERigVMPinDirection::Input)
 	{
-		if (Link->GetTargetPin() == this)
+		for (URigVMLink* Link : Links)
 		{
-			Results.Add(Link);
+			if (Link->GetTargetPin() == this)
+			{
+				Results.Add(Link);
+			}
 		}
-	}
 
-	if (bRecursive)
-	{
-		for (URigVMPin* SubPin : SubPins)
+		if (bRecursive)
 		{
-			Results.Append(SubPin->GetSourceLinks(bRecursive));
+			for (URigVMPin* SubPin : SubPins)
+			{
+				Results.Append(SubPin->GetSourceLinks(bRecursive));
+			}
 		}
 	}
-
 	return Results;
 }
 
 TArray<URigVMLink*> URigVMPin::GetTargetLinks(bool bRecursive) const
 {
 	TArray<URigVMLink*> Results;
-	for (URigVMLink* Link : Links)
+	if(GetDirection() == ERigVMPinDirection::IO ||
+		GetDirection() == ERigVMPinDirection::Output)
 	{
-		if (Link->GetSourcePin() == this)
+		for (URigVMLink* Link : Links)
 		{
-			Results.Add(Link);
+			if (Link->GetSourcePin() == this)
+			{
+				Results.Add(Link);
+			}
 		}
-	}
 
-	if (bRecursive)
-	{
-		for (URigVMPin* SubPin : SubPins)
+		if (bRecursive)
 		{
-			Results.Append(SubPin->GetTargetLinks(bRecursive));
+			for (URigVMPin* SubPin : SubPins)
+			{
+				Results.Append(SubPin->GetTargetLinks(bRecursive));
+			}
 		}
 	}
-
 	return Results;
 }
 
@@ -1605,7 +1648,7 @@ URigVMNode* URigVMPin::GetNode() const
 URigVMGraph* URigVMPin::GetGraph() const
 {
 	URigVMNode* Node = GetNode();
-	if(IsValid(Node))
+	if(Node)
 	{
 		return Node->GetGraph();
 	}
@@ -1613,7 +1656,7 @@ URigVMGraph* URigVMPin::GetGraph() const
 	return nullptr;
 }
 
-bool URigVMPin::CanLink(URigVMPin* InSourcePin, URigVMPin* InTargetPin, FString* OutFailureReason, const FRigVMByteCode* InByteCode, ERigVMPinDirection InUserLinkDirection, bool bInAllowNonArgumentPins, bool bEnableTypeCasting)
+bool URigVMPin::CanLink(URigVMPin* InSourcePin, URigVMPin* InTargetPin, FString* OutFailureReason, const FRigVMByteCode* InByteCode, ERigVMPinDirection InUserLinkDirection, bool bInAllowNonArgumentPins)
 {
 	if (InSourcePin == nullptr || InTargetPin == nullptr)
 	{
@@ -1695,12 +1738,6 @@ bool URigVMPin::CanLink(URigVMPin* InSourcePin, URigVMPin* InTargetPin, FString*
 
 		if (bCPPTypesDiffer)
 		{
-			if(bEnableTypeCasting && RigVMTypeUtils::CanCastTypes(InSourcePin->GetTypeIndex(), InTargetPin->GetTypeIndex()))
-			{
-				bCPPTypesDiffer = false;
-			}
-
-			if(bCPPTypesDiffer)
 			{
 				auto TemplateNodeSupportsType = [](URigVMPin* InPin, const int32& InTypeIndex, FString* OutFailureReason) -> bool
 				{
@@ -1797,14 +1834,11 @@ bool URigVMPin::CanLink(URigVMPin* InSourcePin, URigVMPin* InTargetPin, FString*
 					URigVMPin* RootPin = InSourcePin->GetRootPin();
 					if (!Template->FindArgument(RootPin->GetFName()))
 					{
-						if(!RootPin->IsExecuteContext())
+						if (OutFailureReason)
 						{
-							if (OutFailureReason)
-							{
-								*OutFailureReason = FString::Printf(TEXT("Library pin %s supported types need to be reduced."), *RootPin->GetPinPath(true));
-							}
-							return false;
+							*OutFailureReason = FString::Printf(TEXT("Library pin %s supported types need to be reduced."), *RootPin->GetPinPath(true));
 						}
+						return false;
 					}
 				}
 			}
@@ -1818,14 +1852,11 @@ bool URigVMPin::CanLink(URigVMPin* InSourcePin, URigVMPin* InTargetPin, FString*
 					URigVMPin* RootPin = InTargetPin->GetRootPin();
 					if (!Template->FindArgument(RootPin->GetFName()))
 					{
-						if(!RootPin->IsExecuteContext())
+						if (OutFailureReason)
 						{
-							if (OutFailureReason)
-							{
-								*OutFailureReason = FString::Printf(TEXT("Library pin %s supported types need to be reduced."), *RootPin->GetPinPath(true));
-							}
-							return false;
+							*OutFailureReason = FString::Printf(TEXT("Library pin %s supported types need to be reduced."), *RootPin->GetPinPath(true));
 						}
+						return false;
 					}
 				}
 			}
@@ -1882,11 +1913,14 @@ bool URigVMPin::CanLink(URigVMPin* InSourcePin, URigVMPin* InTargetPin, FString*
 					SourceNodes[SourceNodeIndex]->IsA<URigVMVariableNode>();
 				if (!bNodeCanLinkAnywhere)
 				{
-					// pure / immutable nodes can be connected to any input in any order.
-					// since a new link is going to change the abstract syntax tree 
-					if (!SourceNodes[SourceNodeIndex]->IsMutable())
+					if (URigVMUnitNode* UnitNode = Cast<URigVMUnitNode>(SourceNodes[SourceNodeIndex]))
 					{
-						bNodeCanLinkAnywhere = true;
+						// pure / immutable nodes can be connected to any input in any order.
+						// since a new link is going to change the abstract syntax tree 
+						if (!UnitNode->IsMutable())
+						{
+							bNodeCanLinkAnywhere = true;
+						}
 					}
 				}
 
diff --git a/Engine/Source/Developer/RigVMDeveloper/Public/RigVMModel/RigVMController.h b/Engine/Source/Developer/RigVMDeveloper/Public/RigVMModel/RigVMController.h
--- a/Engine/Source/Developer/RigVMDeveloper/Public/RigVMModel/RigVMController.h
+++ b/Engine/Source/Developer/RigVMDeveloper/Public/RigVMModel/RigVMController.h
@@ -9,21 +9,25 @@
 #include "RigVMModel/Nodes/RigVMParameterNode.h"
 #include "RigVMModel/Nodes/RigVMCommentNode.h"
 #include "RigVMModel/Nodes/RigVMRerouteNode.h"
+#include "RigVMModel/Nodes/RigVMBranchNode.h"
 #include "RigVMModel/Nodes/RigVMIfNode.h"
 #include "RigVMModel/Nodes/RigVMSelectNode.h"
 #include "RigVMModel/Nodes/RigVMTemplateNode.h"
 #include "RigVMModel/Nodes/RigVMEnumNode.h"
 #include "RigVMModel/Nodes/RigVMCollapseNode.h"
 #include "RigVMModel/Nodes/RigVMFunctionReferenceNode.h"
+#include "RigVMModel/Nodes/RigVMArrayNode.h"
 #include "RigVMModel/Nodes/RigVMInvokeEntryNode.h"
 #include "RigVMModel/RigVMBuildData.h"
 #include "RigVMCore/RigVMUserWorkflow.h"
 #include "UObject/Interface.h"
 #include "RigVMController.generated.h"
 
+#ifndef UE_RIGVM_ENABLE_TEMPLATE_NODES
+#define UE_RIGVM_ENABLE_TEMPLATE_NODES 1
+#endif
+
 class URigVMActionStack;
-struct FRigVMGraphFunctionArgument;
-struct FRigVMGraphFunctionHeader;
 
 UENUM()
 enum class ERigVMControllerBulkEditType : uint8
@@ -78,9 +82,10 @@ private:
 DECLARE_DELEGATE_RetVal_OneParam(bool, FRigVMController_ShouldStructUnfoldDelegate, const UStruct*)
 DECLARE_DELEGATE_RetVal_OneParam(TArray<FRigVMExternalVariable>, FRigVMController_GetExternalVariablesDelegate, URigVMGraph*)
 DECLARE_DELEGATE_RetVal(const FRigVMByteCode*, FRigVMController_GetByteCodeDelegate)
-DECLARE_DELEGATE_RetVal_OneParam(bool, FRigVMController_RequestLocalizeFunctionDelegate, FRigVMGraphFunctionIdentifier&)
+DECLARE_DELEGATE_RetVal_OneParam(bool, FRigVMController_IsFunctionAvailableDelegate, URigVMLibraryNode*)
+DECLARE_DELEGATE_RetVal_OneParam(bool, FRigVMController_RequestLocalizeFunctionDelegate, URigVMLibraryNode*)
 DECLARE_DELEGATE_RetVal_ThreeParams(FName, FRigVMController_RequestNewExternalVariableDelegate, FRigVMGraphVariableDescription, bool, bool);
-DECLARE_DELEGATE_RetVal_TwoParams(bool, FRigVMController_IsDependencyCyclicDelegate, const FRigVMGraphFunctionHeader& Dependent, const FRigVMGraphFunctionHeader& Dependency)
+DECLARE_DELEGATE_RetVal_TwoParams(bool, FRigVMController_IsDependencyCyclicDelegate, UObject*, UObject*)
 DECLARE_DELEGATE_RetVal_TwoParams(FRigVMController_BulkEditResult, FRigVMController_RequestBulkEditDialogDelegate, URigVMLibraryNode*, ERigVMControllerBulkEditType)
 DECLARE_DELEGATE_RetVal_OneParam(bool, FRigVMController_RequestBreakLinksDialogDelegate, TArray<URigVMLink*>)
 DECLARE_DELEGATE_FiveParams(FRigVMController_OnBulkEditProgressDelegate, TSoftObjectPtr<URigVMFunctionReferenceNode>, ERigVMControllerBulkEditType, ERigVMControllerBulkEditProgress, int32, int32)
@@ -125,44 +130,6 @@ protected:
 	const uint8* Memory;
 };
 
-// A struct describing the result of a backwards compatibility patch
-USTRUCT()
-struct RIGVMDEVELOPER_API FRigVMClientPatchResult
-{
-public:
-
-	GENERATED_BODY()
-	
-	FRigVMClientPatchResult()
-		: bSucceeded(true)
-		, bChangedContent(false)
-		, bRequiresToMarkPackageDirty(false)
-	{}
-
-	bool Succeeded() const { return bSucceeded; }
-	bool ChangedContent() const { return bChangedContent; }
-	bool RequiresToMarkPackageDirty() const { return bRequiresToMarkPackageDirty; }
-
-	const TArray<FString>& GetErrorMessages() const { return ErrorMessages; }
-	const TArray<FString>& GetRemovedNodes() const { return RemovedNodes; }
-	const TArray<TWeakObjectPtr<const URigVMNode>>& GetAddedNodes() const { return AddedNodes; }
-
-private:
-
-	void Merge(const FRigVMClientPatchResult& InOther);
-	
-	bool bSucceeded;
-	bool bChangedContent;
-	bool bRequiresToMarkPackageDirty;
-
-	TArray<FString> ErrorMessages;
-	TArray<FString> RemovedNodes;
-	TArray<TWeakObjectPtr<const URigVMNode>> AddedNodes;
-
-	friend struct FRigVMClient;
-	friend class URigVMController;
-};
-
 /**
  * The Controller is the sole authority to perform changes
  * on the Graph. The Controller itself is stateless.
@@ -376,21 +343,21 @@ public:
 	// Branch nodes can be used to split the execution of into multiple branches,
 	// allowing to drive behavior by logic.
 	UFUNCTION(BlueprintCallable, Category = RigVMController)
-	URigVMNode* AddBranchNode(const FVector2D& InPosition = FVector2D::ZeroVector, const FString& InNodeName = TEXT(""), bool bSetupUndoRedo = true, bool bPrintPythonCommand = false);
+	URigVMBranchNode* AddBranchNode(const FVector2D& InPosition = FVector2D::ZeroVector, const FString& InNodeName = TEXT(""), bool bSetupUndoRedo = true, bool bPrintPythonCommand = false);
 
 	// Adds an if node to the graph.
 	// If nodes can be used to pick between two values based on a condition.
 	UFUNCTION(BlueprintCallable, Category = RigVMController)
-	URigVMNode* AddIfNode(const FString& InCPPType, const FName& InCPPTypeObjectPath, const FVector2D& InPosition = FVector2D::ZeroVector, const FString& InNodeName = TEXT(""), bool bSetupUndoRedo = true, bool bPrintPythonCommand = false);
+	URigVMIfNode* AddIfNode(const FString& InCPPType, const FName& InCPPTypeObjectPath, const FVector2D& InPosition = FVector2D::ZeroVector, const FString& InNodeName = TEXT(""), bool bSetupUndoRedo = true, bool bPrintPythonCommand = false);
 	UFUNCTION(BlueprintCallable, Category = RigVMController)
-	URigVMNode* AddIfNodeFromStruct(UScriptStruct* InScriptStruct, const FVector2D& InPosition = FVector2D::ZeroVector, const FString& InNodeName = TEXT(""), bool bSetupUndoRedo = true);
+	URigVMIfNode* AddIfNodeFromStruct(UScriptStruct* InScriptStruct, const FVector2D& InPosition = FVector2D::ZeroVector, const FString& InNodeName = TEXT(""), bool bSetupUndoRedo = true);
 
 	// Adds a select node to the graph.
 	// Select nodes can be used to pick between multiple values based on an index.
 	UFUNCTION(BlueprintCallable, Category = RigVMController)
-	URigVMNode* AddSelectNode(const FString& InCPPType, const FName& InCPPTypeObjectPath, const FVector2D& InPosition = FVector2D::ZeroVector, const FString& InNodeName = TEXT(""), bool bSetupUndoRedo = true, bool bPrintPythonCommand = false);
+	URigVMSelectNode* AddSelectNode(const FString& InCPPType, const FName& InCPPTypeObjectPath, const FVector2D& InPosition = FVector2D::ZeroVector, const FString& InNodeName = TEXT(""), bool bSetupUndoRedo = true, bool bPrintPythonCommand = false);
 	UFUNCTION(BlueprintCallable, Category = RigVMController)
-	URigVMNode* AddSelectNodeFromStruct(UScriptStruct* InScriptStruct, const FVector2D& InPosition = FVector2D::ZeroVector, const FString& InNodeName = TEXT(""), bool bSetupUndoRedo = true);
+	URigVMSelectNode* AddSelectNodeFromStruct(UScriptStruct* InScriptStruct, const FVector2D& InPosition = FVector2D::ZeroVector, const FString& InNodeName = TEXT(""), bool bSetupUndoRedo = true);
 
 	// Adds a template node to the graph.
 	UFUNCTION(BlueprintCallable, Category = RigVMController)
@@ -448,12 +415,12 @@ public:
 	// Adds a Array Node to the edited Graph.
 	// This causes a NodeAdded modified event.
 	UFUNCTION(BlueprintCallable, Category = RigVMController)
-	URigVMNode* AddArrayNode(ERigVMOpCode InOpCode, const FString& InCPPType, UObject* InCPPTypeObject, const FVector2D& InPosition = FVector2D::ZeroVector, const FString& InNodeName = TEXT(""), bool bSetupUndoRedo = true, bool bPrintPythonCommand = false, bool bIsPatching = false);
+	URigVMArrayNode* AddArrayNode(ERigVMOpCode InOpCode, const FString& InCPPType, UObject* InCPPTypeObject, const FVector2D& InPosition = FVector2D::ZeroVector, const FString& InNodeName = TEXT(""), bool bSetupUndoRedo = true, bool bPrintPythonCommand = false);
 
 	// Adds a Array Node to the edited Graph given a struct object path name.
 	// This causes a NodeAdded modified event.
 	UFUNCTION(BlueprintCallable, Category = RigVMController)
-	URigVMNode* AddArrayNodeFromObjectPath(ERigVMOpCode InOpCode, const FString& InCPPType, const FString& InCPPTypeObjectPath, const FVector2D& InPosition = FVector2D::ZeroVector, const FString& InNodeName = TEXT(""), bool bSetupUndoRedo = true, bool bPrintPythonCommand = false, bool bIsPatching = false);
+	URigVMArrayNode* AddArrayNodeFromObjectPath(ERigVMOpCode InOpCode, const FString& InCPPType, const FString& InCPPTypeObjectPath, const FVector2D& InPosition = FVector2D::ZeroVector, const FString& InNodeName = TEXT(""), bool bSetupUndoRedo = true, bool bPrintPythonCommand = false);
 
 	// Adds an entry invocation node
 	// This causes a NodeAdded modified event.
@@ -509,23 +476,19 @@ public:
 
 	// Copies a function declaration into this graph's local function library
 	UFUNCTION(BlueprintCallable, Category = RigVMController)
-	URigVMLibraryNode* LocalizeFunctionFromPath(const FString& InHostPath, const FName& InFunctionName, bool bLocalizeDependentPrivateFunctions = true, bool bSetupUndoRedo = true, bool bPrintPythonCommand = false);
-
-	UFUNCTION(BlueprintCallable, Category = RigVMController)
     URigVMLibraryNode* LocalizeFunction(
-		const FRigVMGraphFunctionIdentifier& InFunctionDefinition,
+    	URigVMLibraryNode* InFunctionDefinition,
 		bool bLocalizeDependentPrivateFunctions = true,
-		bool bSetupUndoRedo = true,
-		bool bPrintPythonCommand = false);
+    	bool bSetupUndoRedo = true,
+    	bool bPrintPythonCommand = false);
 
 	// Copies a series of function declaratioms into this graph's local function library
 	UFUNCTION(BlueprintCallable, Category = RigVMController)
-	TMap<FRigVMGraphFunctionIdentifier, URigVMLibraryNode*> LocalizeFunctions(
-		TArray<FRigVMGraphFunctionIdentifier> InFunctionDefinitions,
-		bool bLocalizeDependentPrivateFunctions = true,
-		bool bSetupUndoRedo = true,
-		bool bPrintPythonCommand = false);
-
+    TMap<URigVMLibraryNode*, URigVMLibraryNode*> LocalizeFunctions(
+        TArray<URigVMLibraryNode*> InFunctionDefinitions,
+        bool bLocalizeDependentPrivateFunctions = true,
+        bool bSetupUndoRedo = true,
+        bool bPrintPythonCommand = false);
 
 	// Returns a unique name
 	static FName GetUniqueName(const FName& InName, TFunction<bool(const FName&)> IsNameAvailableFunction, bool bAllowPeriod, bool bAllowSpace);
@@ -769,7 +732,7 @@ public:
 	// Adds a link to the graph.
 	// This causes a LinkAdded modified event.
 	UFUNCTION(BlueprintCallable, Category = RigVMController)
-	bool AddLink(const FString& InOutputPinPath, const FString& InInputPinPath, bool bSetupUndoRedo = true, bool bPrintPythonCommand = false, ERigVMPinDirection InUserDirection = ERigVMPinDirection::Output, bool bCreateCastNode = false);
+	bool AddLink(const FString& InOutputPinPath, const FString& InInputPinPath, bool bSetupUndoRedo = true, bool bPrintPythonCommand = false, ERigVMPinDirection InUserDirection = ERigVMPinDirection::Output);
 
 	// Removes a link from the graph.
 	// This causes a LinkRemoved modified event.
@@ -802,12 +765,6 @@ public:
 	bool SetExposedPinIndex(const FName& InPinName, int32 InNewIndex, bool bSetupUndoRedo = true, bool bPrintPythonCommand = false);
 
 	// Adds a function reference / invocation to the graph
-	UFUNCTION(BlueprintCallable, Category = RigVMController)
-	URigVMFunctionReferenceNode* AddFunctionReferenceNodeFromDescription(const FRigVMGraphFunctionHeader& InFunctionDefinition, const FVector2D& InNodePosition = FVector2D::ZeroVector, const
-	                                                                     FString& InNodeName = TEXT(""), bool bSetupUndoRedo = true, bool bPrintPythonCommand = false, bool bAllowPrivateFunctions = false);
-	UFUNCTION(BlueprintCallable, Category = RigVMController)
-	URigVMFunctionReferenceNode* AddExternalFunctionReferenceNode(const FString& InHostPath, const FName& InFunctionName, const FVector2D& InNodePosition = FVector2D::ZeroVector, const FString& InNodeName = TEXT(""), bool bSetupUndoRedo = true, bool bPrintPythonCommand = false);
-	
 	UFUNCTION(BlueprintCallable, Category = RigVMController)
 	URigVMFunctionReferenceNode* AddFunctionReferenceNode(URigVMLibraryNode* InFunctionDefinition, const FVector2D& InNodePosition = FVector2D::ZeroVector, const FString& InNodeName = TEXT(""), bool bSetupUndoRedo = true, bool bPrintPythonCommand = false);
 
@@ -827,14 +784,6 @@ public:
 	UFUNCTION(BlueprintCallable, Category = RigVMController)
 	bool RenameFunction(const FName& InOldFunctionName, const FName& InNewFunctionName, bool bSetupUndoRedo = true);
 
-	// Mark a function as public/private in the function library
-	UFUNCTION(BlueprintCallable, Category = RigVMController)
-	bool MarkFunctionAsPublic(const FName& InFunctionName, bool bInIsPublic, bool bSetupUndoRedo = true, bool bPrintPythonCommand = false);
-
-	// Returns true if a function is marked as public in the function library
-	UFUNCTION(BlueprintCallable, Category = RigVMController)
-	bool IsFunctionPublic(const FName& InFunctionName);
-
 	// Add a local variable to the graph
 	UFUNCTION(BlueprintCallable, Category = RigVMController)
 	FRigVMGraphVariableDescription AddLocalVariable(const FName& InVariableName, const FString& InCPPType, UObject* InCPPTypeObject, const FString& InDefaultValue, bool bSetupUndoRedo = true, bool bPrintPythonCommand = false);
@@ -858,7 +807,7 @@ public:
 	bool SetLocalVariableTypeFromObjectPath(const FName& InVariableName, const FString& InCPPType, const FString& InCPPTypeObjectPath, bool bSetupUndoRedo = true, bool bPrintPythonCommand = false);
 
 	UFUNCTION(BlueprintCallable, Category = RigVMController)
-	bool SetLocalVariableDefaultValue(const FName& InVariableName, const FString& InDefaultValue, bool bSetupUndoRedo = true, bool bPrintPythonCommand = false);
+	bool SetLocalVariableDefaultValue(const FName& InVariableName, const FString& InDefaultValue, bool bSetupUndoRedo = true, bool bPrintPythonCommand = false, bool bNotify = true);
 
 	// creates the options struct for a given workflow
 	UFUNCTION(BlueprintCallable, Category = RigVMController)
@@ -869,10 +818,10 @@ public:
 	bool PerformUserWorkflow(const FRigVMUserWorkflow& InWorkflow, const URigVMUserWorkflowOptions* InOptions, bool bSetupUndoRedo = true);
 
 	// Determine affected function references for a potential bulk edit on a library node
-	TArray<TSoftObjectPtr<URigVMFunctionReferenceNode>> GetAffectedReferences(ERigVMControllerBulkEditType InEditType, bool bForceLoad = false);
+	TArray<TSoftObjectPtr<URigVMFunctionReferenceNode>> GetAffectedReferences(ERigVMControllerBulkEditType InEditType, bool bForceLoad = false, bool bNotify = true);
 
 	// Determine affected assets for a potential bulk edit on a library node
-	TArray<FAssetData> GetAffectedAssets(ERigVMControllerBulkEditType InEditType, bool bForceLoad = false);
+	TArray<FAssetData> GetAffectedAssets(ERigVMControllerBulkEditType InEditType, bool bForceLoad = false, bool bNotify = true);
 
 	// A delegate that can be set to change the struct unfolding behaviour
 	FRigVMController_ShouldStructUnfoldDelegate UnfoldStructDelegate;
@@ -883,11 +832,17 @@ public:
 	// A delegate to retrieve the current bytecode of the graph
 	FRigVMController_GetByteCodeDelegate GetCurrentByteCodeDelegate;
 
+	// A delegate to determine if a function is public
+	FRigVMController_IsFunctionAvailableDelegate IsFunctionAvailableDelegate;
+
 	// A delegate to localize a function on demand
 	FRigVMController_RequestLocalizeFunctionDelegate RequestLocalizeFunctionDelegate;
 
 	// A delegate to create a new blueprint member variable
 	FRigVMController_RequestNewExternalVariableDelegate RequestNewExternalVariableDelegate;
+	
+	// A delegate to validate if we are allowed to introduce a dependency between two objects
+	FRigVMController_IsDependencyCyclicDelegate IsDependencyCyclicDelegate;
 
 	// A delegate to ask the host / client for a dialog to confirm a bulk edit
 	FRigVMController_RequestBulkEditDialogDelegate RequestBulkEditDialogDelegate;
@@ -904,8 +859,11 @@ public:
 	// A delegate to request to configure an options instance for a node workflow
 	FRigVMController_ConfigureWorkflowOptionsDelegate ConfigureWorkflowOptionsDelegate; 
 
-	int32 DetachLinksFromPinObjects(const TArray<URigVMLink*>* InLinks = nullptr);
-	int32 ReattachLinksToPinObjects(bool bFollowCoreRedirectors = false, const TArray<URigVMLink*>* InLinks = nullptr, bool bSetupOrphanedPins = false, bool bAllowNonArgumentLinks = false);
+	// Returns the build data of the host
+	static URigVMBuildData* GetBuildData(bool bCreateIfNeeded = true);
+
+	int32 DetachLinksFromPinObjects(const TArray<URigVMLink*>* InLinks = nullptr, bool bNotify = false);
+	int32 ReattachLinksToPinObjects(bool bFollowCoreRedirectors = false, const TArray<URigVMLink*>* InLinks = nullptr, bool bNotify = false, bool bSetupOrphanedPins = false, bool bAllowNonArgumentLinks = false);
 	void AddPinRedirector(bool bInput, bool bOutput, const FString& OldPinPath, const FString& NewPinPath);
 
 	// Removes nodes which went stale.
@@ -915,11 +873,11 @@ public:
 	bool ShouldRedirectPin(UScriptStruct* InOwningStruct, const FString& InOldRelativePinPath, FString& InOutNewRelativePinPath) const;
 	bool ShouldRedirectPin(const FString& InOldPinPath, FString& InOutNewPinPath) const;
 
-	void RepopulatePinsOnNode(URigVMNode* InNode, bool bFollowCoreRedirectors = true, bool bSetupOrphanedPins = false);
-	void RemovePinsDuringRepopulate(URigVMNode* InNode, TArray<URigVMPin*>& InPins, bool bSetupOrphanedPins);
+	void RepopulatePinsOnNode(URigVMNode* InNode, bool bFollowCoreRedirectors = true, bool bNotify = false, bool bSetupOrphanedPins = false, bool bDetachAndReattachLinks = false);
+	void RemovePinsDuringRepopulate(URigVMNode* InNode, TArray<URigVMPin*>& InPins, bool bNotify, bool bSetupOrphanedPins);
 
 	// removes any orphan pins that no longer holds a link
-	bool RemoveUnusedOrphanedPins(URigVMNode* InNode);
+	bool RemoveUnusedOrphanedPins(URigVMNode* InNode, bool bNotify);
 
 	// Initializes and recomputes the filtered permutations of all template nodes in the graph
 	// Returns true if any pin has change it's type or link was broken
@@ -962,7 +920,6 @@ public:
 	void ReportAndNotifyInfo(const FString& InMessage) const;
 	void ReportAndNotifyWarning(const FString& InMessage) const;
 	void ReportAndNotifyError(const FString& InMessage) const;
-	void ReportPinTypeChange(URigVMPin* InPin, const FString& InNewCPPType);
 	void SendUserFacingNotification(const FString& InMessage, float InDuration = 0.f, const UObject* InSubject = nullptr, const FName& InBrushName = TEXT("MessageLog.Warning")) const;
 
 	template <typename FmtType, typename... Types>
@@ -1001,12 +958,74 @@ public:
 		ReportAndNotifyError(FString::Printf(Fmt, Args...));
 	}
 
+	/**
+	 * Function to override the notification behavior and temporarily
+	 * disable all notifications. Client code is responsible for calling
+	 * SuspendNotifications(true) once all changes have been done.
+	 */
+	void SuspendNotifications(bool bSuspend) { bSuspendNotifications = bSuspend; }
+
 	/**
 	 * Helper function to disable a series of checks that can be ignored during a unit test
 	 */
 	UFUNCTION(BlueprintCallable, Category = RigVMController)
 	void SetIsRunningUnitTest(bool bIsRunning);
 
+public:
+	
+	struct FPinInfo
+	{
+		FPinInfo();
+		FPinInfo(const URigVMPin* InPin, int32 InParentIndex, ERigVMPinDirection InDirection = ERigVMPinDirection::Invalid);
+		FPinInfo(FProperty* InProperty, ERigVMPinDirection InDirection, int32 InParentIndex, const uint8* InDefaultValueMemory);
+
+		void CorrectExecuteTypeIndex();
+		
+		int32 ParentIndex;
+		FName Name;
+		ERigVMPinDirection Direction;
+		TRigVMTypeIndex TypeIndex;
+		bool bIsArray;
+		FProperty* Property;
+		FString PinPath;
+		FString DefaultValue;
+		FString DisplayName;
+		FString CustomWidgetName;
+		bool bIsExpanded;
+		bool bIsConstant;
+		bool bIsDynamicArray;
+
+		friend uint32 GetTypeHash(const FPinInfo& InPin);
+	};
+
+	struct FPinInfoArray
+	{
+		FPinInfoArray() {}
+		FPinInfoArray(const URigVMNode* InNode);
+		FPinInfoArray(const URigVMNode* InNode, URigVMController* InController);
+
+		int32 Num() const { return Pins.Num(); }
+		const FPinInfo& operator[](int32 InIndex) const { return Pins[InIndex]; }
+		FPinInfo& operator[](int32 InIndex) { return Pins[InIndex]; }
+		TArray<FPinInfo>::RangedForIteratorType begin() const { return Pins.begin(); }
+		TArray<FPinInfo>::RangedForIteratorType end() const { return Pins.end(); }
+
+		int32 AddPin(const URigVMPin* InPin, int32 InParentIndex, ERigVMPinDirection InDirection = ERigVMPinDirection::Invalid);
+		int32 AddPin(FProperty* InProperty, URigVMController* InController, ERigVMPinDirection InDirection, int32 InParentIndex, const uint8* InDefaultValueMemory);
+		int32 AddPin(URigVMController* InController, int32 InParentIndex, const FName& InName, ERigVMPinDirection InDirection, TRigVMTypeIndex InTypeIndex, const FString& InDefaultValue, const uint8* InDefaultValueMemory);
+		void AddPins(UScriptStruct* InScriptStruct, URigVMController* InController, ERigVMPinDirection InDirection, int32 InParentIndex, const uint8* InDefaultValueMemory);
+
+		const FString& GetPinPath(const int32 InIndex) const;
+		int32 GetIndexFromPinPath(const FString& InPinPath) const;
+		const FPinInfo* GetPinFromPinPath(const FString& InPinPath) const;
+		int32 GetRootIndex(const int32 InIndex) const;
+
+		friend uint32 GetTypeHash(const FPinInfoArray& InPins);
+		
+		mutable TArray<FPinInfo> Pins;
+		mutable TMap<FString, int32> PinPathLookup;;
+	};
+
 private:
 
 	UPROPERTY(BlueprintReadOnly, Category = RigVMController, meta = (ScriptName = "ModifiedEvent", AllowPrivateAccess = "true"))
@@ -1024,21 +1043,19 @@ private:
 	bool CanAddNode(URigVMNode* InNode, bool bReportErrors, bool bIgnoreFunctionEntryReturnNodes = false);
 	TObjectPtr<URigVMNode> FindEventNode(const UScriptStruct* InScriptStruct) const;
 	bool CanAddEventNode(UScriptStruct* InScriptStruct, const bool bReportErrors) const;
-	bool CanAddFunctionRefForDefinition(const FRigVMGraphFunctionHeader& InFunctionDefinition, bool bReportErrors, bool bAllowPrivateFunctions=false);
-	void AddPinsForStruct(UStruct* InStruct, URigVMNode* InNode, URigVMPin* InParentPin, ERigVMPinDirection InPinDirection, const FString& InDefaultValue, bool bAutoExpandArrays);
+	bool CanAddFunctionRefForDefinition(URigVMLibraryNode* InFunctionDefinition, bool bReportErrors);
+	void AddPinsForStruct(UStruct* InStruct, URigVMNode* InNode, URigVMPin* InParentPin, ERigVMPinDirection InPinDirection, const FString& InDefaultValue, bool bAutoExpandArrays, bool bNotify = false, const FPinInfoArray* PreviousPins = nullptr);
 	void AddPinsForArray(FArrayProperty* InArrayProperty, URigVMNode* InNode, URigVMPin* InParentPin, ERigVMPinDirection InPinDirection, const TArray<FString>& InDefaultValues, bool bAutoExpandArrays);
-	void AddPinsForTemplate(const FRigVMTemplate* InTemplate, const FRigVMTemplateTypeMap& InPinTypeMap, URigVMNode* InNode);
 	void ConfigurePinFromProperty(FProperty* InProperty, URigVMPin* InOutPin, ERigVMPinDirection InPinDirection = ERigVMPinDirection::Invalid);
 	void ConfigurePinFromPin(URigVMPin* InOutPin, URigVMPin* InPin, bool bCopyDisplayName = false);
-	void ConfigurePinFromArgument(URigVMPin* InOutPin, const FRigVMGraphFunctionArgument& InArgument, bool bCopyDisplayName = false);
 	virtual bool ShouldStructBeUnfolded(const UStruct* InStruct);
 	virtual bool ShouldPinBeUnfolded(URigVMPin* InPin);
-	bool SetPinDefaultValue(URigVMPin* InPin, const FString& InDefaultValue, bool bResizeArrays, bool bSetupUndoRedo, bool bMergeUndoAction);
+	bool SetPinDefaultValue(URigVMPin* InPin, const FString& InDefaultValue, bool bResizeArrays, bool bSetupUndoRedo, bool bMergeUndoAction, bool bNotify = true);
 	bool ResetPinDefaultValue(URigVMPin* InPin, bool bSetupUndoRedo);
 	static FString GetPinInitialDefaultValue(const URigVMPin* InPin);
 	static FString GetPinInitialDefaultValueFromStruct(UScriptStruct* ScriptStruct, const URigVMPin* InPin, uint32 InOffset);
 	URigVMPin* InsertArrayPin(URigVMPin* ArrayPin, int32 InIndex, const FString& InDefaultValue, bool bSetupUndoRedo);
-	bool RemovePin(URigVMPin* InPinToRemove, bool bSetupUndoRedo);
+	bool RemovePin(URigVMPin* InPinToRemove, bool bSetupUndoRedo, bool bNotify);
 	FProperty* FindPropertyForPin(const FString& InPinPath);
 	bool BindPinToVariable(URigVMPin* InPin, const FString& InNewBoundVariablePath, bool bSetupUndoRedo, const FString& InVariableNodeName = FString());
 	bool UnbindPinFromVariable(URigVMPin* InPin, bool bSetupUndoRedo);
@@ -1053,10 +1070,9 @@ private:
 	void RelinkSourceAndTargetPins(URigVMNode* RigNode, bool bSetupUndoRedo = true);
 
 public:
-	bool AddLink(URigVMPin* OutputPin, URigVMPin* InputPin, bool bSetupUndoRedo = true, ERigVMPinDirection InUserDirection = ERigVMPinDirection::Invalid, bool bCreateCastNode = false);
+	bool AddLink(URigVMPin* OutputPin, URigVMPin* InputPin, bool bSetupUndoRedo = true, ERigVMPinDirection InUserDirection = ERigVMPinDirection::Invalid);
 	bool BreakLink(URigVMPin* OutputPin, URigVMPin* InputPin, bool bSetupUndoRedo = true);
 	bool BreakAllLinks(URigVMPin* Pin, bool bAsInput, bool bSetupUndoRedo = true);
-	void EnableTypeCasting(bool bEnabled = true) { bEnableTypeCasting = bEnabled; }
 
 private:
 	bool BreakAllLinksRecursive(URigVMPin* Pin, bool bAsInput, bool bTowardsParent, bool bSetupUndoRedo);
@@ -1073,7 +1089,7 @@ private:
 	URigVMCollapseNode* PromoteFunctionReferenceNodeToCollapseNode(URigVMFunctionReferenceNode* InFunctionRefNode, bool bSetupUndoRedo, bool bRemoveFunctionDefinition);
 	void SetReferencedFunction(URigVMFunctionReferenceNode* InFunctionRefNode, URigVMLibraryNode* InNewReferencedNode, bool bSetupUndoRedo);
 
-	void RefreshFunctionPins(URigVMNode* InNode);
+	void RefreshFunctionPins(URigVMNode* InNode, bool bNotify = true);
 
 	void ReportRemovedLink(const FString& InSourcePinPath, const FString& InTargetPinPath);
 
@@ -1203,18 +1219,6 @@ public:
 	// Inquire on the commonly used types for a template node. This can be used to resolve a node without user input (as a default)
 	FRigVMTemplate::FTypeMap GetCommonlyUsedTypesForTemplate(const URigVMTemplateNode* InNode) const;
 #endif
-
-protected:
-
-	// backwards compatibility code
-	FRigVMClientPatchResult PatchUnitNodesOnLoad();
-	FRigVMClientPatchResult PatchDispatchNodesOnLoad();
-	FRigVMClientPatchResult PatchBranchNodesOnLoad();
-	FRigVMClientPatchResult PatchIfSelectNodesOnLoad();
-	FRigVMClientPatchResult PatchArrayNodesOnLoad();
-
-	// work to do after a duplication of the host asset
-	void PostDuplicateHost(const FString& InOldPathName, const FString& InNewPathName);
 	
 private: 
 	UPROPERTY(transient)
@@ -1227,7 +1231,6 @@ private:
 	bool bReportWarningsAndErrors;
 	bool bIgnoreRerouteCompactnessChanges;
 	ERigVMPinDirection UserLinkDirection;
-	bool bEnableTypeCasting;
 
 	// temporary maps used for pin redirection
 	// only valid between Detach & ReattachLinksToPinObjects
@@ -1246,17 +1249,17 @@ private:
 		{
 		}
 
-		friend uint32 GetTypeHash(const FControlRigStructPinRedirectorKey& Cache)
+		friend FORCEINLINE uint32 GetTypeHash(const FControlRigStructPinRedirectorKey& Cache)
 		{
 			return HashCombine(GetTypeHash(Cache.Struct), GetTypeHash(Cache.PinPathInNode));
 		}
 
-		bool operator ==(const FControlRigStructPinRedirectorKey& Other) const
+		FORCEINLINE bool operator ==(const FControlRigStructPinRedirectorKey& Other) const
 		{
 			return Struct == Other.Struct && PinPathInNode == Other.PinPathInNode;
 		}
 
-		bool operator !=(const FControlRigStructPinRedirectorKey& Other) const
+		FORCEINLINE bool operator !=(const FControlRigStructPinRedirectorKey& Other) const
 		{
 			return Struct != Other.Struct || PinPathInNode != Other.PinPathInNode;
 		}
@@ -1273,14 +1276,7 @@ private:
 	bool bIsTransacting; // Performing undo/redo transaction
 	bool bIsRunningUnitTest;
 	bool bIsFullyResolvingTemplateNode;
-
-public:
-	
 	bool bSuspendRecomputingTemplateFilters;
-	bool bSuspendTemplateComputation;
-
-private:
-
 #if WITH_EDITOR
 	bool bRegisterTemplateNodeUsage;
 #endif
@@ -1297,8 +1293,7 @@ private:
 	friend class FRigVMParserAST;
 	friend class FRigVMControllerCompileBracketScope;
 	friend class FRigVMControllerGraphGuard;
-	friend class FRigVMControllerNotifGuard;
-	friend struct FRigVMClient;
+	friend struct FPinInfoArray;
 };
 
 class FRigVMControllerGraphGuard
@@ -1332,28 +1327,6 @@ private:
 	int32 NumGraphs;
 };
 
-class FRigVMControllerNotifGuard
-{
-public:
-
-	FRigVMControllerNotifGuard(URigVMController* InController, bool bInSuspendNotifications = true)
-		: Controller(InController)
-	{
-		bPreviousSuspendNotifications = Controller->bSuspendNotifications;
-		Controller->bSuspendNotifications = bInSuspendNotifications;
-	}
-
-	~FRigVMControllerNotifGuard()
-	{
-		Controller->bSuspendNotifications = bPreviousSuspendNotifications;
-	}
-
-private:
-
-	URigVMController* Controller;
-	bool bPreviousSuspendNotifications;
-};
-
 USTRUCT()
 struct FRigVMController_CommonTypePerTemplate
 {
diff --git a/Engine/Source/Runtime/RigVM/Public/RigVMDefines.h b/Engine/Source/Runtime/RigVM/Public/RigVMDefines.h
--- a/Engine/Source/Runtime/RigVM/Public/RigVMDefines.h
+++ b/Engine/Source/Runtime/RigVM/Public/RigVMDefines.h
@@ -29,3 +29,7 @@
 #ifndef UE_RIGVM_DEBUG_EXECUTION
 #define UE_RIGVM_DEBUG_EXECUTION 0
 #endif
+
+#ifndef UE_RIGVMCONTROLLER_VERBOSE_REPOPULATE
+#define UE_RIGVMCONTROLLER_VERBOSE_REPOPULATE UE_BUILD_DEBUG
+#endif
