RigVM: Avoid recreating pins if not required during RepopulatePins

diff --git a/Engine/Source/Developer/RigVMDeveloper/Private/RigVMModel/RigVMController.cpp b/Engine/Source/Developer/RigVMDeveloper/Private/RigVMModel/RigVMController.cpp
--- a/Engine/Source/Developer/RigVMDeveloper/Private/RigVMModel/RigVMController.cpp
+++ b/Engine/Source/Developer/RigVMDeveloper/Private/RigVMModel/RigVMController.cpp
@@ -299,6 +299,434 @@ void URigVMController::SetIsRunningUnitTest(bool bIsRunning)
 	}	
 }
 
+URigVMController::FPinInfo::FPinInfo()
+	: ParentIndex(INDEX_NONE)
+	, Name(NAME_None)
+	, Direction(ERigVMPinDirection::Invalid)
+	, TypeIndex(INDEX_NONE)
+	, bIsArray(false)
+	, Property(nullptr)
+	, bIsExpanded(false)
+	, bIsConstant(false)
+	, bIsDynamicArray(false)
+{
+}
+
+URigVMController::FPinInfo::FPinInfo(const URigVMPin* InPin, int32 InParentIndex, ERigVMPinDirection InDirection)
+	: ParentIndex(InParentIndex)
+	, Name(InPin->GetName())
+	, Direction(InDirection == ERigVMPinDirection::Invalid ? InPin->GetDirection() : InDirection)
+	, TypeIndex(InPin->GetTypeIndex())
+	, bIsArray(InPin->IsArray())
+	, Property(nullptr)
+	, bIsExpanded(InPin->IsExpanded())
+	, bIsConstant(InPin->IsDefinedAsConstant())
+	, bIsDynamicArray(InPin->IsDynamicArray())
+{
+	// this method describes the info as currently represented in the model.
+
+	CorrectExecuteTypeIndex();
+	DefaultValue = InPin->GetDefaultValue();
+	
+	if(DefaultValue.IsEmpty() && (InPin->IsArray() || InPin->IsStruct()))
+	{
+		static const FString EmptyBraces(TEXT("()"));
+		DefaultValue = EmptyBraces;
+	}
+}
+
+URigVMController::FPinInfo::FPinInfo(FProperty* InProperty, ERigVMPinDirection InDirection, int32 InParentIndex, const uint8* InDefaultValueMemory)
+	: ParentIndex(InParentIndex)
+	, Name(InProperty->GetFName())
+	, Direction(InDirection)
+	, TypeIndex(INDEX_NONE)
+	, bIsArray(InProperty->IsA<FArrayProperty>())
+	, Property(InProperty)
+	, bIsExpanded(false)
+	, bIsConstant(false)
+	, bIsDynamicArray(false)
+{
+	// this method describes the info as needed based on the property structure
+
+	if (Direction == ERigVMPinDirection::Invalid)
+	{
+		Direction = FRigVMStruct::GetPinDirectionFromProperty(Property);
+	}
+
+#if WITH_EDITOR
+
+	if (CastField<FArrayProperty>(InProperty->GetOwnerProperty()) == nullptr)
+	{
+		const FString DisplayNameText = InProperty->GetDisplayNameText().ToString();
+		if (!DisplayNameText.IsEmpty())
+		{
+			DisplayName = *DisplayNameText;
+		}
+	}
+	
+	bIsConstant = InProperty->HasMetaData(TEXT("Constant"));
+	CustomWidgetName = InProperty->GetMetaData(TEXT("CustomWidget"));
+	if (InProperty->HasMetaData(FRigVMStruct::ExpandPinByDefaultMetaName))
+	{
+		bIsExpanded = true;
+	}
+
+#endif
+
+#if WITH_EDITOR
+	if (Direction == ERigVMPinDirection::Hidden)
+	{
+		if (!InProperty->HasMetaData(TEXT("ArraySize")))
+		{
+			bIsDynamicArray = true;
+		}
+	}
+	if (bIsDynamicArray)
+	{
+		if (InProperty->HasMetaData(FRigVMStruct::SingletonMetaName))
+		{
+			bIsDynamicArray = false;
+		}
+	}
+#endif
+
+	UObject* CPPTypeObject = nullptr;
+	
+	FProperty* PropertyForType = InProperty;
+	if (const FArrayProperty* ArrayProperty = CastField<FArrayProperty>(PropertyForType))
+	{
+		PropertyForType = ArrayProperty->Inner;
+	}
+
+	if (const FStructProperty* StructProperty = CastField<FStructProperty>(PropertyForType))
+	{
+		CPPTypeObject = StructProperty->Struct;
+	}
+	else if (const FObjectProperty* ObjectProperty = CastField<FObjectProperty>(PropertyForType))
+	{
+		if(RigVMCore::SupportsUObjects())
+		{
+			CPPTypeObject = ObjectProperty->PropertyClass;
+		}
+	}
+	else if (const FInterfaceProperty* InterfaceProperty = CastField<FInterfaceProperty>(PropertyForType))
+	{
+		if (RigVMCore::SupportsUInterfaces())
+		{
+			CPPTypeObject = InterfaceProperty->InterfaceClass;
+		}
+	}
+	else
+	{
+		if (const FEnumProperty* EnumProperty = CastField<FEnumProperty>(PropertyForType))
+		{
+			CPPTypeObject = EnumProperty->GetEnum();
+		}
+		if (const FByteProperty* ByteProperty = CastField<FByteProperty>(PropertyForType))
+		{
+			CPPTypeObject = ByteProperty->Enum;
+		}
+
+		if(InDefaultValueMemory)
+		{
+			InProperty->ExportText_Direct(DefaultValue, InDefaultValueMemory, InDefaultValueMemory, nullptr, PPF_None, nullptr);
+		}
+	}
+
+	FString ExtendedCppType;
+	FString CPPType = InProperty->GetCPPType(&ExtendedCppType);
+	CPPType += ExtendedCppType;
+	CPPType = RigVMTypeUtils::PostProcessCPPType(CPPType, CPPTypeObject);
+	TypeIndex = FRigVMRegistry::Get().GetTypeIndexFromCPPType(CPPType);
+	CorrectExecuteTypeIndex();
+}
+
+void URigVMController::FPinInfo::CorrectExecuteTypeIndex()
+{
+	const FRigVMRegistry& Registry = FRigVMRegistry::Get();
+	if(Registry.IsExecuteType(TypeIndex))
+	{
+		TRigVMTypeIndex DefaultExecuteType = RigVMTypeUtils::TypeIndex::Execute;
+		if(Registry.IsArrayType(TypeIndex))
+		{
+			DefaultExecuteType = Registry.GetArrayTypeFromBaseTypeIndex(DefaultExecuteType);
+		}
+		TypeIndex = DefaultExecuteType;
+	}
+}
+
+uint32 GetTypeHash(const URigVMController::FPinInfo& InPin)
+{
+	uint32 Hash = 0; //GetTypeHash(InPin.ParentIndex);
+	Hash = HashCombine(Hash, GetTypeHash(InPin.Name));
+	Hash = HashCombine(Hash, GetTypeHash((int32)InPin.Direction));
+	Hash = HashCombine(Hash, GetTypeHash((int32)InPin.TypeIndex));
+	Hash = HashCombine(Hash, GetTypeHash(InPin.bIsArray));
+	// we are not hashing the parent index,  pinpath, default value or the property since
+	// it doesn't matter for the structure validity of the node
+	return Hash;
+}
+
+URigVMController::FPinInfoArray::FPinInfoArray(const URigVMNode* InNode)
+{
+	// this method adds all pins as currently represented in the model.
+	for(const URigVMPin* Pin : InNode->GetPins())
+	{
+		(void)AddPin(Pin, INDEX_NONE);
+	}
+}
+
+int32 URigVMController::FPinInfoArray::AddPin(const URigVMPin* InPin, int32 InParentIndex, ERigVMPinDirection InDirection)
+{
+	// this method adds all pins as currently represented in the model.
+	const int32 Index = Pins.Emplace(InPin, InParentIndex, InDirection);
+	for(const URigVMPin* SubPin : InPin->GetSubPins())
+	{
+		const int32 SubPinIndex = AddPin(SubPin, Index, InDirection);
+		Pins[Index].SubPins.Add(SubPinIndex);
+	}
+	return Index;
+}
+
+URigVMController::FPinInfoArray::FPinInfoArray(const URigVMNode* InNode, URigVMController* InController, const FPinInfoArray* InPreviousPinInfos)
+{
+	// this method adds pins as needed based on the property structure
+	for(const URigVMPin* Pin : InNode->GetPins())
+	{
+		const FString DefaultValue = Pin->GetDefaultValue();
+		(void)AddPin(InController, INDEX_NONE, Pin->GetFName(), Pin->GetDirection(), Pin->GetTypeIndex(), DefaultValue, nullptr, InPreviousPinInfos);
+	}
+}
+
+URigVMController::FPinInfoArray::FPinInfoArray(const FRigVMGraphFunctionHeader& FunctionHeader,
+	URigVMController* InController, const FPinInfoArray* InPreviousPinInfos)
+{
+	// this method adds pins as needed based on the property structure
+	const FRigVMRegistry& Registry = FRigVMRegistry::Get();
+	for(const FRigVMGraphFunctionArgument& FunctionArgument : FunctionHeader.Arguments)
+	{
+		const TRigVMTypeIndex TypeIndex = Registry.GetTypeIndexFromCPPType(FunctionArgument.CPPType.ToString());
+		(void)AddPin(InController, INDEX_NONE, FunctionArgument.Name, FunctionArgument.Direction, TypeIndex, FunctionArgument.DefaultValue, nullptr, InPreviousPinInfos);
+	}
+}
+
+int32 URigVMController::FPinInfoArray::AddPin(FProperty* InProperty, URigVMController* InController,
+                                              ERigVMPinDirection InDirection, int32 InParentIndex, const uint8* InDefaultValueMemory)
+{
+	// this method adds pins as needed based on the property structure
+
+	check(InDefaultValueMemory);
+	
+	const int32 Index = Pins.Emplace(InProperty, InDirection, InParentIndex, InDefaultValueMemory);
+	if(InParentIndex != INDEX_NONE)
+	{
+		Pins[InParentIndex].SubPins.Add(Index);
+	}
+
+	if (const FStructProperty* StructProperty = CastField<FStructProperty>(InProperty))
+	{
+		(void)AddPins(StructProperty->Struct, InController, Pins[Index].Direction, Index, InDefaultValueMemory);
+	}
+	else if (const FArrayProperty* ArrayProperty = CastField<FArrayProperty>(InProperty))
+	{
+		FScriptArrayHelper ArrayHelper(ArrayProperty, InDefaultValueMemory);
+		for(int32 ElementIndex = 0; ElementIndex < ArrayHelper.Num(); ElementIndex++)
+		{
+			const uint8* ElementDefaultValueMemory = ArrayHelper.GetRawPtr(ElementIndex);
+			const int32 SubIndex = AddPin(ArrayProperty->Inner, InController, Pins[Index].Direction, Index, ElementDefaultValueMemory);
+			Pins[SubIndex].Name = *FString::FormatAsNumber(ElementIndex);
+		}
+	}
+
+	return Index;
+}
+
+int32 URigVMController::FPinInfoArray::AddPin(URigVMController* InController, int32 InParentIndex, const FName& InName, ERigVMPinDirection InDirection,
+	TRigVMTypeIndex InTypeIndex, const FString& InDefaultValue, const uint8* InDefaultValueMemory, const FPinInfoArray* InPreviousPinInfos)
+{
+	const FRigVMRegistry& Registry = FRigVMRegistry::Get();
+		
+	FPinInfo Info;
+	Info.ParentIndex = InParentIndex;
+	Info.Name = InName;
+	Info.Direction = InDirection;
+	Info.TypeIndex = InTypeIndex;
+	Info.bIsArray = Registry.IsArrayType(InTypeIndex);
+	Info.DefaultValue = InDefaultValue;
+	Info.CorrectExecuteTypeIndex();
+
+	const int32 Index = Pins.Add(Info);
+
+	if(InPreviousPinInfos)
+	{
+		const FString& PinPath = GetPinPath(Index);
+		const int32 PreviousIndex = InPreviousPinInfos->GetIndexFromPinPath(PinPath);
+		if(PreviousIndex != INDEX_NONE)
+		{
+			const FPinInfo& PreviousPin = (*InPreviousPinInfos)[PreviousIndex];
+			if(PreviousPin.TypeIndex == InTypeIndex)
+			{
+				Pins[Index].DefaultValue = PreviousPin.DefaultValue;
+			}
+		}
+	}
+	
+	const FRigVMTemplateArgumentType& Type = Registry.GetType(InTypeIndex);
+	if(!Type.IsWildCard())
+	{
+		if(Info.bIsArray)
+		{
+			const TRigVMTypeIndex& ElementTypeIndex = Registry.GetBaseTypeFromArrayTypeIndex(InTypeIndex);
+			const FRigVMTemplateArgumentType& ElementType = Registry.GetType(ElementTypeIndex);
+
+			const TArray<FString> Elements = URigVMPin::SplitDefaultValue(Pins[Index].DefaultValue);
+			for(int32 ElementIndex = 0; ElementIndex < Elements.Num(); ElementIndex++)
+			{
+				const FString& ElementDefaultValue = Elements[ElementIndex];
+				const uint8* ElementDefaultValueMemory = nullptr;
+				FStructOnScope ElementDefaultValueMemoryScope;
+
+				if(UScriptStruct* ElementScriptStruct = Cast<UScriptStruct>(ElementType.CPPTypeObject))
+				{
+					ElementDefaultValueMemoryScope = FStructOnScope(ElementScriptStruct);
+		
+					FRigVMPinDefaultValueImportErrorContext ErrorPipe;
+					ElementScriptStruct->ImportText(*ElementDefaultValue, ElementDefaultValueMemoryScope.GetStructMemory(), nullptr, PPF_None, &ErrorPipe, FString());
+					ElementDefaultValueMemory = ElementDefaultValueMemoryScope.GetStructMemory();
+				}
+
+				(void)AddPin(InController, Index, *FString::FormatAsNumber(ElementIndex), InDirection, ElementTypeIndex, ElementDefaultValue, ElementDefaultValueMemory, InPreviousPinInfos);
+			}
+		}
+		else if(UScriptStruct* ScriptStruct = Cast<UScriptStruct>(Type.CPPTypeObject))
+		{
+			const uint8* DefaultValueMemory = InDefaultValueMemory;
+			FStructOnScope DefaultValueMemoryScope;
+			if(DefaultValueMemory == nullptr)
+			{
+				FRigVMPinDefaultValueImportErrorContext ErrorPipe;
+				DefaultValueMemoryScope = FStructOnScope(ScriptStruct);
+				ScriptStruct->ImportText(*Pins[Index].DefaultValue, DefaultValueMemoryScope.GetStructMemory(), nullptr, PPF_None, &ErrorPipe, FString());
+				DefaultValueMemory = DefaultValueMemoryScope.GetStructMemory();
+			}
+			AddPins(ScriptStruct, InController, InDirection, Index, DefaultValueMemory);
+		}
+	}
+
+	if(InParentIndex != INDEX_NONE)
+	{
+		Pins[InParentIndex].SubPins.Add(Index);
+	}
+
+	return Index;
+}
+
+void URigVMController::FPinInfoArray::AddPins(UScriptStruct* InScriptStruct, URigVMController* InController,
+	ERigVMPinDirection InDirection, int32 InParentIndex, const uint8* InDefaultValueMemory)
+{
+	if (InController->ShouldStructBeUnfolded(InScriptStruct))
+	{
+		TArray<UStruct*> StructsToVisit = FRigVMTemplate::GetSuperStructs(InScriptStruct, true);
+		for(UStruct* StructToVisit : StructsToVisit)
+		{
+			// using EFieldIterationFlags::None excludes the
+			// properties of the super struct in this iterator.
+			for (TFieldIterator<FProperty> It(StructToVisit, EFieldIterationFlags::None); It; ++It)
+			{
+				const uint8* DefaultValueMemory = nullptr;
+				if(InDefaultValueMemory)
+				{
+					DefaultValueMemory = It->ContainerPtrToValuePtr<uint8>(InDefaultValueMemory);
+				}
+				(void)AddPin(*It, InController, InDirection, InParentIndex, DefaultValueMemory);
+			}
+		}
+	}
+}
+
+const FString& URigVMController::FPinInfoArray::GetPinPath(const int32 InIndex) const
+{
+	if(!Pins.IsValidIndex(InIndex))
+	{
+		static const FString EmptyString;
+		return EmptyString;
+	}
+
+	if(Pins[InIndex].PinPath.IsEmpty())
+	{
+		if(Pins[InIndex].ParentIndex == INDEX_NONE)
+		{
+			Pins[InIndex].PinPath = Pins[InIndex].Name.ToString();
+		}
+		else
+		{
+			Pins[InIndex].PinPath = URigVMPin::JoinPinPath(GetPinPath(Pins[InIndex].ParentIndex), Pins[InIndex].Name.ToString());
+		}
+	}
+
+	return Pins[InIndex].PinPath;
+}
+
+int32 URigVMController::FPinInfoArray::GetIndexFromPinPath(const FString& InPinPath) const
+{
+	if(PinPathLookup.Num() != Num())
+	{
+		PinPathLookup.Reset();
+		for(int32 Index=0; Index < Num(); Index++)
+		{
+			PinPathLookup.Add(GetPinPath(Index), Index);
+		}
+	}
+
+	if(const int32* Index = PinPathLookup.Find(InPinPath))
+	{
+		return *Index;
+	}
+	return INDEX_NONE;
+}
+
+const URigVMController::FPinInfo* URigVMController::FPinInfoArray::GetPinFromPinPath(const FString& InPinPath) const
+{
+	const int32 Index = GetIndexFromPinPath(InPinPath);
+	if(Pins.IsValidIndex(Index))
+	{
+		return &Pins[Index];
+	}
+	return nullptr;
+}
+
+int32 URigVMController::FPinInfoArray::GetRootIndex(const int32 InIndex) const
+{
+	if(Pins.IsValidIndex(InIndex))
+	{
+		if(Pins[InIndex].ParentIndex == INDEX_NONE)
+		{
+			return InIndex;
+		}
+		return GetRootIndex(Pins[InIndex].ParentIndex);
+	}
+	return INDEX_NONE;
+}
+
+uint32 GetTypeHash(const URigVMController::FPinInfoArray& InPins)
+{
+	TArray<uint32> Hashes;
+	Hashes.Reserve(InPins.Num());
+	
+	uint32 OverAllHash = GetTypeHash(InPins.Num());
+	for(const URigVMController::FPinInfo& Info : InPins)
+	{
+		uint32 PinHash = GetTypeHash(Info);
+		if(Info.ParentIndex != INDEX_NONE)
+		{
+			PinHash = HashCombine(PinHash, Hashes[Info.ParentIndex]);
+		}
+		Hashes.Add(PinHash);
+		OverAllHash = HashCombine(OverAllHash, PinHash); 
+	}
+	return OverAllHash;
+}
+
 void URigVMController::HandleModifiedEvent(ERigVMGraphNotifType InNotifType, URigVMGraph* InGraph, UObject* InSubject)
 {
 	switch (InNotifType)
@@ -5762,10 +6190,7 @@ URigVMCollapseNode* URigVMController::PromoteFunctionReferenceNodeToCollapseNode
 			{
 				if (URigVMVariableNode* VariableNode = Cast<URigVMVariableNode>(Node))
 				{
-					TArray<URigVMLink*> VariableLinks = VariableNode->GetLinks();
-					DetachLinksFromPinObjects(&VariableLinks);
-					RepopulatePinsOnNode(VariableNode);
-					ReattachLinksToPinObjects(false, &VariableLinks);
+					RepopulatePinsOnNode(VariableNode, true, false, true);
 				}
 			}
 
@@ -5834,10 +6259,7 @@ void URigVMController::RefreshFunctionPins(URigVMNode* InNode)
 
 	if (EntryNode || ReturnNode)
 	{
-		TArray<URigVMLink*> Links = InNode->GetLinks();
-		DetachLinksFromPinObjects(&Links);
-		RepopulatePinsOnNode(InNode, false);
-		ReattachLinksToPinObjects(false, &Links);
+		RepopulatePinsOnNode(InNode, false, false, true);
 	}
 }
 
@@ -8541,15 +8963,6 @@ bool URigVMController::RemovePin(URigVMPin* InPinToRemove, bool bSetupUndoRedo)
 		BreakAllLinksRecursive(InPinToRemove, false, false, bSetupUndoRedo);
 	}
 
-	if (URigVMPin* ParentPin = InPinToRemove->GetParentPin())
-	{
-		ParentPin->SubPins.Remove(InPinToRemove);
-	}
-	else if(URigVMNode* Node = InPinToRemove->GetNode())
-	{
-		Node->Pins.Remove(InPinToRemove);
-	}
-
 	TArray<URigVMPin*> SubPins = InPinToRemove->GetSubPins();
 	for (URigVMPin* SubPin : SubPins)
 	{
@@ -8559,6 +8972,16 @@ bool URigVMController::RemovePin(URigVMPin* InPinToRemove, bool bSetupUndoRedo)
 		}
 	}
 
+	if (URigVMPin* ParentPin = InPinToRemove->GetParentPin())
+	{
+		ParentPin->SubPins.Remove(InPinToRemove);
+	}
+	else if(URigVMNode* Node = InPinToRemove->GetNode())
+	{
+		Node->Pins.Remove(InPinToRemove);
+		Node->OrphanedPins.Remove(InPinToRemove);
+	}
+
 	if (!bSuspendNotifications)
 	{
 		Notify(ERigVMGraphNotifType::PinRemoved, InPinToRemove);
@@ -11208,10 +11631,7 @@ bool URigVMController::ChangeExposedPinType(const FName& InPinName, const FStrin
 	{
 		if (InterfaceNode)
 		{
-			const TArray<URigVMLink*> Links = InterfaceNode->GetLinks();
-			DetachLinksFromPinObjects(&Links);
-			RepopulatePinsOnNode(InterfaceNode, true, bSetupOrphanPins);
-			ReattachLinksToPinObjects(true, &Links, bSetupOrphanPins);
+			RepopulatePinsOnNode(InterfaceNode, true, bSetupOrphanPins, true);
 			RemoveUnusedOrphanedPins(InterfaceNode);
 		}
 	}
@@ -11402,11 +11822,8 @@ URigVMFunctionReferenceNode* URigVMController::AddFunctionReferenceNodeFromDescr
 
 	FRigVMControllerCompileBracketScope CompileScope(this);
 	
-	{
-		TGuardValue<bool> SuspendNotifications(bSuspendNotifications, true);
-		RepopulatePinsOnNode(FunctionRefNode, false);
-	}
-	//UpdateTemplateNodePinTypes(FunctionRefNode, false); // We will need to add this if we ever allow functions to be templates
+	RepopulatePinsOnNode(FunctionRefNode, false, false, false);
+	UpdateTemplateNodePinTypes(FunctionRefNode, false);
 
 	Notify(ERigVMGraphNotifType::NodeAdded, FunctionRefNode);
 
@@ -13899,12 +14316,14 @@ bool URigVMController::CanAddFunctionRefForDefinition(const FRigVMGraphFunctionH
 	return true;
 }
 
-void URigVMController::AddPinsForStruct(UStruct* InStruct, URigVMNode* InNode, URigVMPin* InParentPin, ERigVMPinDirection InPinDirection, const FString& InDefaultValue, bool bAutoExpandArrays)
+void URigVMController::AddPinsForStruct(UStruct* InStruct, URigVMNode* InNode, URigVMPin* InParentPin, ERigVMPinDirection InPinDirection, const FString& InDefaultValue, bool bAutoExpandArrays, const FPinInfoArray* PreviousPins)
 {
 	if(!ShouldStructBeUnfolded(InStruct))
 	{
 		return;
 	}
+
+	// todo: reuse the default values when creating the pins
 	
 	TArray<FString> MemberNameValuePairs = URigVMPin::SplitDefaultValue(InDefaultValue);
 	TMap<FName, FString> MemberValues;
@@ -14887,7 +15306,7 @@ bool URigVMController::ShouldRedirectPin(const FString& InOldPinPath, FString& I
 	return false;
 }
 
-void URigVMController::RepopulatePinsOnNode(URigVMNode* InNode, bool bFollowCoreRedirectors, bool bSetupOrphanedPins)
+void URigVMController::RepopulatePinsOnNode(URigVMNode* InNode, bool bFollowCoreRedirectors, bool bSetupOrphanedPins, bool bDetachAndReattachLinks)
 {
 	if (InNode == nullptr)
 	{
@@ -14907,6 +15326,7 @@ void URigVMController::RepopulatePinsOnNode(URigVMNode* InNode, bool bFollowCore
 	URigVMDispatchNode* DispatchNode = Cast<URigVMDispatchNode>(InNode);
 
 	FScopeLock Lock(&PinPathCoreRedirectorsLock);
+	const FRigVMRegistry& Registry = FRigVMRegistry::Get();
 
 	URigVMGraph* Graph = GetGraph();
 	check(Graph);
@@ -14920,22 +15340,18 @@ void URigVMController::RepopulatePinsOnNode(URigVMNode* InNode, bool bFollowCore
 		}
 	}
 
-	// step 1/3: keep a record of the current state of the node's pins
-	TMap<FString, FString> RedirectedPinPaths;
-	if (bFollowCoreRedirectors)
+	bool bSetupOrphanPinsForThisNode = bSetupOrphanedPins;
+	if(CollapseNode)
 	{
-		RedirectedPinPaths = GetRedirectedPinPaths(InNode);
+		if(CollapseNode->GetOuter()->IsA<URigVMFunctionLibrary>())
+		{
+			bSetupOrphanPinsForThisNode = false;
+		}
 	}
-	TMap<FString, FPinState> PinStates = GetPinStates(InNode);
 
-	// also in case this node is part of an injection
-	FName InjectionInputPinName = NAME_None;
-	FName InjectionOutputPinName = NAME_None;
-	if (URigVMInjectionInfo* InjectionInfo = InNode->GetInjectionInfo())
-	{
-		InjectionInputPinName = InjectionInfo->InputPin ? InjectionInfo->InputPin->GetFName() : NAME_None;
-		InjectionOutputPinName = InjectionInfo->OutputPin ? InjectionInfo->OutputPin->GetFName() : NAME_None;
-	}
+	const FPinInfoArray PreviousPinInfos(InNode);
+	const uint32 PreviousPinHash = GetTypeHash(PreviousPinInfos);
+	FPinInfoArray NewPinInfos;
 
 	// step 2/3: clear pins on the node and repopulate the node with new pins
 	if (UnitNode != nullptr)
@@ -14948,20 +15364,6 @@ void URigVMController::RepopulatePinsOnNode(URigVMNode* InNode, bool bFollowCore
 			return;
 		}
 
-		RemovePinsDuringRepopulate(UnitNode, UnitNode->Pins, bSetupOrphanedPins);
-
-		if (ScriptStruct == nullptr)
-		{
-			ReportWarningf(
-				TEXT("Control Rig '%s', Node '%s' has no struct assigned. Do you have a broken redirect?"),
-				*UnitNode->GetOutermost()->GetPathName(),
-				*UnitNode->GetName()
-				);
-
-			RemoveNode(UnitNode, false, true);
-			return;
-		}
-
 		FString NodeColorMetadata;
 		ScriptStruct->GetStringMetaDataHierarchical(*URigVMNode::NodeColorName, &NodeColorMetadata);
 		if (!NodeColorMetadata.IsEmpty())
@@ -14969,24 +15371,80 @@ void URigVMController::RepopulatePinsOnNode(URigVMNode* InNode, bool bFollowCore
 			UnitNode->NodeColor = GetColorFromMetadata(NodeColorMetadata);
 		}
 
-		FString ExportedDefaultValue;
-		CreateDefaultValueForStructIfRequired(ScriptStruct, ExportedDefaultValue);
-		AddPinsForStruct(ScriptStruct, UnitNode, nullptr, ERigVMPinDirection::Invalid, ExportedDefaultValue, false);
+		const TSharedPtr<FStructOnScope> DefaultValueContent = UnitNode->ConstructStructInstance(false);
+		NewPinInfos.AddPins(ScriptStruct, this, ERigVMPinDirection::Invalid, INDEX_NONE, DefaultValueContent->GetStructMemory());
 	}
 	else if (DispatchNode)
 	{
-		FRigVMTemplateTypeMap PinTypeMap;
+		TMap<FName, TRigVMTypeIndex> PinTypeMap;
 		for (const URigVMPin* Pin : DispatchNode->Pins)
 		{
 			PinTypeMap.Add(Pin->GetFName(), Pin->GetTypeIndex());
 		}
 		
-		RemovePinsDuringRepopulate(DispatchNode, DispatchNode->Pins, bSetupOrphanedPins);
-
 		const FRigVMTemplate* Template = DispatchNode->GetTemplate();
-		AddPinsForTemplate(Template, PinTypeMap, DispatchNode);
 		
-		ResolveTemplateNodeMetaData(DispatchNode, false);
+		for (int32 ArgIndex = 0; ArgIndex < Template->NumArguments(); ArgIndex++)
+		{
+			const FRigVMTemplateArgument* Arg = Template->GetArgument(ArgIndex);
+
+			TRigVMTypeIndex TypeIndex = INDEX_NONE;
+			if(const TRigVMTypeIndex* ExistingTypeIndex = PinTypeMap.Find(Arg->GetName()))
+			{
+				TypeIndex = *ExistingTypeIndex;
+				if(!Arg->SupportsTypeIndex(TypeIndex))
+				{
+					TypeIndex = INDEX_NONE;
+				}
+			}
+
+			if(TypeIndex == INDEX_NONE)
+			{
+				if(Arg->IsSingleton())
+				{
+					TypeIndex = Arg->GetSupportedTypeIndices()[0];
+				}
+				else if(Arg->GetArrayType() == FRigVMTemplateArgument::EArrayType_ArrayValue)
+				{
+					TypeIndex = RigVMTypeUtils::TypeIndex::WildCardArray;
+				}
+				else
+				{
+					TypeIndex = RigVMTypeUtils::TypeIndex::WildCard;
+				}
+			}
+
+			FString DefaultValue;
+			UScriptStruct* ArgumentScriptStruct = nullptr;
+			const uint8* DefaultValueMemory = nullptr;
+
+			if(const URigVMPin* ArgumentPin = DispatchNode->FindPin(Arg->Name.ToString()))
+			{
+				DefaultValue = ArgumentPin->GetDefaultValue();
+				ArgumentScriptStruct = Cast<UScriptStruct>(ArgumentPin->GetCPPTypeObject());
+			}
+			else if(const FRigVMDispatchFactory* Factory = DispatchNode->GetFactory())
+			{
+				if(Arg->IsSingleton())
+				{
+					DefaultValue = Factory->GetArgumentDefaultValue(Arg->Name, Arg->GetTypeIndices()[0]);
+					const FRigVMTemplateArgumentType& Type = Registry.GetType(Arg->GetTypeIndices()[0]);
+					ArgumentScriptStruct = Cast<UScriptStruct>(Type.CPPTypeObject);
+				}
+			}
+
+			FStructOnScope DefaultValueMemoryScope; // has to be in this scope so that DefaultValueMemory is valid
+			if(ArgumentScriptStruct && !DefaultValue.IsEmpty())
+			{
+				DefaultValueMemoryScope = FStructOnScope(ArgumentScriptStruct);
+
+				FRigVMPinDefaultValueImportErrorContext ErrorPipe;
+				ArgumentScriptStruct->ImportText(*DefaultValue, DefaultValueMemoryScope.GetStructMemory(), nullptr, PPF_None, &ErrorPipe, FString());
+				DefaultValueMemory = DefaultValueMemoryScope.GetStructMemory();
+			}
+			
+			(void)NewPinInfos.AddPin(this, INDEX_NONE, Arg->Name, Arg->GetDirection(), TypeIndex, DefaultValue, DefaultValueMemory, &PreviousPinInfos);
+		}
 	}
 	else if ((RerouteNode != nullptr) || (VariableNode != nullptr))
 	{
@@ -15048,7 +15506,7 @@ void URigVMController::RepopulatePinsOnNode(URigVMNode* InNode, bool bFollowCore
 				
 					if(RigVMTypeUtils::CPPTypeFromExternalVariable(Variable, CPPType, &CPPTypeObject))
 					{
-						RefreshVariableNode(VariableNode->GetFName(), Variable.Name, CPPType, Variable.TypeObject, false, bSetupOrphanedPins);
+						RefreshVariableNode(VariableNode->GetFName(), Variable.Name, CPPType, Variable.TypeObject, false, bSetupOrphanPinsForThisNode);
 					}
 					else
 					{
@@ -15069,38 +15527,14 @@ void URigVMController::RepopulatePinsOnNode(URigVMNode* InNode, bool bFollowCore
 				);
 			}
 		}
-		
-		RemovePinsDuringRepopulate(InNode, ValuePin->SubPins, bSetupOrphanedPins);
-
-		if (ValuePin->IsStruct())
-		{
-			UScriptStruct* ScriptStruct = ValuePin->GetScriptStruct();
-			if (ScriptStruct == nullptr)
-			{
-				ReportErrorf(
-					TEXT("Control Rig '%s', Node '%s' has no struct assigned. Do you have a broken redirect?"),
-					*InNode->GetOutermost()->GetPathName(),
-					*InNode->GetName()
-				);
 
-				RemoveNode(InNode, false, true);
-				return;
-			}
-
-			FString ExportedDefaultValue;
-			CreateDefaultValueForStructIfRequired(ScriptStruct, ExportedDefaultValue);
-			{
-				TGuardValue<bool> SuspendNotifications(bSuspendNotifications, true);
-				AddPinsForStruct(ScriptStruct, InNode, ValuePin, ValuePin->Direction, ExportedDefaultValue, false);
-			}
-		}
+		NewPinInfos = FPinInfoArray(InNode, this, &PreviousPinInfos);
 	}
 	else if (EntryNode || ReturnNode)
 	{
 		if (URigVMLibraryNode* LibraryNode = Cast<URigVMLibraryNode>(InNode->GetGraph()->GetOuter()))
 		{
 			bool bIsEntryNode = EntryNode != nullptr;
-			RemovePinsDuringRepopulate(InNode, InNode->Pins, bSetupOrphanedPins);
 
 			TArray<URigVMPin*> SortedLibraryPins;
 
@@ -15141,136 +15575,483 @@ void URigVMController::RepopulatePinsOnNode(URigVMNode* InNode, bool bFollowCore
 					}
 				}
 
-				URigVMPin* ExposedPin = NewObject<URigVMPin>(InNode, LibraryPin->GetFName());
-				ConfigurePinFromPin(ExposedPin, LibraryPin, true);
+				ERigVMPinDirection Direction = bIsEntryNode ? ERigVMPinDirection::Output : ERigVMPinDirection::Input;  
+				(void)NewPinInfos.AddPin(LibraryPin, INDEX_NONE, Direction);
+			}
+		}
+		else
+		{
+			// due to earlier bugs with copy and paste we can find entry and return nodes under the top level
+			// graph. we'll ignore these for now.
+		}
+	}
+	else if (CollapseNode)
+	{
+		NewPinInfos = FPinInfoArray(InNode, this, &PreviousPinInfos);
+	}
+	else if (FunctionRefNode)
+	{
+		const FRigVMGraphFunctionHeader& FunctionHeader = FunctionRefNode->GetReferencedFunctionHeader();
+		if(FunctionHeader.IsValid())
+		{
+			NewPinInfos = FPinInfoArray(FunctionHeader, this, &PreviousPinInfos);
+		}
+	}
+	else
+	{
+		return;
+	}
 
-				if (bIsEntryNode)
-				{
-					ExposedPin->Direction = ERigVMPinDirection::Output;
-				}
-				else
+	auto RecursivelyRepopulatePinsOnCollapseNode = [this, bFollowCoreRedirectors, bSetupOrphanedPins]
+		(URigVMCollapseNode* InCollapseNode)
+		{
+			if(InCollapseNode)
+			{
+				FRigVMControllerGraphGuard GraphGuard(this, InCollapseNode->GetContainedGraph(), false);
+				TGuardValue<bool> GuardEditGraph(InCollapseNode->ContainedGraph->bEditable, true);
+				// need to get a copy of the node array since the following function could remove nodes from the graph
+				// we don't want to remove elements from the array we are iterating over.
+				TArray<URigVMNode*> ContainedNodes = InCollapseNode->GetContainedNodes();
+				for (URigVMNode* ContainedNode : ContainedNodes)
 				{
-					ExposedPin->Direction = ERigVMPinDirection::Input;
+					RepopulatePinsOnNode(ContainedNode, bFollowCoreRedirectors, bSetupOrphanedPins);
 				}
+			}
+		};
 
-				AddNodePin(InNode, ExposedPin);
+	// if the nodes match in structure - nothing to do here
+	if(GetTypeHash(NewPinInfos) == PreviousPinHash)
+	{
+		// we at least need to recurse into a collapse node
+		RecursivelyRepopulatePinsOnCollapseNode(CollapseNode);
+		return;
+	}
+	
+#if UE_RIGVMCONTROLLER_VERBOSE_REPOPULATE
+	UE_LOG(LogRigVMDeveloper, Display, TEXT("Repopulating pins on node %s"), *InNode->GetPathName());
+#endif
+	
+	bool bRequireDetachLinks = false;
+	bool bRequirePinStates = false;
 
-				if (ExposedPin->IsStruct())
+	TArray<int32> NewPinsToAdd, PreviousPinsToRemove, PreviousPinsToOrphan, PreviousPinsToUpdate;
+	for(int32 Index = 0; Index < PreviousPinInfos.Num(); Index++)
+	{
+		const FString PinPath = PreviousPinInfos.GetPinPath(Index);
+		const int32 NewIndex = NewPinInfos.GetIndexFromPinPath(PinPath);
+		if(NewIndex == INDEX_NONE)
+		{
+			const int32 RootIndex = PreviousPinInfos.GetRootIndex(Index);
+			
+			if(PreviousPinInfos[Index].Direction != ERigVMPinDirection::Hidden)
+			{
+				if(URigVMPin* Pin = InNode->FindPin(PinPath))
 				{
-					TGuardValue<bool> SuspendNotifications(bSuspendNotifications, true);
-					AddPinsForStruct(ExposedPin->GetScriptStruct(), InNode, ExposedPin, ExposedPin->GetDirection(), FString(), false);
+					if(!Pin->GetLinks().IsEmpty())
+					{
+						bRequireDetachLinks = true;
+						bRequirePinStates = true;
+					}
+					
+					if(!PreviousPinsToOrphan.Contains(RootIndex))
+					{
+						URigVMPin* RootPin = Pin->GetRootPin();
+
+						bool bPinShouldBeOrphaned = bSetupOrphanPinsForThisNode; 
+						if(!bPinShouldBeOrphaned)
+						{
+							if(RootPin->GetSourceLinks(true).Num() > 0 ||
+								RootPin->GetTargetLinks(true).Num() > 0)
+							{
+								bPinShouldBeOrphaned = true;
+							}
+						}
+		
+						if(bPinShouldBeOrphaned)
+						{
+							PreviousPinsToOrphan.Add(RootIndex);
+							bRequireDetachLinks = true;
+							bRequirePinStates = true;
+#if UE_RIGVMCONTROLLER_VERBOSE_REPOPULATE
+							UE_LOG(LogRigVMDeveloper, Display, TEXT("Previously existing pin '%s' needs to be orphaned."), *RootPin->GetPinPath());
+#endif
+						}
+					}
 				}
+			}
 
-				Notify(ERigVMGraphNotifType::PinAdded, ExposedPin);
+			if(!PreviousPinsToOrphan.Contains(RootIndex))
+			{
+				PreviousPinsToRemove.Add(Index);
+#if UE_RIGVMCONTROLLER_VERBOSE_REPOPULATE
+				UE_LOG(LogRigVMDeveloper, Display, TEXT("Previously existing pin '%s' is now obsolete."), *PinPath);
+#endif
 			}
 		}
-		else
+		else if(GetTypeHash(PreviousPinInfos[Index]) != GetTypeHash(NewPinInfos[NewIndex]))
 		{
-			// due to earlier bugs with copy and paste we can find entry and return nodes under the top level
-			// graph. we'll ignore these for now.
+			const bool bTypesDiffer = !Registry.CanMatchTypes(PreviousPinInfos[Index].TypeIndex, NewPinInfos[NewIndex].TypeIndex, true);
+			if(PreviousPinInfos[Index].Direction != NewPinInfos[NewIndex].Direction)
+			{
+				bRequireDetachLinks = true;
+			}
+			else if(PreviousPinInfos[Index].Direction != ERigVMPinDirection::Hidden)
+			{
+				bRequireDetachLinks |= bTypesDiffer;
+			}
+
+			if(Registry.CanMatchTypes(PreviousPinInfos[Index].TypeIndex, NewPinInfos[NewIndex].TypeIndex, true))
+			{
+				PreviousPinsToUpdate.Add(Index);
+			}
+			
+#if UE_RIGVMCONTROLLER_VERBOSE_REPOPULATE
+			const FName& PreviousCPPType = Registry.GetType(PreviousPinInfos[Index].TypeIndex).CPPType;
+			const FName& NewCPPType = Registry.GetType(NewPinInfos[NewIndex].TypeIndex).CPPType;
+
+			const FString PreviousDirection = StaticEnum<ERigVMPinDirection>()->GetDisplayNameTextByValue((int64)PreviousPinInfos[Index].Direction).ToString();
+			const FString NewDirection = StaticEnum<ERigVMPinDirection>()->GetDisplayNameTextByValue((int64)NewPinInfos[NewIndex].Direction).ToString();
+
+			UE_LOG(LogRigVMDeveloper, Display,
+				TEXT("Previous pin '%s' (Index %d, %s, %s) differs with new pin (Index %d, %s, %s)."),
+				*PinPath,
+				Index,
+				*PreviousCPPType.ToString(),
+				*PreviousDirection,
+				NewIndex,
+				*NewCPPType.ToString(),
+				*NewDirection
+			);
+#endif
 		}
 	}
-	else if (CollapseNode)
+	for(int32 Index = 0; Index < NewPinInfos.Num(); Index++)
 	{
-		// PinStates were already saved earlier, and will be applied later, no need to do it here 
+		const FString PinPath = NewPinInfos.GetPinPath(Index);
+		const int32 PreviousIndex = PreviousPinInfos.GetIndexFromPinPath(PinPath);
+		if(PreviousIndex == INDEX_NONE)
+		{
+			NewPinsToAdd.Add(Index);
 
-		// since we are replacing existing pins, their names have to be saved
-		// and assigned to the new pins only after we removed the old ones
-		TArray<TTuple<URigVMPin*, FName>> NewRootPinInfos;
-		for (URigVMPin* RootPin : InNode->Pins)
+#if UE_RIGVMCONTROLLER_VERBOSE_REPOPULATE
+			UE_LOG(LogRigVMDeveloper, Display, TEXT("Newly required pin '%s' needs to be added."), *PinPath);
+#endif
+		}
+	}
+
+	auto CreatePinFromPinInfo = [this, &Registry, &PreviousPinInfos](const FPinInfo& InPinInfo, const FString& InPinPath, UObject* InOuter) -> URigVMPin*
+	{
+		URigVMPin* Pin = NewObject<URigVMPin>(InOuter, InPinInfo.Name);
+		if(InPinInfo.Property)
 		{
-			URigVMPin* NewRootPin = NewObject<URigVMPin>(InNode);
-			ConfigurePinFromPin(NewRootPin, RootPin, true);
-			EnsurePinValidity(NewRootPin, false);
-		
-			NewRootPinInfos.Add(TTuple<URigVMPin*, FName>(NewRootPin, RootPin->GetFName()));
+			ConfigurePinFromProperty(InPinInfo.Property, Pin, InPinInfo.Direction);
 		}
-		
-		RemovePinsDuringRepopulate(InNode, InNode->Pins, bSetupOrphanedPins);
-		
-		for (TTuple<URigVMPin*, FName> NewRootPinInfo : NewRootPinInfos)
+		else
 		{
-			RenameObject(NewRootPinInfo.Key, *NewRootPinInfo.Value.ToString(), InNode);
-			AddNodePin(InNode, NewRootPinInfo.Key);
+			const FRigVMTemplateArgumentType& Type = Registry.GetType(InPinInfo.TypeIndex);
+			Pin->CPPType = Type.CPPType.ToString();
+			Pin->CPPTypeObject = Type.CPPTypeObject;
+			if(Pin->CPPTypeObject)
+			{
+				Pin->CPPTypeObjectPath = *Pin->CPPTypeObject->GetPathName();
+			}
+			if(Registry.IsExecuteType(InPinInfo.TypeIndex))
+			{
+				MakeExecutePin(Pin);
+			}
+				
+			Pin->Direction = InPinInfo.Direction;
+			Pin->DisplayName = InPinInfo.DisplayName.IsEmpty() ? NAME_None : FName(*InPinInfo.DisplayName);
+			Pin->bIsConstant = InPinInfo.bIsConstant;
+			Pin->bIsDynamicArray = InPinInfo.bIsDynamicArray;
+			Pin->CustomWidgetName = InPinInfo.CustomWidgetName.IsEmpty() ? NAME_None : FName(*InPinInfo.CustomWidgetName);
 		}
-		
-		for (URigVMPin* Pin : InNode->Pins)
+
+		Pin->bIsExpanded = InPinInfo.bIsExpanded;
+		Pin->DefaultValue = InPinInfo.DefaultValue;
+
+		// reuse expansion state and default value
+		if(const FPinInfo* PreviousPin = PreviousPinInfos.GetPinFromPinPath(InPinPath))
 		{
-			if (Pin->IsStruct())
+			if(PreviousPin->TypeIndex == InPinInfo.TypeIndex)
 			{
-				TGuardValue<bool> SuspendNotifications(bSuspendNotifications, true);
-				AddPinsForStruct(Pin->GetScriptStruct(), InNode, Pin, Pin->GetDirection(), FString(), false);
+				Pin->bIsExpanded = PreviousPin->bIsExpanded;
+				Pin->DefaultValue = PreviousPin->DefaultValue;
 			}
-			Notify(ERigVMGraphNotifType::PinAdded, Pin);
 		}
 
-		if (CollapseNode->GetOuter()->IsA<URigVMFunctionLibrary>())
+		if (URigVMPin* ParentPin = Cast<URigVMPin>(InOuter))
 		{
-			// no need to notify since the function library graph is invisible anyway
-			TGuardValue<bool> SuspendNotifications(bSuspendNotifications, true);
-			RemoveUnusedOrphanedPins(CollapseNode);
+			AddSubPin(ParentPin, Pin);
+		}
+		else
+		{
+			AddNodePin(CastChecked<URigVMNode>(InOuter), Pin);
 		}
 
-		FRigVMControllerGraphGuard GraphGuard(this, CollapseNode->GetContainedGraph(), false);
-		TGuardValue<bool> GuardEditGraph(CollapseNode->ContainedGraph->bEditable, true);
-		// need to get a copy of the node array since the following function could remove nodes from the graph
-		// we don't want to remove elements from the array we are iterating over.
-		TArray<URigVMNode*> ContainedNodes = CollapseNode->GetContainedNodes();
-		for (URigVMNode* ContainedNode : ContainedNodes)
+		Notify(ERigVMGraphNotifType::PinAdded, Pin);
+
+		return Pin;
+	};
+
+	// step 1/3: keep a record of the current state of the node's pins
+	TMap<FString, FString> RedirectedPinPaths;
+	if (bFollowCoreRedirectors)
+	{
+		RedirectedPinPaths = GetRedirectedPinPaths(InNode);
+	}
+
+	// also in case this node is part of an injection
+	FName InjectionInputPinName = NAME_None;
+	FName InjectionOutputPinName = NAME_None;
+	if (URigVMInjectionInfo* InjectionInfo = InNode->GetInjectionInfo())
+	{
+		InjectionInputPinName = InjectionInfo->InputPin ? InjectionInfo->InputPin->GetFName() : NAME_None;
+		InjectionOutputPinName = InjectionInfo->OutputPin ? InjectionInfo->OutputPin->GetFName() : NAME_None;
+	}
+
+	TMap<FString, FPinState> PinStates;
+	TArray<URigVMLink*> DetachedLinks;
+
+	if(bRequirePinStates)
+	{
+		PinStates = GetPinStates(InNode);
+	}
+
+	if(bDetachAndReattachLinks && bRequireDetachLinks)
+	{
+#if UE_RIGVMCONTROLLER_VERBOSE_REPOPULATE
+		UE_LOG(LogRigVMDeveloper, Display, TEXT("Detaching links of node %s."), *InNode->GetPathName());
+#endif
+		DetachedLinks = InNode->GetLinks();
+		DetachLinksFromPinObjects(&DetachedLinks);
+	}
+
+	// we can do a simpler version of the update and simply
+	// add new pins, remove obsolete pins and change types as needed 
+#if UE_RIGVMCONTROLLER_VERBOSE_REPOPULATE
+	UE_LOG(LogRigVMDeveloper, Display, TEXT("Performing fast update of node %s ..."), *InNode->GetPathName());
+#endif
+
+	// orphan pins
+	for(int32 Index = 0; Index < PreviousPinsToOrphan.Num(); Index++)
+	{
+		const FString& PinPath = PreviousPinInfos.GetPinPath(PreviousPinsToOrphan[Index]);
+		if(URigVMPin* Pin = InNode->FindPin(PinPath))
 		{
-			RepopulatePinsOnNode(ContainedNode, bFollowCoreRedirectors, bSetupOrphanedPins);
+#if UE_RIGVMCONTROLLER_VERBOSE_REPOPULATE
+			UE_LOG(LogRigVMDeveloper, Display, TEXT("Orphaning pin '%s'."), *PinPath);
+#endif
+			check(Pin->IsRootPin());
+
+			const FString OrphanedName = FString::Printf(TEXT("%s%s"), *URigVMPin::OrphanPinPrefix, *Pin->GetName());
+			if(InNode->FindPin(OrphanedName) == nullptr)
+			{
+				Pin->DisplayName = Pin->GetFName();
+				RenameObject(Pin, *OrphanedName, nullptr);
+				InNode->Pins.Remove(Pin);
+
+				Notify(ERigVMGraphNotifType::PinRemoved, Pin);
+				InNode->OrphanedPins.Add(Pin);
+
+				Notify(ERigVMGraphNotifType::PinAdded, Pin);
+			}
+			else
+			{
+				RemovePin(Pin, false);
+			}
 		}
 	}
-	else if (FunctionRefNode)
+
+	// remove obsolete pins
+	for(int32 Index = PreviousPinsToRemove.Num() - 1; Index >= 0; Index--)
 	{
-		FRigVMGraphFunctionHeader& Header = FunctionRefNode->ReferencedFunctionHeader;
-		
-		// we want to make sure notify the graph of a potential name change
-		// when repopulating the function ref node
-		Notify(ERigVMGraphNotifType::NodeRenamed, FunctionRefNode);
-		RemovePinsDuringRepopulate(InNode, InNode->Pins, bSetupOrphanedPins);
+		const FString& PinPath = PreviousPinInfos.GetPinPath(PreviousPinsToRemove[Index]);
+		if(URigVMPin* Pin = InNode->FindPin(PinPath))
+		{
+#if UE_RIGVMCONTROLLER_VERBOSE_REPOPULATE
+			UE_LOG(LogRigVMDeveloper, Display, TEXT("Removing pin '%s'."), *PinPath);
+#endif
+			RemovePin(Pin, false);
+		}
+	}
 
-		TMap<FString, FPinState> ReferencedPinStates;
-		if (Header.LibraryPointer.LibraryNode.ResolveObject())
+	// add missing pins
+	for(int32 Index = 0; Index < NewPinsToAdd.Num(); Index++)
+	{
+		const FString& PinPath = NewPinInfos.GetPinPath(NewPinsToAdd[Index]);
+		FString ParentPinPath, PinName;
+		UObject* OuterForPin = InNode;
+		if(URigVMPin::SplitPinPathAtEnd(PinPath, ParentPinPath, PinName))
 		{
-			URigVMLibraryNode* LibraryNode = Cast<URigVMLibraryNode>(Header.LibraryPointer.LibraryNode.ResolveObject());
-			ReferencedPinStates = GetPinStates(LibraryNode);
+			OuterForPin = InNode->FindPin(ParentPinPath);
 		}
+		
+		(void)CreatePinFromPinInfo(NewPinInfos[NewPinsToAdd[Index]], PinPath, OuterForPin);
+#if UE_RIGVMCONTROLLER_VERBOSE_REPOPULATE
+		UE_LOG(LogRigVMDeveloper, Display, TEXT("Adding new pin '%s'."), *PinPath);
+#endif
+	}
 
-		for (const FRigVMGraphFunctionArgument& ReferencedArgument : Header.Arguments)
+	// update existing pins
+	for(int32 Index = 0; Index < PreviousPinsToUpdate.Num(); Index++)
+	{
+		const FString& PinPath = PreviousPinInfos.GetPinPath(PreviousPinsToUpdate[Index]);
+		const FPinInfo* NewPinInfo = NewPinInfos.GetPinFromPinPath(PinPath);
+		check(NewPinInfo);
+		
+		if(URigVMPin* Pin = InNode->FindPin(PinPath))
 		{
-			URigVMPin* NewPin = NewObject<URigVMPin>(InNode, ReferencedArgument.Name);
-			ConfigurePinFromArgument(NewPin, ReferencedArgument, true);
-			EnsurePinValidity(NewPin, false);
+			if(Pin->IsExecuteContext())
+			{
+				MakeExecutePin(Pin);
+			}
 			
-			AddNodePin(InNode, NewPin);
+			if(Pin->GetTypeIndex() != NewPinInfo->TypeIndex)
+			{
+				// we expect these changes to only apply to float and double pins.
+				check((NewPinInfo->TypeIndex == RigVMTypeUtils::TypeIndex::Float) ||
+					(NewPinInfo->TypeIndex == RigVMTypeUtils::TypeIndex::FloatArray) ||
+					(NewPinInfo->TypeIndex == RigVMTypeUtils::TypeIndex::Double) ||
+					(NewPinInfo->TypeIndex == RigVMTypeUtils::TypeIndex::DoubleArray));
+
+				const FRigVMTemplateArgumentType& NewType = Registry.GetType(NewPinInfo->TypeIndex);
+#if UE_RIGVMCONTROLLER_VERBOSE_REPOPULATE
+				UE_LOG(LogRigVMDeveloper, Display, TEXT("Changing pin '%s' type from %s to %s."),
+					*PinPath,
+					*Pin->CPPType,
+					*NewType.CPPType.ToString()
+				);
+#endif
+				Pin->CPPType = NewType.CPPType.ToString();
+				Pin->CPPTypeObject = NewType.CPPTypeObject;
+				check(Pin->CPPTypeObject == nullptr);
+				Pin->CPPTypeObjectPath = NAME_None;
+				Pin->LastKnownTypeIndex = NewPinInfo->TypeIndex;
 
-			if (NewPin->IsStruct())
+				Notify(ERigVMGraphNotifType::PinTypeChanged, Pin);
+			}
+		}
+	}
+
+	// create a map representing the order of expected pins
+	TMap<FString, TArray<FName>> PinOrder;
+	for(int32 Index = 0; Index < NewPinInfos.Num(); Index++)
+	{
+		const FPinInfo& NewPin = NewPinInfos[Index];
+		FString ParentPinPath;
+		if(NewPin.ParentIndex != INDEX_NONE)
+		{
+			ParentPinPath = NewPinInfos.GetPinPath(NewPin.ParentIndex);
+			if(NewPinInfos[NewPin.ParentIndex].bIsArray)
 			{
-				TGuardValue<bool> SuspendNotifications(bSuspendNotifications, true);
-				AddPinsForStruct(NewPin->GetScriptStruct(), InNode, NewPin, NewPin->GetDirection(), FString(), false);
+				continue;
 			}
+		}
+
+		TArray<FName>& SubPinOrder = PinOrder.FindOrAdd(ParentPinPath);
+		SubPinOrder.Add(NewPin.Name);
+	}
+
+	auto SortPinArray = [this](TArray<TObjectPtr<URigVMPin>>& Pins, const TArray<FName>* PinOrder)
+	{
+		if(PinOrder == nullptr)
+		{
+			return;
+		}
+
+		if(Pins.Num() < 2)
+		{
+			return;
+		}
 
-			Notify(ERigVMGraphNotifType::PinAdded, NewPin);
+		const TArray<TObjectPtr<URigVMPin>> PreviousPins = Pins;
+
+		if(Pins[0]->IsArrayElement())
+		{
+			Algo::Sort(Pins, [PinOrder](const TObjectPtr<URigVMPin>& A, const TObjectPtr<URigVMPin>& B) -> bool
+			{
+				return A->GetFName().Compare(B->GetFName()) < 0;
+			});
+		}
+		else
+		{
+			Algo::Sort(Pins, [PinOrder](const TObjectPtr<URigVMPin>& A, const TObjectPtr<URigVMPin>& B) -> bool
+			{
+				const int32 IndexA = PinOrder->Find(A->GetFName());
+				const int32 IndexB = PinOrder->Find(B->GetFName());
+				return IndexA < IndexB;
+			});
 		}
 
-		if (ReferencedPinStates.Num() > 0)
+		for(int32 Index=0;Index<Pins.Num();Index++)
 		{
-			ApplyPinStates(InNode, ReferencedPinStates);
+			if(PreviousPins[Index] != Pins[Index])
+			{
+				Notify(ERigVMGraphNotifType::PinIndexChanged, Pins[Index]);
+#if UE_RIGVMCONTROLLER_VERBOSE_REPOPULATE
+				UE_LOG(LogRigVMDeveloper, Display, TEXT("Pin '%s' changed index from %d to %d."),
+					*Pins[Index]->GetPinPath(),
+					PreviousPins.Find(Pins[Index]),
+					Index
+				);
+#endif
+			}
 		}
+	};
+
+	SortPinArray(InNode->Pins, PinOrder.Find(FString()));
+	for(URigVMPin* Pin : InNode->Pins)
+	{
+		SortPinArray(Pin->SubPins, PinOrder.Find(Pin->GetPinPath()));
 	}
-	else
+	
+	if(DispatchNode)
 	{
-		return;
+		ResolveTemplateNodeMetaData(DispatchNode, false);
 	}
+	else if(CollapseNode)
+	{
+		if (CollapseNode->GetOuter()->IsA<URigVMFunctionLibrary>())
+		{
+			// no need to notify since the function library graph is invisible anyway
+			RemoveUnusedOrphanedPins(CollapseNode);
+		}
 
-	ApplyPinStates(InNode, PinStates, RedirectedPinPaths);
+		RecursivelyRepopulatePinsOnCollapseNode(CollapseNode);
+	}
+	else if (FunctionRefNode)
+	{
+		// we want to make sure notify the graph of a potential name change
+		// when repopulating the function ref node
+		Notify(ERigVMGraphNotifType::NodeRenamed, FunctionRefNode);
+	}
+
+	if(!PinStates.IsEmpty())
+	{
+#if UE_RIGVMCONTROLLER_VERBOSE_REPOPULATE
+		UE_LOG(LogRigVMDeveloper, Display, TEXT("Reapplying pin-states of node %s..."), *InNode->GetPathName());
+#endif
+		ApplyPinStates(InNode, PinStates, RedirectedPinPaths);
+	}
+
+	if(!DetachedLinks.IsEmpty())
+	{
+#if UE_RIGVMCONTROLLER_VERBOSE_REPOPULATE
+		UE_LOG(LogRigVMDeveloper, Display, TEXT("Reattaching links of node %s..."), *InNode->GetPathName());
+#endif
+		ReattachLinksToPinObjects(bFollowCoreRedirectors, &DetachedLinks, bSetupOrphanPinsForThisNode);
+	}
 
 	if (URigVMInjectionInfo* InjectionInfo = InNode->GetInjectionInfo())
 	{
 		InjectionInfo->InputPin = InNode->FindPin(InjectionInputPinName.ToString());
 		InjectionInfo->OutputPin = InNode->FindPin(InjectionOutputPinName.ToString());
 	}
+
+#if UE_RIGVMCONTROLLER_VERBOSE_REPOPULATE
+	UE_LOG(LogRigVMDeveloper, Display, TEXT("Repopulate of node %s is completed.\n"), *InNode->GetPathName());
+#endif
 }
 
 void URigVMController::RemovePinsDuringRepopulate(URigVMNode* InNode, TArray<URigVMPin*>& InPins, bool bSetupOrphanedPins)
@@ -15355,11 +16136,12 @@ bool URigVMController::RemoveUnusedOrphanedPins(URigVMNode* InNode)
 	{
 		return true;
 	}
-	
+
+	TArray<URigVMPin*> OrphanedPins = InNode->OrphanedPins; 
 	TArray<URigVMPin*> RemainingOrphanPins;
-	for(int32 PinIndex=0; PinIndex < InNode->OrphanedPins.Num(); PinIndex++)
+	for(int32 PinIndex=0; PinIndex < OrphanedPins.Num(); PinIndex++)
 	{
-		URigVMPin* OrphanedPin = InNode->OrphanedPins[PinIndex];
+		URigVMPin* OrphanedPin = OrphanedPins[PinIndex];
 
 		const int32 NumSourceLinks = OrphanedPin->GetSourceLinks(true).Num(); 
 		const int32 NumTargetLinks = OrphanedPin->GetTargetLinks(true).Num();
@@ -16970,6 +17752,10 @@ bool URigVMController::UpdateTemplateNodePinTypes(URigVMTemplateNode* InNode, bo
 	bool bAnyTypeChanged = false;
 
 	const FRigVMTemplate* Template = InNode->GetTemplate();
+	if(Template == nullptr)
+	{
+		return false;
+	}
 
 	// Because interface nodes reference the same template, they need to be resolved at the same time (otherwise we
 	// might end up in a situation where they resolve to different permutations).
@@ -17040,7 +17826,7 @@ bool URigVMController::UpdateTemplateNodePinTypes(URigVMTemplateNode* InNode, bo
 			Pins.Append(ReturnNode->GetPins());
 		}
 	}
-	
+
 	TArray<const FRigVMTemplateArgument*> Arguments;
 	Arguments.SetNumUninitialized(Pins.Num());
 	for(int32 PinIndex=0; PinIndex < Pins.Num(); ++PinIndex)
@@ -19125,12 +19911,7 @@ void URigVMController::RefreshFunctionReferences(URigVMLibraryNode* InFunctionDe
 		FunctionLibrary->ForEachReference(InFunctionDefinition->GetFName(), [this, bSetupUndoRedo](URigVMFunctionReferenceNode* ReferenceNode)
 		{
 			FRigVMControllerGraphGuard GraphGuard(this, ReferenceNode->GetGraph(), bSetupUndoRedo);
-
-			TArray<URigVMLink*> Links = ReferenceNode->GetLinks();
-			DetachLinksFromPinObjects(&Links);
-			RepopulatePinsOnNode(ReferenceNode, false);
-			TGuardValue<bool> ReportGuard(bReportWarningsAndErrors, false);
-			ReattachLinksToPinObjects(false, &Links);
+			RepopulatePinsOnNode(ReferenceNode, false, false, true);
 		});
 	}
 }
diff --git a/Engine/Source/Developer/RigVMDeveloper/Private/RigVMModel/RigVMPin.cpp b/Engine/Source/Developer/RigVMDeveloper/Private/RigVMModel/RigVMPin.cpp
--- a/Engine/Source/Developer/RigVMDeveloper/Private/RigVMModel/RigVMPin.cpp
+++ b/Engine/Source/Developer/RigVMDeveloper/Private/RigVMModel/RigVMPin.cpp
@@ -1544,44 +1544,50 @@ TArray<URigVMPin*> URigVMPin::GetLinkedTargetPins(bool bRecursive) const
 TArray<URigVMLink*> URigVMPin::GetSourceLinks(bool bRecursive) const
 {
 	TArray<URigVMLink*> Results;
-	for (URigVMLink* Link : Links)
+	if(GetDirection() == ERigVMPinDirection::IO ||
+		GetDirection() == ERigVMPinDirection::Input)
 	{
-		if (Link->GetTargetPin() == this)
+		for (URigVMLink* Link : Links)
 		{
-			Results.Add(Link);
+			if (Link->GetTargetPin() == this)
+			{
+				Results.Add(Link);
+			}
 		}
-	}
 
-	if (bRecursive)
-	{
-		for (URigVMPin* SubPin : SubPins)
+		if (bRecursive)
 		{
-			Results.Append(SubPin->GetSourceLinks(bRecursive));
+			for (URigVMPin* SubPin : SubPins)
+			{
+				Results.Append(SubPin->GetSourceLinks(bRecursive));
+			}
 		}
 	}
-
 	return Results;
 }
 
 TArray<URigVMLink*> URigVMPin::GetTargetLinks(bool bRecursive) const
 {
 	TArray<URigVMLink*> Results;
-	for (URigVMLink* Link : Links)
+	if(GetDirection() == ERigVMPinDirection::IO ||
+		GetDirection() == ERigVMPinDirection::Output)
 	{
-		if (Link->GetSourcePin() == this)
+		for (URigVMLink* Link : Links)
 		{
-			Results.Add(Link);
+			if (Link->GetSourcePin() == this)
+			{
+				Results.Add(Link);
+			}
 		}
-	}
 
-	if (bRecursive)
-	{
-		for (URigVMPin* SubPin : SubPins)
+		if (bRecursive)
 		{
-			Results.Append(SubPin->GetTargetLinks(bRecursive));
+			for (URigVMPin* SubPin : SubPins)
+			{
+				Results.Append(SubPin->GetTargetLinks(bRecursive));
+			}
 		}
 	}
-
 	return Results;
 }
 
diff --git a/Engine/Source/Developer/RigVMDeveloper/Public/RigVMModel/RigVMController.h b/Engine/Source/Developer/RigVMDeveloper/Public/RigVMModel/RigVMController.h
--- a/Engine/Source/Developer/RigVMDeveloper/Public/RigVMModel/RigVMController.h
+++ b/Engine/Source/Developer/RigVMDeveloper/Public/RigVMModel/RigVMController.h
@@ -915,7 +915,7 @@ public:
 	bool ShouldRedirectPin(UScriptStruct* InOwningStruct, const FString& InOldRelativePinPath, FString& InOutNewRelativePinPath) const;
 	bool ShouldRedirectPin(const FString& InOldPinPath, FString& InOutNewPinPath) const;
 
-	void RepopulatePinsOnNode(URigVMNode* InNode, bool bFollowCoreRedirectors = true, bool bSetupOrphanedPins = false);
+	void RepopulatePinsOnNode(URigVMNode* InNode, bool bFollowCoreRedirectors = true, bool bSetupOrphanedPins = false, bool bDetachAndReattachLinks = false);
 	void RemovePinsDuringRepopulate(URigVMNode* InNode, TArray<URigVMPin*>& InPins, bool bSetupOrphanedPins);
 
 	// removes any orphan pins that no longer holds a link
@@ -1007,6 +1007,63 @@ public:
 	UFUNCTION(BlueprintCallable, Category = RigVMController)
 	void SetIsRunningUnitTest(bool bIsRunning);
 
+public:
+	
+	struct FPinInfo
+	{
+		FPinInfo();
+		FPinInfo(const URigVMPin* InPin, int32 InParentIndex, ERigVMPinDirection InDirection = ERigVMPinDirection::Invalid);
+		FPinInfo(FProperty* InProperty, ERigVMPinDirection InDirection, int32 InParentIndex, const uint8* InDefaultValueMemory);
+
+		void CorrectExecuteTypeIndex();
+		
+		int32 ParentIndex;
+		FName Name;
+		ERigVMPinDirection Direction;
+		TRigVMTypeIndex TypeIndex;
+		bool bIsArray;
+		FProperty* Property;
+		FString PinPath;
+		FString DefaultValue;
+		FString DisplayName;
+		FString CustomWidgetName;
+		bool bIsExpanded;
+		bool bIsConstant;
+		bool bIsDynamicArray;
+		TArray<int32> SubPins;
+
+		friend uint32 GetTypeHash(const FPinInfo& InPin);
+	};
+
+	struct FPinInfoArray
+	{
+		FPinInfoArray() {}
+		FPinInfoArray(const URigVMNode* InNode);
+		FPinInfoArray(const URigVMNode* InNode, URigVMController* InController, const FPinInfoArray* InPreviousPinInfos = nullptr);
+		FPinInfoArray(const FRigVMGraphFunctionHeader& FunctionHeader, URigVMController* InController, const FPinInfoArray* InPreviousPinInfos = nullptr);
+
+		int32 Num() const { return Pins.Num(); }
+		const FPinInfo& operator[](int32 InIndex) const { return Pins[InIndex]; }
+		FPinInfo& operator[](int32 InIndex) { return Pins[InIndex]; }
+		TArray<FPinInfo>::RangedForIteratorType begin() const { return Pins.begin(); }
+		TArray<FPinInfo>::RangedForIteratorType end() const { return Pins.end(); }
+
+		int32 AddPin(const URigVMPin* InPin, int32 InParentIndex, ERigVMPinDirection InDirection = ERigVMPinDirection::Invalid);
+		int32 AddPin(FProperty* InProperty, URigVMController* InController, ERigVMPinDirection InDirection, int32 InParentIndex, const uint8* InDefaultValueMemory);
+		int32 AddPin(URigVMController* InController, int32 InParentIndex, const FName& InName, ERigVMPinDirection InDirection, TRigVMTypeIndex InTypeIndex, const FString& InDefaultValue, const uint8* InDefaultValueMemory, const FPinInfoArray* InPreviousPinInfos);
+		void AddPins(UScriptStruct* InScriptStruct, URigVMController* InController, ERigVMPinDirection InDirection, int32 InParentIndex, const uint8* InDefaultValueMemory);
+
+		const FString& GetPinPath(const int32 InIndex) const;
+		int32 GetIndexFromPinPath(const FString& InPinPath) const;
+		const FPinInfo* GetPinFromPinPath(const FString& InPinPath) const;
+		int32 GetRootIndex(const int32 InIndex) const;
+
+		friend uint32 GetTypeHash(const FPinInfoArray& InPins);
+		
+		mutable TArray<FPinInfo> Pins;
+		mutable TMap<FString, int32> PinPathLookup;;
+	};
+
 private:
 
 	UPROPERTY(BlueprintReadOnly, Category = RigVMController, meta = (ScriptName = "ModifiedEvent", AllowPrivateAccess = "true"))
@@ -1025,7 +1082,7 @@ private:
 	TObjectPtr<URigVMNode> FindEventNode(const UScriptStruct* InScriptStruct) const;
 	bool CanAddEventNode(UScriptStruct* InScriptStruct, const bool bReportErrors) const;
 	bool CanAddFunctionRefForDefinition(const FRigVMGraphFunctionHeader& InFunctionDefinition, bool bReportErrors, bool bAllowPrivateFunctions=false);
-	void AddPinsForStruct(UStruct* InStruct, URigVMNode* InNode, URigVMPin* InParentPin, ERigVMPinDirection InPinDirection, const FString& InDefaultValue, bool bAutoExpandArrays);
+	void AddPinsForStruct(UStruct* InStruct, URigVMNode* InNode, URigVMPin* InParentPin, ERigVMPinDirection InPinDirection, const FString& InDefaultValue, bool bAutoExpandArrays, const FPinInfoArray* PreviousPins = nullptr);
 	void AddPinsForArray(FArrayProperty* InArrayProperty, URigVMNode* InNode, URigVMPin* InParentPin, ERigVMPinDirection InPinDirection, const TArray<FString>& InDefaultValues, bool bAutoExpandArrays);
 	void AddPinsForTemplate(const FRigVMTemplate* InTemplate, const FRigVMTemplateTypeMap& InPinTypeMap, URigVMNode* InNode);
 	void ConfigurePinFromProperty(FProperty* InProperty, URigVMPin* InOutPin, ERigVMPinDirection InPinDirection = ERigVMPinDirection::Invalid);
@@ -1297,6 +1354,7 @@ private:
 	friend class FRigVMParserAST;
 	friend class FRigVMControllerCompileBracketScope;
 	friend class FRigVMControllerGraphGuard;
+	friend struct FPinInfoArray;
 	friend class FRigVMControllerNotifGuard;
 	friend struct FRigVMClient;
 };
diff --git a/Engine/Source/Runtime/RigVM/Public/RigVMDefines.h b/Engine/Source/Runtime/RigVM/Public/RigVMDefines.h
--- a/Engine/Source/Runtime/RigVM/Public/RigVMDefines.h
+++ b/Engine/Source/Runtime/RigVM/Public/RigVMDefines.h
@@ -29,3 +29,7 @@
 #ifndef UE_RIGVM_DEBUG_EXECUTION
 #define UE_RIGVM_DEBUG_EXECUTION 0
 #endif
+
+#ifndef UE_RIGVMCONTROLLER_VERBOSE_REPOPULATE
+#define UE_RIGVMCONTROLLER_VERBOSE_REPOPULATE UE_BUILD_DEBUG
+#endif
