Control Rig: Fix repopulate pins on node not changing types when needed

diff --git a/Engine/Source/Developer/RigVMDeveloper/Private/RigVMModel/RigVMController.cpp b/Engine/Source/Developer/RigVMDeveloper/Private/RigVMModel/RigVMController.cpp
--- a/Engine/Source/Developer/RigVMDeveloper/Private/RigVMModel/RigVMController.cpp
+++ b/Engine/Source/Developer/RigVMDeveloper/Private/RigVMModel/RigVMController.cpp
@@ -15704,6 +15704,11 @@ void URigVMController::RepopulatePinsOnNode(URigVMNode* InNode, bool bFollowCore
 			{
 				PreviousPinsToUpdate.Add(Index);
 			}
+			else
+            {
+            	PreviousPinsToRemove.Add(Index);
+            	NewPinsToAdd.Add(NewIndex);
+            }
 			
 #if UE_RIGVMCONTROLLER_VERBOSE_REPOPULATE
 			const FName& PreviousCPPType = Registry.GetType(PreviousPinInfos[Index].TypeIndex).CPPType;
