[Backout] - CL24116737
Original CL Desc
-----------------------------------------------------------------
RigVM: Avoid recreating pins if not required during RepopulatePins

diff --git a/Engine/Source/Developer/RigVMDeveloper/Private/RigVMModel/RigVMPin.cpp b/Engine/Source/Developer/RigVMDeveloper/Private/RigVMModel/RigVMPin.cpp
--- a/Engine/Source/Developer/RigVMDeveloper/Private/RigVMModel/RigVMPin.cpp
+++ b/Engine/Source/Developer/RigVMDeveloper/Private/RigVMModel/RigVMPin.cpp
@@ -1544,50 +1544,44 @@ TArray<URigVMPin*> URigVMPin::GetLinkedTargetPins(bool bRecursive) const
 TArray<URigVMLink*> URigVMPin::GetSourceLinks(bool bRecursive) const
 {
 	TArray<URigVMLink*> Results;
-	if(GetDirection() == ERigVMPinDirection::IO ||
-		GetDirection() == ERigVMPinDirection::Input)
+	for (URigVMLink* Link : Links)
 	{
-		for (URigVMLink* Link : Links)
+		if (Link->GetTargetPin() == this)
 		{
-			if (Link->GetTargetPin() == this)
-			{
-				Results.Add(Link);
-			}
+			Results.Add(Link);
 		}
+	}
 
-		if (bRecursive)
+	if (bRecursive)
+	{
+		for (URigVMPin* SubPin : SubPins)
 		{
-			for (URigVMPin* SubPin : SubPins)
-			{
-				Results.Append(SubPin->GetSourceLinks(bRecursive));
-			}
+			Results.Append(SubPin->GetSourceLinks(bRecursive));
 		}
 	}
+
 	return Results;
 }
 
 TArray<URigVMLink*> URigVMPin::GetTargetLinks(bool bRecursive) const
 {
 	TArray<URigVMLink*> Results;
-	if(GetDirection() == ERigVMPinDirection::IO ||
-		GetDirection() == ERigVMPinDirection::Output)
+	for (URigVMLink* Link : Links)
 	{
-		for (URigVMLink* Link : Links)
+		if (Link->GetSourcePin() == this)
 		{
-			if (Link->GetSourcePin() == this)
-			{
-				Results.Add(Link);
-			}
+			Results.Add(Link);
 		}
+	}
 
-		if (bRecursive)
+	if (bRecursive)
+	{
+		for (URigVMPin* SubPin : SubPins)
 		{
-			for (URigVMPin* SubPin : SubPins)
-			{
-				Results.Append(SubPin->GetTargetLinks(bRecursive));
-			}
+			Results.Append(SubPin->GetTargetLinks(bRecursive));
 		}
 	}
+
 	return Results;
 }
 
diff --git a/Engine/Source/Developer/RigVMDeveloper/Public/RigVMModel/RigVMController.h b/Engine/Source/Developer/RigVMDeveloper/Public/RigVMModel/RigVMController.h
--- a/Engine/Source/Developer/RigVMDeveloper/Public/RigVMModel/RigVMController.h
+++ b/Engine/Source/Developer/RigVMDeveloper/Public/RigVMModel/RigVMController.h
@@ -9,21 +9,25 @@
 #include "RigVMModel/Nodes/RigVMParameterNode.h"
 #include "RigVMModel/Nodes/RigVMCommentNode.h"
 #include "RigVMModel/Nodes/RigVMRerouteNode.h"
+#include "RigVMModel/Nodes/RigVMBranchNode.h"
 #include "RigVMModel/Nodes/RigVMIfNode.h"
 #include "RigVMModel/Nodes/RigVMSelectNode.h"
 #include "RigVMModel/Nodes/RigVMTemplateNode.h"
 #include "RigVMModel/Nodes/RigVMEnumNode.h"
 #include "RigVMModel/Nodes/RigVMCollapseNode.h"
 #include "RigVMModel/Nodes/RigVMFunctionReferenceNode.h"
+#include "RigVMModel/Nodes/RigVMArrayNode.h"
 #include "RigVMModel/Nodes/RigVMInvokeEntryNode.h"
 #include "RigVMModel/RigVMBuildData.h"
 #include "RigVMCore/RigVMUserWorkflow.h"
 #include "UObject/Interface.h"
 #include "RigVMController.generated.h"
 
+#ifndef UE_RIGVM_ENABLE_TEMPLATE_NODES
+#define UE_RIGVM_ENABLE_TEMPLATE_NODES 1
+#endif
+
 class URigVMActionStack;
-struct FRigVMGraphFunctionArgument;
-struct FRigVMGraphFunctionHeader;
 
 UENUM()
 enum class ERigVMControllerBulkEditType : uint8
@@ -78,9 +82,10 @@ private:
 DECLARE_DELEGATE_RetVal_OneParam(bool, FRigVMController_ShouldStructUnfoldDelegate, const UStruct*)
 DECLARE_DELEGATE_RetVal_OneParam(TArray<FRigVMExternalVariable>, FRigVMController_GetExternalVariablesDelegate, URigVMGraph*)
 DECLARE_DELEGATE_RetVal(const FRigVMByteCode*, FRigVMController_GetByteCodeDelegate)
-DECLARE_DELEGATE_RetVal_OneParam(bool, FRigVMController_RequestLocalizeFunctionDelegate, FRigVMGraphFunctionIdentifier&)
+DECLARE_DELEGATE_RetVal_OneParam(bool, FRigVMController_IsFunctionAvailableDelegate, URigVMLibraryNode*)
+DECLARE_DELEGATE_RetVal_OneParam(bool, FRigVMController_RequestLocalizeFunctionDelegate, URigVMLibraryNode*)
 DECLARE_DELEGATE_RetVal_ThreeParams(FName, FRigVMController_RequestNewExternalVariableDelegate, FRigVMGraphVariableDescription, bool, bool);
-DECLARE_DELEGATE_RetVal_TwoParams(bool, FRigVMController_IsDependencyCyclicDelegate, const FRigVMGraphFunctionHeader& Dependent, const FRigVMGraphFunctionHeader& Dependency)
+DECLARE_DELEGATE_RetVal_TwoParams(bool, FRigVMController_IsDependencyCyclicDelegate, UObject*, UObject*)
 DECLARE_DELEGATE_RetVal_TwoParams(FRigVMController_BulkEditResult, FRigVMController_RequestBulkEditDialogDelegate, URigVMLibraryNode*, ERigVMControllerBulkEditType)
 DECLARE_DELEGATE_RetVal_OneParam(bool, FRigVMController_RequestBreakLinksDialogDelegate, TArray<URigVMLink*>)
 DECLARE_DELEGATE_FiveParams(FRigVMController_OnBulkEditProgressDelegate, TSoftObjectPtr<URigVMFunctionReferenceNode>, ERigVMControllerBulkEditType, ERigVMControllerBulkEditProgress, int32, int32)
@@ -125,44 +130,6 @@ protected:
 	const uint8* Memory;
 };
 
-// A struct describing the result of a backwards compatibility patch
-USTRUCT()
-struct RIGVMDEVELOPER_API FRigVMClientPatchResult
-{
-public:
-
-	GENERATED_BODY()
-	
-	FRigVMClientPatchResult()
-		: bSucceeded(true)
-		, bChangedContent(false)
-		, bRequiresToMarkPackageDirty(false)
-	{}
-
-	bool Succeeded() const { return bSucceeded; }
-	bool ChangedContent() const { return bChangedContent; }
-	bool RequiresToMarkPackageDirty() const { return bRequiresToMarkPackageDirty; }
-
-	const TArray<FString>& GetErrorMessages() const { return ErrorMessages; }
-	const TArray<FString>& GetRemovedNodes() const { return RemovedNodes; }
-	const TArray<TWeakObjectPtr<const URigVMNode>>& GetAddedNodes() const { return AddedNodes; }
-
-private:
-
-	void Merge(const FRigVMClientPatchResult& InOther);
-	
-	bool bSucceeded;
-	bool bChangedContent;
-	bool bRequiresToMarkPackageDirty;
-
-	TArray<FString> ErrorMessages;
-	TArray<FString> RemovedNodes;
-	TArray<TWeakObjectPtr<const URigVMNode>> AddedNodes;
-
-	friend struct FRigVMClient;
-	friend class URigVMController;
-};
-
 /**
  * The Controller is the sole authority to perform changes
  * on the Graph. The Controller itself is stateless.
@@ -376,21 +343,21 @@ public:
 	// Branch nodes can be used to split the execution of into multiple branches,
 	// allowing to drive behavior by logic.
 	UFUNCTION(BlueprintCallable, Category = RigVMController)
-	URigVMNode* AddBranchNode(const FVector2D& InPosition = FVector2D::ZeroVector, const FString& InNodeName = TEXT(""), bool bSetupUndoRedo = true, bool bPrintPythonCommand = false);
+	URigVMBranchNode* AddBranchNode(const FVector2D& InPosition = FVector2D::ZeroVector, const FString& InNodeName = TEXT(""), bool bSetupUndoRedo = true, bool bPrintPythonCommand = false);
 
 	// Adds an if node to the graph.
 	// If nodes can be used to pick between two values based on a condition.
 	UFUNCTION(BlueprintCallable, Category = RigVMController)
-	URigVMNode* AddIfNode(const FString& InCPPType, const FName& InCPPTypeObjectPath, const FVector2D& InPosition = FVector2D::ZeroVector, const FString& InNodeName = TEXT(""), bool bSetupUndoRedo = true, bool bPrintPythonCommand = false);
+	URigVMIfNode* AddIfNode(const FString& InCPPType, const FName& InCPPTypeObjectPath, const FVector2D& InPosition = FVector2D::ZeroVector, const FString& InNodeName = TEXT(""), bool bSetupUndoRedo = true, bool bPrintPythonCommand = false);
 	UFUNCTION(BlueprintCallable, Category = RigVMController)
-	URigVMNode* AddIfNodeFromStruct(UScriptStruct* InScriptStruct, const FVector2D& InPosition = FVector2D::ZeroVector, const FString& InNodeName = TEXT(""), bool bSetupUndoRedo = true);
+	URigVMIfNode* AddIfNodeFromStruct(UScriptStruct* InScriptStruct, const FVector2D& InPosition = FVector2D::ZeroVector, const FString& InNodeName = TEXT(""), bool bSetupUndoRedo = true);
 
 	// Adds a select node to the graph.
 	// Select nodes can be used to pick between multiple values based on an index.
 	UFUNCTION(BlueprintCallable, Category = RigVMController)
-	URigVMNode* AddSelectNode(const FString& InCPPType, const FName& InCPPTypeObjectPath, const FVector2D& InPosition = FVector2D::ZeroVector, const FString& InNodeName = TEXT(""), bool bSetupUndoRedo = true, bool bPrintPythonCommand = false);
+	URigVMSelectNode* AddSelectNode(const FString& InCPPType, const FName& InCPPTypeObjectPath, const FVector2D& InPosition = FVector2D::ZeroVector, const FString& InNodeName = TEXT(""), bool bSetupUndoRedo = true, bool bPrintPythonCommand = false);
 	UFUNCTION(BlueprintCallable, Category = RigVMController)
-	URigVMNode* AddSelectNodeFromStruct(UScriptStruct* InScriptStruct, const FVector2D& InPosition = FVector2D::ZeroVector, const FString& InNodeName = TEXT(""), bool bSetupUndoRedo = true);
+	URigVMSelectNode* AddSelectNodeFromStruct(UScriptStruct* InScriptStruct, const FVector2D& InPosition = FVector2D::ZeroVector, const FString& InNodeName = TEXT(""), bool bSetupUndoRedo = true);
 
 	// Adds a template node to the graph.
 	UFUNCTION(BlueprintCallable, Category = RigVMController)
@@ -448,12 +415,12 @@ public:
 	// Adds a Array Node to the edited Graph.
 	// This causes a NodeAdded modified event.
 	UFUNCTION(BlueprintCallable, Category = RigVMController)
-	URigVMNode* AddArrayNode(ERigVMOpCode InOpCode, const FString& InCPPType, UObject* InCPPTypeObject, const FVector2D& InPosition = FVector2D::ZeroVector, const FString& InNodeName = TEXT(""), bool bSetupUndoRedo = true, bool bPrintPythonCommand = false, bool bIsPatching = false);
+	URigVMArrayNode* AddArrayNode(ERigVMOpCode InOpCode, const FString& InCPPType, UObject* InCPPTypeObject, const FVector2D& InPosition = FVector2D::ZeroVector, const FString& InNodeName = TEXT(""), bool bSetupUndoRedo = true, bool bPrintPythonCommand = false);
 
 	// Adds a Array Node to the edited Graph given a struct object path name.
 	// This causes a NodeAdded modified event.
 	UFUNCTION(BlueprintCallable, Category = RigVMController)
-	URigVMNode* AddArrayNodeFromObjectPath(ERigVMOpCode InOpCode, const FString& InCPPType, const FString& InCPPTypeObjectPath, const FVector2D& InPosition = FVector2D::ZeroVector, const FString& InNodeName = TEXT(""), bool bSetupUndoRedo = true, bool bPrintPythonCommand = false, bool bIsPatching = false);
+	URigVMArrayNode* AddArrayNodeFromObjectPath(ERigVMOpCode InOpCode, const FString& InCPPType, const FString& InCPPTypeObjectPath, const FVector2D& InPosition = FVector2D::ZeroVector, const FString& InNodeName = TEXT(""), bool bSetupUndoRedo = true, bool bPrintPythonCommand = false);
 
 	// Adds an entry invocation node
 	// This causes a NodeAdded modified event.
@@ -509,23 +476,19 @@ public:
 
 	// Copies a function declaration into this graph's local function library
 	UFUNCTION(BlueprintCallable, Category = RigVMController)
-	URigVMLibraryNode* LocalizeFunctionFromPath(const FString& InHostPath, const FName& InFunctionName, bool bLocalizeDependentPrivateFunctions = true, bool bSetupUndoRedo = true, bool bPrintPythonCommand = false);
-
-	UFUNCTION(BlueprintCallable, Category = RigVMController)
     URigVMLibraryNode* LocalizeFunction(
-		const FRigVMGraphFunctionIdentifier& InFunctionDefinition,
+    	URigVMLibraryNode* InFunctionDefinition,
 		bool bLocalizeDependentPrivateFunctions = true,
-		bool bSetupUndoRedo = true,
-		bool bPrintPythonCommand = false);
+    	bool bSetupUndoRedo = true,
+    	bool bPrintPythonCommand = false);
 
 	// Copies a series of function declaratioms into this graph's local function library
 	UFUNCTION(BlueprintCallable, Category = RigVMController)
-	TMap<FRigVMGraphFunctionIdentifier, URigVMLibraryNode*> LocalizeFunctions(
-		TArray<FRigVMGraphFunctionIdentifier> InFunctionDefinitions,
-		bool bLocalizeDependentPrivateFunctions = true,
-		bool bSetupUndoRedo = true,
-		bool bPrintPythonCommand = false);
-
+    TMap<URigVMLibraryNode*, URigVMLibraryNode*> LocalizeFunctions(
+        TArray<URigVMLibraryNode*> InFunctionDefinitions,
+        bool bLocalizeDependentPrivateFunctions = true,
+        bool bSetupUndoRedo = true,
+        bool bPrintPythonCommand = false);
 
 	// Returns a unique name
 	static FName GetUniqueName(const FName& InName, TFunction<bool(const FName&)> IsNameAvailableFunction, bool bAllowPeriod, bool bAllowSpace);
@@ -769,7 +732,7 @@ public:
 	// Adds a link to the graph.
 	// This causes a LinkAdded modified event.
 	UFUNCTION(BlueprintCallable, Category = RigVMController)
-	bool AddLink(const FString& InOutputPinPath, const FString& InInputPinPath, bool bSetupUndoRedo = true, bool bPrintPythonCommand = false, ERigVMPinDirection InUserDirection = ERigVMPinDirection::Output, bool bCreateCastNode = false);
+	bool AddLink(const FString& InOutputPinPath, const FString& InInputPinPath, bool bSetupUndoRedo = true, bool bPrintPythonCommand = false, ERigVMPinDirection InUserDirection = ERigVMPinDirection::Output);
 
 	// Removes a link from the graph.
 	// This causes a LinkRemoved modified event.
@@ -802,12 +765,6 @@ public:
 	bool SetExposedPinIndex(const FName& InPinName, int32 InNewIndex, bool bSetupUndoRedo = true, bool bPrintPythonCommand = false);
 
 	// Adds a function reference / invocation to the graph
-	UFUNCTION(BlueprintCallable, Category = RigVMController)
-	URigVMFunctionReferenceNode* AddFunctionReferenceNodeFromDescription(const FRigVMGraphFunctionHeader& InFunctionDefinition, const FVector2D& InNodePosition = FVector2D::ZeroVector, const
-	                                                                     FString& InNodeName = TEXT(""), bool bSetupUndoRedo = true, bool bPrintPythonCommand = false, bool bAllowPrivateFunctions = false);
-	UFUNCTION(BlueprintCallable, Category = RigVMController)
-	URigVMFunctionReferenceNode* AddExternalFunctionReferenceNode(const FString& InHostPath, const FName& InFunctionName, const FVector2D& InNodePosition = FVector2D::ZeroVector, const FString& InNodeName = TEXT(""), bool bSetupUndoRedo = true, bool bPrintPythonCommand = false);
-	
 	UFUNCTION(BlueprintCallable, Category = RigVMController)
 	URigVMFunctionReferenceNode* AddFunctionReferenceNode(URigVMLibraryNode* InFunctionDefinition, const FVector2D& InNodePosition = FVector2D::ZeroVector, const FString& InNodeName = TEXT(""), bool bSetupUndoRedo = true, bool bPrintPythonCommand = false);
 
@@ -827,14 +784,6 @@ public:
 	UFUNCTION(BlueprintCallable, Category = RigVMController)
 	bool RenameFunction(const FName& InOldFunctionName, const FName& InNewFunctionName, bool bSetupUndoRedo = true);
 
-	// Mark a function as public/private in the function library
-	UFUNCTION(BlueprintCallable, Category = RigVMController)
-	bool MarkFunctionAsPublic(const FName& InFunctionName, bool bInIsPublic, bool bSetupUndoRedo = true, bool bPrintPythonCommand = false);
-
-	// Returns true if a function is marked as public in the function library
-	UFUNCTION(BlueprintCallable, Category = RigVMController)
-	bool IsFunctionPublic(const FName& InFunctionName);
-
 	// Add a local variable to the graph
 	UFUNCTION(BlueprintCallable, Category = RigVMController)
 	FRigVMGraphVariableDescription AddLocalVariable(const FName& InVariableName, const FString& InCPPType, UObject* InCPPTypeObject, const FString& InDefaultValue, bool bSetupUndoRedo = true, bool bPrintPythonCommand = false);
@@ -858,7 +807,7 @@ public:
 	bool SetLocalVariableTypeFromObjectPath(const FName& InVariableName, const FString& InCPPType, const FString& InCPPTypeObjectPath, bool bSetupUndoRedo = true, bool bPrintPythonCommand = false);
 
 	UFUNCTION(BlueprintCallable, Category = RigVMController)
-	bool SetLocalVariableDefaultValue(const FName& InVariableName, const FString& InDefaultValue, bool bSetupUndoRedo = true, bool bPrintPythonCommand = false);
+	bool SetLocalVariableDefaultValue(const FName& InVariableName, const FString& InDefaultValue, bool bSetupUndoRedo = true, bool bPrintPythonCommand = false, bool bNotify = true);
 
 	// creates the options struct for a given workflow
 	UFUNCTION(BlueprintCallable, Category = RigVMController)
@@ -869,10 +818,10 @@ public:
 	bool PerformUserWorkflow(const FRigVMUserWorkflow& InWorkflow, const URigVMUserWorkflowOptions* InOptions, bool bSetupUndoRedo = true);
 
 	// Determine affected function references for a potential bulk edit on a library node
-	TArray<TSoftObjectPtr<URigVMFunctionReferenceNode>> GetAffectedReferences(ERigVMControllerBulkEditType InEditType, bool bForceLoad = false);
+	TArray<TSoftObjectPtr<URigVMFunctionReferenceNode>> GetAffectedReferences(ERigVMControllerBulkEditType InEditType, bool bForceLoad = false, bool bNotify = true);
 
 	// Determine affected assets for a potential bulk edit on a library node
-	TArray<FAssetData> GetAffectedAssets(ERigVMControllerBulkEditType InEditType, bool bForceLoad = false);
+	TArray<FAssetData> GetAffectedAssets(ERigVMControllerBulkEditType InEditType, bool bForceLoad = false, bool bNotify = true);
 
 	// A delegate that can be set to change the struct unfolding behaviour
 	FRigVMController_ShouldStructUnfoldDelegate UnfoldStructDelegate;
@@ -883,11 +832,17 @@ public:
 	// A delegate to retrieve the current bytecode of the graph
 	FRigVMController_GetByteCodeDelegate GetCurrentByteCodeDelegate;
 
+	// A delegate to determine if a function is public
+	FRigVMController_IsFunctionAvailableDelegate IsFunctionAvailableDelegate;
+
 	// A delegate to localize a function on demand
 	FRigVMController_RequestLocalizeFunctionDelegate RequestLocalizeFunctionDelegate;
 
 	// A delegate to create a new blueprint member variable
 	FRigVMController_RequestNewExternalVariableDelegate RequestNewExternalVariableDelegate;
+	
+	// A delegate to validate if we are allowed to introduce a dependency between two objects
+	FRigVMController_IsDependencyCyclicDelegate IsDependencyCyclicDelegate;
 
 	// A delegate to ask the host / client for a dialog to confirm a bulk edit
 	FRigVMController_RequestBulkEditDialogDelegate RequestBulkEditDialogDelegate;
@@ -904,8 +859,11 @@ public:
 	// A delegate to request to configure an options instance for a node workflow
 	FRigVMController_ConfigureWorkflowOptionsDelegate ConfigureWorkflowOptionsDelegate; 
 
-	int32 DetachLinksFromPinObjects(const TArray<URigVMLink*>* InLinks = nullptr);
-	int32 ReattachLinksToPinObjects(bool bFollowCoreRedirectors = false, const TArray<URigVMLink*>* InLinks = nullptr, bool bSetupOrphanedPins = false, bool bAllowNonArgumentLinks = false);
+	// Returns the build data of the host
+	static URigVMBuildData* GetBuildData(bool bCreateIfNeeded = true);
+
+	int32 DetachLinksFromPinObjects(const TArray<URigVMLink*>* InLinks = nullptr, bool bNotify = false);
+	int32 ReattachLinksToPinObjects(bool bFollowCoreRedirectors = false, const TArray<URigVMLink*>* InLinks = nullptr, bool bNotify = false, bool bSetupOrphanedPins = false, bool bAllowNonArgumentLinks = false);
 	void AddPinRedirector(bool bInput, bool bOutput, const FString& OldPinPath, const FString& NewPinPath);
 
 	// Removes nodes which went stale.
@@ -915,11 +873,11 @@ public:
 	bool ShouldRedirectPin(UScriptStruct* InOwningStruct, const FString& InOldRelativePinPath, FString& InOutNewRelativePinPath) const;
 	bool ShouldRedirectPin(const FString& InOldPinPath, FString& InOutNewPinPath) const;
 
-	void RepopulatePinsOnNode(URigVMNode* InNode, bool bFollowCoreRedirectors = true, bool bSetupOrphanedPins = false, bool bDetachAndReattachLinks = false);
-	void RemovePinsDuringRepopulate(URigVMNode* InNode, TArray<URigVMPin*>& InPins, bool bSetupOrphanedPins);
+	void RepopulatePinsOnNode(URigVMNode* InNode, bool bFollowCoreRedirectors = true, bool bNotify = false, bool bSetupOrphanedPins = false);
+	void RemovePinsDuringRepopulate(URigVMNode* InNode, TArray<URigVMPin*>& InPins, bool bNotify, bool bSetupOrphanedPins);
 
 	// removes any orphan pins that no longer holds a link
-	bool RemoveUnusedOrphanedPins(URigVMNode* InNode);
+	bool RemoveUnusedOrphanedPins(URigVMNode* InNode, bool bNotify);
 
 	// Initializes and recomputes the filtered permutations of all template nodes in the graph
 	// Returns true if any pin has change it's type or link was broken
@@ -962,7 +920,6 @@ public:
 	void ReportAndNotifyInfo(const FString& InMessage) const;
 	void ReportAndNotifyWarning(const FString& InMessage) const;
 	void ReportAndNotifyError(const FString& InMessage) const;
-	void ReportPinTypeChange(URigVMPin* InPin, const FString& InNewCPPType);
 	void SendUserFacingNotification(const FString& InMessage, float InDuration = 0.f, const UObject* InSubject = nullptr, const FName& InBrushName = TEXT("MessageLog.Warning")) const;
 
 	template <typename FmtType, typename... Types>
@@ -1001,69 +958,19 @@ public:
 		ReportAndNotifyError(FString::Printf(Fmt, Args...));
 	}
 
+	/**
+	 * Function to override the notification behavior and temporarily
+	 * disable all notifications. Client code is responsible for calling
+	 * SuspendNotifications(true) once all changes have been done.
+	 */
+	void SuspendNotifications(bool bSuspend) { bSuspendNotifications = bSuspend; }
+
 	/**
 	 * Helper function to disable a series of checks that can be ignored during a unit test
 	 */
 	UFUNCTION(BlueprintCallable, Category = RigVMController)
 	void SetIsRunningUnitTest(bool bIsRunning);
 
-public:
-	
-	struct FPinInfo
-	{
-		FPinInfo();
-		FPinInfo(const URigVMPin* InPin, int32 InParentIndex, ERigVMPinDirection InDirection = ERigVMPinDirection::Invalid);
-		FPinInfo(FProperty* InProperty, ERigVMPinDirection InDirection, int32 InParentIndex, const uint8* InDefaultValueMemory);
-
-		void CorrectExecuteTypeIndex();
-		
-		int32 ParentIndex;
-		FName Name;
-		ERigVMPinDirection Direction;
-		TRigVMTypeIndex TypeIndex;
-		bool bIsArray;
-		FProperty* Property;
-		FString PinPath;
-		FString DefaultValue;
-		FString DisplayName;
-		FString CustomWidgetName;
-		bool bIsExpanded;
-		bool bIsConstant;
-		bool bIsDynamicArray;
-		TArray<int32> SubPins;
-
-		friend uint32 GetTypeHash(const FPinInfo& InPin);
-	};
-
-	struct FPinInfoArray
-	{
-		FPinInfoArray() {}
-		FPinInfoArray(const URigVMNode* InNode);
-		FPinInfoArray(const URigVMNode* InNode, URigVMController* InController, const FPinInfoArray* InPreviousPinInfos = nullptr);
-		FPinInfoArray(const FRigVMGraphFunctionHeader& FunctionHeader, URigVMController* InController, const FPinInfoArray* InPreviousPinInfos = nullptr);
-
-		int32 Num() const { return Pins.Num(); }
-		const FPinInfo& operator[](int32 InIndex) const { return Pins[InIndex]; }
-		FPinInfo& operator[](int32 InIndex) { return Pins[InIndex]; }
-		TArray<FPinInfo>::RangedForIteratorType begin() const { return Pins.begin(); }
-		TArray<FPinInfo>::RangedForIteratorType end() const { return Pins.end(); }
-
-		int32 AddPin(const URigVMPin* InPin, int32 InParentIndex, ERigVMPinDirection InDirection = ERigVMPinDirection::Invalid);
-		int32 AddPin(FProperty* InProperty, URigVMController* InController, ERigVMPinDirection InDirection, int32 InParentIndex, const uint8* InDefaultValueMemory);
-		int32 AddPin(URigVMController* InController, int32 InParentIndex, const FName& InName, ERigVMPinDirection InDirection, TRigVMTypeIndex InTypeIndex, const FString& InDefaultValue, const uint8* InDefaultValueMemory, const FPinInfoArray* InPreviousPinInfos);
-		void AddPins(UScriptStruct* InScriptStruct, URigVMController* InController, ERigVMPinDirection InDirection, int32 InParentIndex, const uint8* InDefaultValueMemory);
-
-		const FString& GetPinPath(const int32 InIndex) const;
-		int32 GetIndexFromPinPath(const FString& InPinPath) const;
-		const FPinInfo* GetPinFromPinPath(const FString& InPinPath) const;
-		int32 GetRootIndex(const int32 InIndex) const;
-
-		friend uint32 GetTypeHash(const FPinInfoArray& InPins);
-		
-		mutable TArray<FPinInfo> Pins;
-		mutable TMap<FString, int32> PinPathLookup;;
-	};
-
 private:
 
 	UPROPERTY(BlueprintReadOnly, Category = RigVMController, meta = (ScriptName = "ModifiedEvent", AllowPrivateAccess = "true"))
@@ -1081,21 +988,19 @@ private:
 	bool CanAddNode(URigVMNode* InNode, bool bReportErrors, bool bIgnoreFunctionEntryReturnNodes = false);
 	TObjectPtr<URigVMNode> FindEventNode(const UScriptStruct* InScriptStruct) const;
 	bool CanAddEventNode(UScriptStruct* InScriptStruct, const bool bReportErrors) const;
-	bool CanAddFunctionRefForDefinition(const FRigVMGraphFunctionHeader& InFunctionDefinition, bool bReportErrors, bool bAllowPrivateFunctions=false);
-	void AddPinsForStruct(UStruct* InStruct, URigVMNode* InNode, URigVMPin* InParentPin, ERigVMPinDirection InPinDirection, const FString& InDefaultValue, bool bAutoExpandArrays, const FPinInfoArray* PreviousPins = nullptr);
+	bool CanAddFunctionRefForDefinition(URigVMLibraryNode* InFunctionDefinition, bool bReportErrors);
+	void AddPinsForStruct(UStruct* InStruct, URigVMNode* InNode, URigVMPin* InParentPin, ERigVMPinDirection InPinDirection, const FString& InDefaultValue, bool bAutoExpandArrays, bool bNotify = false);
 	void AddPinsForArray(FArrayProperty* InArrayProperty, URigVMNode* InNode, URigVMPin* InParentPin, ERigVMPinDirection InPinDirection, const TArray<FString>& InDefaultValues, bool bAutoExpandArrays);
-	void AddPinsForTemplate(const FRigVMTemplate* InTemplate, const FRigVMTemplateTypeMap& InPinTypeMap, URigVMNode* InNode);
 	void ConfigurePinFromProperty(FProperty* InProperty, URigVMPin* InOutPin, ERigVMPinDirection InPinDirection = ERigVMPinDirection::Invalid);
 	void ConfigurePinFromPin(URigVMPin* InOutPin, URigVMPin* InPin, bool bCopyDisplayName = false);
-	void ConfigurePinFromArgument(URigVMPin* InOutPin, const FRigVMGraphFunctionArgument& InArgument, bool bCopyDisplayName = false);
 	virtual bool ShouldStructBeUnfolded(const UStruct* InStruct);
 	virtual bool ShouldPinBeUnfolded(URigVMPin* InPin);
-	bool SetPinDefaultValue(URigVMPin* InPin, const FString& InDefaultValue, bool bResizeArrays, bool bSetupUndoRedo, bool bMergeUndoAction);
+	bool SetPinDefaultValue(URigVMPin* InPin, const FString& InDefaultValue, bool bResizeArrays, bool bSetupUndoRedo, bool bMergeUndoAction, bool bNotify = true);
 	bool ResetPinDefaultValue(URigVMPin* InPin, bool bSetupUndoRedo);
 	static FString GetPinInitialDefaultValue(const URigVMPin* InPin);
 	static FString GetPinInitialDefaultValueFromStruct(UScriptStruct* ScriptStruct, const URigVMPin* InPin, uint32 InOffset);
 	URigVMPin* InsertArrayPin(URigVMPin* ArrayPin, int32 InIndex, const FString& InDefaultValue, bool bSetupUndoRedo);
-	bool RemovePin(URigVMPin* InPinToRemove, bool bSetupUndoRedo);
+	bool RemovePin(URigVMPin* InPinToRemove, bool bSetupUndoRedo, bool bNotify);
 	FProperty* FindPropertyForPin(const FString& InPinPath);
 	bool BindPinToVariable(URigVMPin* InPin, const FString& InNewBoundVariablePath, bool bSetupUndoRedo, const FString& InVariableNodeName = FString());
 	bool UnbindPinFromVariable(URigVMPin* InPin, bool bSetupUndoRedo);
@@ -1110,10 +1015,9 @@ private:
 	void RelinkSourceAndTargetPins(URigVMNode* RigNode, bool bSetupUndoRedo = true);
 
 public:
-	bool AddLink(URigVMPin* OutputPin, URigVMPin* InputPin, bool bSetupUndoRedo = true, ERigVMPinDirection InUserDirection = ERigVMPinDirection::Invalid, bool bCreateCastNode = false);
+	bool AddLink(URigVMPin* OutputPin, URigVMPin* InputPin, bool bSetupUndoRedo = true, ERigVMPinDirection InUserDirection = ERigVMPinDirection::Invalid);
 	bool BreakLink(URigVMPin* OutputPin, URigVMPin* InputPin, bool bSetupUndoRedo = true);
 	bool BreakAllLinks(URigVMPin* Pin, bool bAsInput, bool bSetupUndoRedo = true);
-	void EnableTypeCasting(bool bEnabled = true) { bEnableTypeCasting = bEnabled; }
 
 private:
 	bool BreakAllLinksRecursive(URigVMPin* Pin, bool bAsInput, bool bTowardsParent, bool bSetupUndoRedo);
@@ -1130,7 +1034,7 @@ private:
 	URigVMCollapseNode* PromoteFunctionReferenceNodeToCollapseNode(URigVMFunctionReferenceNode* InFunctionRefNode, bool bSetupUndoRedo, bool bRemoveFunctionDefinition);
 	void SetReferencedFunction(URigVMFunctionReferenceNode* InFunctionRefNode, URigVMLibraryNode* InNewReferencedNode, bool bSetupUndoRedo);
 
-	void RefreshFunctionPins(URigVMNode* InNode);
+	void RefreshFunctionPins(URigVMNode* InNode, bool bNotify = true);
 
 	void ReportRemovedLink(const FString& InSourcePinPath, const FString& InTargetPinPath);
 
@@ -1260,18 +1164,6 @@ public:
 	// Inquire on the commonly used types for a template node. This can be used to resolve a node without user input (as a default)
 	FRigVMTemplate::FTypeMap GetCommonlyUsedTypesForTemplate(const URigVMTemplateNode* InNode) const;
 #endif
-
-protected:
-
-	// backwards compatibility code
-	FRigVMClientPatchResult PatchUnitNodesOnLoad();
-	FRigVMClientPatchResult PatchDispatchNodesOnLoad();
-	FRigVMClientPatchResult PatchBranchNodesOnLoad();
-	FRigVMClientPatchResult PatchIfSelectNodesOnLoad();
-	FRigVMClientPatchResult PatchArrayNodesOnLoad();
-
-	// work to do after a duplication of the host asset
-	void PostDuplicateHost(const FString& InOldPathName, const FString& InNewPathName);
 	
 private: 
 	UPROPERTY(transient)
@@ -1284,7 +1176,6 @@ private:
 	bool bReportWarningsAndErrors;
 	bool bIgnoreRerouteCompactnessChanges;
 	ERigVMPinDirection UserLinkDirection;
-	bool bEnableTypeCasting;
 
 	// temporary maps used for pin redirection
 	// only valid between Detach & ReattachLinksToPinObjects
@@ -1303,17 +1194,17 @@ private:
 		{
 		}
 
-		friend uint32 GetTypeHash(const FControlRigStructPinRedirectorKey& Cache)
+		friend FORCEINLINE uint32 GetTypeHash(const FControlRigStructPinRedirectorKey& Cache)
 		{
 			return HashCombine(GetTypeHash(Cache.Struct), GetTypeHash(Cache.PinPathInNode));
 		}
 
-		bool operator ==(const FControlRigStructPinRedirectorKey& Other) const
+		FORCEINLINE bool operator ==(const FControlRigStructPinRedirectorKey& Other) const
 		{
 			return Struct == Other.Struct && PinPathInNode == Other.PinPathInNode;
 		}
 
-		bool operator !=(const FControlRigStructPinRedirectorKey& Other) const
+		FORCEINLINE bool operator !=(const FControlRigStructPinRedirectorKey& Other) const
 		{
 			return Struct != Other.Struct || PinPathInNode != Other.PinPathInNode;
 		}
@@ -1330,14 +1221,7 @@ private:
 	bool bIsTransacting; // Performing undo/redo transaction
 	bool bIsRunningUnitTest;
 	bool bIsFullyResolvingTemplateNode;
-
-public:
-	
 	bool bSuspendRecomputingTemplateFilters;
-	bool bSuspendTemplateComputation;
-
-private:
-
 #if WITH_EDITOR
 	bool bRegisterTemplateNodeUsage;
 #endif
@@ -1354,9 +1238,6 @@ private:
 	friend class FRigVMParserAST;
 	friend class FRigVMControllerCompileBracketScope;
 	friend class FRigVMControllerGraphGuard;
-	friend struct FPinInfoArray;
-	friend class FRigVMControllerNotifGuard;
-	friend struct FRigVMClient;
 };
 
 class FRigVMControllerGraphGuard
@@ -1390,28 +1271,6 @@ private:
 	int32 NumGraphs;
 };
 
-class FRigVMControllerNotifGuard
-{
-public:
-
-	FRigVMControllerNotifGuard(URigVMController* InController, bool bInSuspendNotifications = true)
-		: Controller(InController)
-	{
-		bPreviousSuspendNotifications = Controller->bSuspendNotifications;
-		Controller->bSuspendNotifications = bInSuspendNotifications;
-	}
-
-	~FRigVMControllerNotifGuard()
-	{
-		Controller->bSuspendNotifications = bPreviousSuspendNotifications;
-	}
-
-private:
-
-	URigVMController* Controller;
-	bool bPreviousSuspendNotifications;
-};
-
 USTRUCT()
 struct FRigVMController_CommonTypePerTemplate
 {
diff --git a/Engine/Source/Runtime/RigVM/Public/RigVMDefines.h b/Engine/Source/Runtime/RigVM/Public/RigVMDefines.h
--- a/Engine/Source/Runtime/RigVM/Public/RigVMDefines.h
+++ b/Engine/Source/Runtime/RigVM/Public/RigVMDefines.h
@@ -29,7 +29,3 @@
 #ifndef UE_RIGVM_DEBUG_EXECUTION
 #define UE_RIGVM_DEBUG_EXECUTION 0
 #endif
-
-#ifndef UE_RIGVMCONTROLLER_VERBOSE_REPOPULATE
-#define UE_RIGVMCONTROLLER_VERBOSE_REPOPULATE UE_BUILD_DEBUG
-#endif
