RigVM: Attempt3 : Avoid recreating pins if not required during RepopulatePins

diff --git a/Engine/Source/Developer/RigVMDeveloper/Private/RigVMModel/RigVMController.cpp b/Engine/Source/Developer/RigVMDeveloper/Private/RigVMModel/RigVMController.cpp
--- a/Engine/Source/Developer/RigVMDeveloper/Private/RigVMModel/RigVMController.cpp
+++ b/Engine/Source/Developer/RigVMDeveloper/Private/RigVMModel/RigVMController.cpp
@@ -15639,10 +15639,10 @@ void URigVMController::RepopulatePinsOnNode(URigVMNode* InNode, bool bFollowCore
 	{
 		const FString PinPath = PreviousPinInfos.GetPinPath(Index);
 		const int32 NewIndex = NewPinInfos.GetIndexFromPinPath(PinPath);
+		
 		if(NewIndex == INDEX_NONE)
 		{
 			const int32 RootIndex = PreviousPinInfos.GetRootIndex(Index);
-			
 			if(PreviousPinInfos[Index].Direction != ERigVMPinDirection::Hidden)
 			{
 				if(URigVMPin* Pin = InNode->FindPin(PinPath))
@@ -15670,6 +15670,7 @@ void URigVMController::RepopulatePinsOnNode(URigVMNode* InNode, bool bFollowCore
 						if(bPinShouldBeOrphaned)
 						{
 							PreviousPinsToOrphan.Add(RootIndex);
+							
 							bRequireDetachLinks = true;
 							bRequirePinStates = true;
 #if UE_RIGVMCONTROLLER_VERBOSE_REPOPULATE
@@ -15737,15 +15738,28 @@ void URigVMController::RepopulatePinsOnNode(URigVMNode* InNode, bool bFollowCore
 		if(PreviousIndex == INDEX_NONE)
 		{
 			NewPinsToAdd.Add(Index);
-
 #if UE_RIGVMCONTROLLER_VERBOSE_REPOPULATE
 			UE_LOG(LogRigVMDeveloper, Display, TEXT("Newly required pin '%s' needs to be added."), *PinPath);
 #endif
 		}
+		else
+		{
+			// the previous pin exists - but it has been orphaned
+			const int32 PreviousRootIndex = PreviousPinInfos.GetRootIndex(PreviousIndex);
+			if(PreviousPinsToOrphan.Contains(PreviousRootIndex))
+			{
+				NewPinsToAdd.Add(Index);
+
+#if UE_RIGVMCONTROLLER_VERBOSE_REPOPULATE
+				UE_LOG(LogRigVMDeveloper, Display, TEXT("Orphaned pin '%s' needs to be re-added."), *PinPath);
+#endif
+			}
+		}
 	}
 
 	auto CreatePinFromPinInfo = [this, &Registry, &PreviousPinInfos](const FPinInfo& InPinInfo, const FString& InPinPath, UObject* InOuter) -> URigVMPin*
 	{
+		check(InOuter);
 		URigVMPin* Pin = NewObject<URigVMPin>(InOuter, InPinInfo.Name);
 		if(InPinInfo.Property)
 		{
@@ -15858,7 +15872,6 @@ void URigVMController::RepopulatePinsOnNode(URigVMNode* InNode, bool bFollowCore
 
 				Notify(ERigVMGraphNotifType::PinRemoved, Pin);
 				InNode->OrphanedPins.Add(Pin);
-
 				Notify(ERigVMGraphNotifType::PinAdded, Pin);
 			}
 			else
@@ -15880,7 +15893,6 @@ void URigVMController::RepopulatePinsOnNode(URigVMNode* InNode, bool bFollowCore
 			RemovePin(Pin, false);
 		}
 	}
-
 	// add missing pins
 	for(int32 Index = 0; Index < NewPinsToAdd.Num(); Index++)
 	{
@@ -15897,7 +15909,6 @@ void URigVMController::RepopulatePinsOnNode(URigVMNode* InNode, bool bFollowCore
 		UE_LOG(LogRigVMDeveloper, Display, TEXT("Adding new pin '%s'."), *PinPath);
 #endif
 	}
-
 	// update existing pins
 	for(int32 Index = 0; Index < PreviousPinsToUpdate.Num(); Index++)
 	{
