Control Rig: Populate execute pins in dispatch nodes

diff --git a/Engine/Source/Developer/RigVMDeveloper/Private/RigVMModel/RigVMController.cpp b/Engine/Source/Developer/RigVMDeveloper/Private/RigVMModel/RigVMController.cpp
--- a/Engine/Source/Developer/RigVMDeveloper/Private/RigVMModel/RigVMController.cpp
+++ b/Engine/Source/Developer/RigVMDeveloper/Private/RigVMModel/RigVMController.cpp
@@ -15383,6 +15383,31 @@ void URigVMController::RepopulatePinsOnNode(URigVMNode* InNode, bool bFollowCore
 		}
 		
 		const FRigVMTemplate* Template = DispatchNode->GetTemplate();
+
+		FRigVMDispatchContext DispatchContext = DispatchNode->GetDispatchContext();
+		auto AddExecutePins = [Template, DispatchNode, &Registry, &DispatchContext, &NewPinInfos, &PreviousPinInfos, this](ERigVMPinDirection InPinDirection)
+		{
+			for (int32 ArgIndex = 0; ArgIndex < Template->NumExecuteArguments(DispatchContext); ArgIndex++)
+			{
+				const FRigVMExecuteArgument* Arg = Template->GetExecuteArgument(ArgIndex, DispatchContext);
+				const FRigVMTemplateArgumentType Type = Registry.GetType(Arg->TypeIndex);
+				const TRigVMTypeIndex TypeIndex = Registry.GetTypeIndex(Type);
+
+				FString DefaultValue;
+				if(Registry.IsArrayType(Arg->TypeIndex))
+				{
+					if(const FRigVMDispatchFactory* Factory = DispatchNode->GetFactory())
+					{
+						DefaultValue = Factory->GetArgumentDefaultValue(Arg->Name, Arg->TypeIndex);
+					}
+				}
+
+				(void)NewPinInfos.AddPin(this, INDEX_NONE, Arg->Name, Arg->Direction, TypeIndex, DefaultValue, nullptr, &PreviousPinInfos);
+			}
+		};
+
+		AddExecutePins(ERigVMPinDirection::IO);
+		AddExecutePins(ERigVMPinDirection::Input);
 		
 		for (int32 ArgIndex = 0; ArgIndex < Template->NumArguments(); ArgIndex++)
 		{
@@ -15445,6 +15470,8 @@ void URigVMController::RepopulatePinsOnNode(URigVMNode* InNode, bool bFollowCore
 			
 			(void)NewPinInfos.AddPin(this, INDEX_NONE, Arg->Name, Arg->GetDirection(), TypeIndex, DefaultValue, DefaultValueMemory, &PreviousPinInfos);
 		}
+
+		AddExecutePins(ERigVMPinDirection::Output);
 	}
 	else if ((RerouteNode != nullptr) || (VariableNode != nullptr))
 	{
