Control Rig: Fix ensure default value with bad format

diff --git a/Engine/Source/Developer/RigVMDeveloper/Private/RigVMModel/RigVMController.cpp b/Engine/Source/Developer/RigVMDeveloper/Private/RigVMModel/RigVMController.cpp
--- a/Engine/Source/Developer/RigVMDeveloper/Private/RigVMModel/RigVMController.cpp
+++ b/Engine/Source/Developer/RigVMDeveloper/Private/RigVMModel/RigVMController.cpp
@@ -16650,6 +16650,16 @@ void URigVMController::PostProcessDefaultValue(URigVMPin* Pin, FString& OutDefau
 			OutDefaultValue.Reset();
 		}
 	}
+	if (Pin->IsStruct() || Pin->IsArray())
+	{
+		if (!OutDefaultValue.IsEmpty())
+		{
+			if (OutDefaultValue[0] != TCHAR('(') || OutDefaultValue[OutDefaultValue.Len()-1] != TCHAR(')'))
+			{
+				OutDefaultValue.Reset();
+			}
+		}
+	}
 
 	if (Pin->IsArray() && OutDefaultValue.IsEmpty())
 	{
diff --git a/Engine/Source/Developer/RigVMDeveloper/Private/RigVMModel/RigVMPin.cpp b/Engine/Source/Developer/RigVMDeveloper/Private/RigVMModel/RigVMPin.cpp
--- a/Engine/Source/Developer/RigVMDeveloper/Private/RigVMModel/RigVMPin.cpp
+++ b/Engine/Source/Developer/RigVMDeveloper/Private/RigVMModel/RigVMPin.cpp
@@ -696,11 +696,22 @@ FString URigVMPin::GetDefaultValue(const URigVMPin::FPinOverride& InOverride) co
 
 	if (IsArray())
 	{
+		FRigVMRegistry& Registry = FRigVMRegistry::Get();
+		const TRigVMTypeIndex ArrayType = GetTypeIndex();
+		if (!Registry.IsArrayType(ArrayType))
+		{
+			return TEXT("()");
+		}
 		if (SubPins.Num() > 0)
 		{
+			const TRigVMTypeIndex ElementType = Registry.GetBaseTypeFromArrayTypeIndex(ArrayType);
 			TArray<FString> ElementDefaultValues;
 			for (URigVMPin* SubPin : SubPins)
 			{
+				if (SubPin->GetTypeIndex() != ElementType)
+				{
+					return TEXT("()");
+				}
 				FString ElementDefaultValue = SubPin->GetDefaultValue(InOverride);
 				if (SubPin->IsStringType())
 				{
