Control Rig: Simplify template functionality

diff --git a/Engine/Plugins/Animation/ControlRig/Source/ControlRigDeveloper/Private/ControlRigBlueprint.cpp b/Engine/Plugins/Animation/ControlRig/Source/ControlRigDeveloper/Private/ControlRigBlueprint.cpp
--- a/Engine/Plugins/Animation/ControlRig/Source/ControlRigDeveloper/Private/ControlRigBlueprint.cpp
+++ b/Engine/Plugins/Animation/ControlRig/Source/ControlRigDeveloper/Private/ControlRigBlueprint.cpp
@@ -548,6 +548,19 @@ void UControlRigBlueprint::HandleConfigureRigVMController(const FRigVMClient* In
 		return false;
 	});
 
+	InControllerToConfigure->RequestPinTypeSelectionDelegate.BindLambda([WeakThis](const TArray<TRigVMTypeIndex>& InTypes) -> TRigVMTypeIndex 
+	{
+		if(WeakThis.IsValid())
+		{
+			UControlRigBlueprint* StrongThis = WeakThis.Get();
+			if(StrongThis->OnRequestPinTypeSelectionDialog().IsBound())
+			{
+				return StrongThis->OnRequestPinTypeSelectionDialog().Execute(InTypes);
+			}
+		}
+		return INDEX_NONE;
+	});
+
 	InControllerToConfigure->RequestNewExternalVariableDelegate.BindLambda([WeakThis](FRigVMGraphVariableDescription InVariable, bool bInIsPublic, bool bInIsReadOnly) -> FName
 	{
 		if (WeakThis.IsValid())
@@ -903,7 +916,7 @@ void UControlRigBlueprint::PostLoad()
 			{
 				URigVMController* Controller = GetOrCreateController(GraphToValidate);
 				FRigVMControllerNotifGuard NotifGuard(Controller, true);
-				Controller->RemoveUnusedOrphanedPins(Node);				
+				Controller->RemoveUnusedOrphanedPins(Node);
 			}
 				
 			for(URigVMNode* Node : GraphToValidate->GetNodes())
@@ -1339,11 +1352,6 @@ void UControlRigBlueprint::RefreshAllModels(EControlRigBlueprintLoadType InLoadT
 		Controller->ReattachLinksToPinObjects(true /* follow redirectors */, nullptr, true, true);
 	}
 
-	if(bIsPostLoad)
-	{
-		PatchTemplateNodesWithPreferredPermutation();
-	}
-	
 	TArray<URigVMGraph*> GraphsToClean = GetAllModels();
 
 	// Sort from leaf graphs to root
@@ -1367,7 +1375,7 @@ void UControlRigBlueprint::RefreshAllModels(EControlRigBlueprintLoadType InLoadT
 					if (URigVMFunctionReferenceNode* FunctionReferenceNode = Cast<URigVMFunctionReferenceNode>(LibraryNode))
 					{
 						if (FunctionReferenceNode->GetReferencedFunctionHeader().LibraryPointer.LibraryNode.GetLongPackageName() != GetPackage()->GetPathName())
-						{							
+						{
 							continue;
 						}
 						if (URigVMLibraryNode* ReferencedNode = FunctionReferenceNode->LoadReferencedNode())
@@ -1380,7 +1388,7 @@ void UControlRigBlueprint::RefreshAllModels(EControlRigBlueprintLoadType InLoadT
 					if (URigVMGraph* ContainedGraph = LibraryNode->GetContainedGraph())
 					{
 						ContainedGraphs.Add(ContainedGraph);
-					}					
+					}
 				}
 			}
 			
@@ -1406,7 +1414,6 @@ void UControlRigBlueprint::RefreshAllModels(EControlRigBlueprintLoadType InLoadT
 	{
 		URigVMGraph* GraphToClean = SortedGraphsToClean[GraphIndex];
 		URigVMController* Controller = GetOrCreateController(GraphToClean);
-		TGuardValue<bool> RecomputeGuard(Controller->bSuspendRecomputingOuterTemplateFilters, true);
 		TGuardValue<bool> GuardEditGraph(GraphToClean->bEditable, true);
 		FRigVMControllerNotifGuard NotifGuard(Controller, true);
 		
@@ -1423,125 +1430,64 @@ void UControlRigBlueprint::RefreshAllModels(EControlRigBlueprintLoadType InLoadT
 				{
 					TemplateNode->InvalidateCache();
 					TemplateNode->PostLoad();
-
-					// Fix preferred type of execute context which are not the base class
-					{
-						const FRigVMRegistry& Registry = FRigVMRegistry::Get();
-						for (FRigVMTemplatePreferredType& PreferredType : TemplateNode->PreferredPermutationPairs)
-						{
-							if (!Registry.IsExecuteType(PreferredType.TypeIndex))
-							{
-								continue;
-							}
-							
-							if (URigVMPin* Pin = TemplateNode->FindPin(PreferredType.GetArgument().ToString()))
-							{
-								if (Pin->GetTypeIndex() == PreferredType.GetTypeIndex())
-								{
-									continue;
-								}
-
-								bDirtyDuringLoad = true;
-#if UE_RIGVM_DEBUG_TYPEINDEX
-								// Create a TRigVMTypeIndex from the int32 index so that it can be modified if needed
-								// and then converted back to int32
-								TRigVMTypeIndex PreferredTypeIndex = PreferredType.TypeIndex;
-								Registry.ConvertExecuteContextToBaseType(PreferredTypeIndex);
-								PreferredType.TypeIndex = PreferredTypeIndex;
-#else
-								Registry.ConvertExecuteContextToBaseType(PreferredType.TypeIndex);
-#endif
-							}
-						}
-					}
-
-					// Remove preferred types that do not match the pin type
-					{
-						int32 NumRemoved = TemplateNode->PreferredPermutationPairs.RemoveAll([TemplateNode](const FRigVMTemplatePreferredType& PreferredType)
-						{
-							if (URigVMPin* Pin = TemplateNode->FindPin(PreferredType.GetArgument().ToString()))
-							{
-								if (Pin->GetTypeIndex() == PreferredType.GetTypeIndex())
-								{
-									return false;
-								}
-							}
-							return true;
-						});
-
-						if (NumRemoved > 0)
-						{
-							bDirtyDuringLoad = true;
-						}
-					}
-
-					if(!TemplateNode->HasWildCardPin() &&
-						!TemplateNode->IsSingleton() &&
-						TemplateNode->PreferredPermutationPairs.IsEmpty())
-					{
-						TemplateNode->InitializeFilteredPermutationsFromTypes(false);
-					}
 				}
 			}
-
-			if (URigVMLibraryNode* LibraryNode = GraphToClean->GetTypedOuter<URigVMLibraryNode>())
-			{
-				Controller->UpdateLibraryTemplate(LibraryNode, false);
-			}
 		}
 
 #if WITH_EDITOR
 
 		if(bIsPostLoad)
 		{
-			if (Controller->RecomputeAllTemplateFilteredPermutations(false))
-			{
-				ensureMsgf(false, TEXT("Pin type changed during load %s"), *GetPackage()->GetPathName());
-				// if you hit this ensure consider to uncomment the code
-				// in URigVMController::ReportPinTypeChange for debugging
-			}
-
 			for(URigVMNode* ModelNode : GraphToClean->GetNodes())
 			{
+				if (ModelNode->HasWildCardPin())
+				{
+					continue;
+				}
 				if(URigVMUnitNode* UnitNode = Cast<URigVMUnitNode>(ModelNode))
 				{
-					if (!UnitNode->HasWildCardPin())
+					UScriptStruct* ScriptStruct = UnitNode->GetScriptStruct(); 
+					if(ScriptStruct == nullptr)
 					{
-						UScriptStruct* ScriptStruct = UnitNode->GetScriptStruct(); 
-						if(ScriptStruct == nullptr)
-						{
-							Controller->FullyResolveTemplateNode(UnitNode, INDEX_NONE, false);
-						}
+						Controller->FullyResolveTemplateNode(UnitNode, INDEX_NONE, false);
+					}
 
-						// Try to find a deprecated template
-						if (UnitNode->GetScriptStruct() == nullptr && !UnitNode->TemplateNotation.IsNone())
+					// Try to find a deprecated template
+					if (UnitNode->GetScriptStruct() == nullptr && !UnitNode->TemplateNotation.IsNone())
+					{
+						const FRigVMTemplate* Template = FRigVMRegistry::Get().FindTemplate(UnitNode->TemplateNotation, true);
+						FRigVMTemplate::FTypeMap TypeMap = UnitNode->GetTemplatePinTypeMap(true);
+						int32 Permutation;
+						if (Template->FullyResolve(TypeMap, Permutation))
 						{
-							const FRigVMTemplate* Template = FRigVMRegistry::Get().FindTemplate(UnitNode->TemplateNotation, true);
-							FRigVMTemplate::FTypeMap TypeMap;
-							for (URigVMPin* Pin : UnitNode->GetPins())
-							{
-								check(!Pin->IsWildCard());
-								TypeMap.Add(Pin->GetFName(), Pin->GetTypeIndex());
-							}
-
-							int32 Permutation;
-							if (Template->FullyResolve(TypeMap, Permutation))
-							{
-								const FRigVMFunction* Function = Template->GetPermutation(Permutation);
-								UnitNode->ResolvedFunctionName = Function->GetName();
-							}
+							const FRigVMFunction* Function = Template->GetPermutation(Permutation);
+							UnitNode->ResolvedFunctionName = Function->GetName();
 						}
+					}
 
-						if (UnitNode->GetScriptStruct() == nullptr)
-						{
-							static const TCHAR UnresolvedUnitNodeMessage[] = TEXT("Node %s could not be resolved.");
-							Controller->ReportErrorf(UnresolvedUnitNodeMessage, *ModelNode->GetNodePath(true));
-						}
+					if (UnitNode->GetScriptStruct() == nullptr)
+					{
+						static const TCHAR UnresolvedUnitNodeMessage[] = TEXT("Node %s could not be resolved.");
+						Controller->ReportErrorf(UnresolvedUnitNodeMessage, *ModelNode->GetNodePath(true));
 					}
 				}
 				if (URigVMDispatchNode* DispatchNode = Cast<URigVMDispatchNode>(ModelNode))
 				{
-					if (DispatchNode->GetFactory() == nullptr)
+					if (const FRigVMDispatchFactory* Factory = DispatchNode->GetFactory())
+					{
+						FRigVMTemplate::FTypeMap TypeMap = DispatchNode->GetTemplatePinTypeMap(true);
+						const FString ResolvedPermutationName = Factory->GetPermutationName(TypeMap);
+						if(const FRigVMFunction* Function = FRigVMRegistry::Get().FindFunction(*ResolvedPermutationName))
+						{
+							DispatchNode->ResolvedFunctionName = Function->GetName();
+						}
+						else
+						{
+							static const TCHAR UnresolvedDispatchPermutationMessage[] = TEXT("Dispatch factory %s could not find permutation..");
+							Controller->ReportErrorf(UnresolvedDispatchPermutationMessage, *Factory->GetFactoryName().ToString());
+						}
+					}
+					else
 					{
 						static const TCHAR UnresolvedDispatchNodeMessage[] = TEXT("Dispatch node %s has no factory..");
 						Controller->ReportErrorf(UnresolvedDispatchNodeMessage, *ModelNode->GetNodePath(true));
@@ -4512,29 +4458,6 @@ void UControlRigBlueprint::PatchParameterNodesOnLoad()
 #endif	
 }
 
-void UControlRigBlueprint::PatchTemplateNodesWithPreferredPermutation()
-{
-#if WITH_EDITOR
-
-	if (GetLinkerCustomVersion(FControlRigObjectVersion::GUID) < FControlRigObjectVersion::TemplatesPreferredPermutatation)
-	{
-		for (URigVMGraph* Graph : GetAllModels())
-		{
-			for(URigVMNode* ModelNode : Graph->GetNodes())
-			{
-				if (URigVMTemplateNode* TemplateNode = Cast<URigVMTemplateNode>(ModelNode))
-				{
-					TemplateNode->InvalidateCache();
-					TemplateNode->PostLoad();
-				}
-			}
-			
-			GetController(Graph)->InitializeFilteredPermutationsFromTemplateTypes();
-		}
-	}
-#endif
-}
-
 void UControlRigBlueprint::PatchLinksWithCast()
 {
 #if WITH_EDITOR
@@ -4583,7 +4506,6 @@ void UControlRigBlueprint::PatchLinksWithCast()
 				}
 			}
 
-			TGuardValue<bool> DisableRecomputeTemplates(Controller->bSuspendRecomputingTemplateFilters, true);
 			Controller->BreakLink(Tuple.Get<2>(), Tuple.Get<3>(), false);
 
 			// notify the user that the link has been broken.
diff --git a/Engine/Plugins/Animation/ControlRig/Source/ControlRigDeveloper/Public/ControlRigBlueprint.h b/Engine/Plugins/Animation/ControlRig/Source/ControlRigDeveloper/Public/ControlRigBlueprint.h
--- a/Engine/Plugins/Animation/ControlRig/Source/ControlRigDeveloper/Public/ControlRigBlueprint.h
+++ b/Engine/Plugins/Animation/ControlRig/Source/ControlRigDeveloper/Public/ControlRigBlueprint.h
@@ -43,6 +43,7 @@ DECLARE_EVENT_ThreeParams(UControlRigBlueprint, FOnLocalizeFunctionDialogRequest
 DECLARE_EVENT_ThreeParams(UControlRigBlueprint, FOnReportCompilerMessage, EMessageSeverity::Type, UObject*, const FString&);
 DECLARE_DELEGATE_RetVal_FourParams(FRigVMController_BulkEditResult, FControlRigOnBulkEditDialogRequestedDelegate, UControlRigBlueprint*, URigVMController*, URigVMLibraryNode*, ERigVMControllerBulkEditType);
 DECLARE_DELEGATE_RetVal_OneParam(bool, FControlRigOnBreakLinksDialogRequestedDelegate, TArray<URigVMLink*>);
+DECLARE_DELEGATE_RetVal_OneParam(TRigVMTypeIndex, FControlRigOnPinTypeSelectionRequestedDelegate, const TArray<TRigVMTypeIndex>&);
 DECLARE_EVENT(UControlRigBlueprint, FOnBreakpointAdded);
 DECLARE_EVENT_OneParam(UControlRigBlueprint, FOnRequestInspectObject, const TArray<UObject*>& );
 
@@ -611,7 +612,6 @@ private:
 	void PatchVariableNodesWithIncorrectType();
 	void PatchPropagateToChildren();
 	void PatchParameterNodesOnLoad();
-	void PatchTemplateNodesWithPreferredPermutation();
 	void PatchLinksWithCast();
 	void PatchFunctionsOnLoad();
 
@@ -705,6 +705,8 @@ public:
 
 	FControlRigOnBreakLinksDialogRequestedDelegate& OnRequestBreakLinksDialog() { return RequestBreakLinksDialog; }
 
+	FControlRigOnPinTypeSelectionRequestedDelegate& OnRequestPinTypeSelectionDialog() { return RequestPinTypeSelectionDialog; }
+
 	FRigVMController_RequestJumpToHyperlinkDelegate& OnRequestJumpToHyperlink() { return RequestJumpToHyperlink; };
 
 	FOnReportCompilerMessage& OnReportCompilerMessage() { return ReportCompilerMessageEvent; }
@@ -723,6 +725,7 @@ private:
 	FOnReportCompilerMessage ReportCompilerMessageEvent;
 	FControlRigOnBulkEditDialogRequestedDelegate RequestBulkEditDialog;
 	FControlRigOnBreakLinksDialogRequestedDelegate RequestBreakLinksDialog;
+	FControlRigOnPinTypeSelectionRequestedDelegate RequestPinTypeSelectionDialog;
 	FRigVMController_RequestJumpToHyperlinkDelegate RequestJumpToHyperlink;
 
 #endif
diff --git a/Engine/Plugins/Animation/ControlRig/Source/ControlRigEditor/Private/ControlRigEditorModule.cpp b/Engine/Plugins/Animation/ControlRig/Source/ControlRigEditor/Private/ControlRigEditorModule.cpp
--- a/Engine/Plugins/Animation/ControlRig/Source/ControlRigEditor/Private/ControlRigEditorModule.cpp
+++ b/Engine/Plugins/Animation/ControlRig/Source/ControlRigEditor/Private/ControlRigEditorModule.cpp
@@ -1470,7 +1470,14 @@ void FControlRigEditorModule::GetContextMenuActions(const UControlRigGraphSchema
 											TSharedRef<SControlRigChangePinType> ChangePinTypeWidget =
 											SNew(SControlRigChangePinType)
 											.Blueprint(RigBlueprint)
-											.ModelPins({ModelPin});
+											.Types(Argument->GetTypeIndices())
+											.OnTypeSelected_Lambda([RigBlueprint, ModelPin](const TRigVMTypeIndex& TypeSelected)
+											{
+												if (URigVMController* Controller = RigBlueprint->GetController(ModelPin->GetGraph()))
+												{
+													Controller->ResolveWildCardPin(ModelPin, TypeSelected, true, true);
+												}
+											});
 
 											TemplatesSection.AddEntry(FToolMenuEntry::InitWidget("ChangePinTypeWidget", ChangePinTypeWidget, FText(), true));
 										}
diff --git a/Engine/Plugins/Animation/ControlRig/Source/ControlRigEditor/Private/Editor/ControlRigEditor.cpp b/Engine/Plugins/Animation/ControlRig/Source/ControlRigEditor/Private/Editor/ControlRigEditor.cpp
--- a/Engine/Plugins/Animation/ControlRig/Source/ControlRigEditor/Private/Editor/ControlRigEditor.cpp
+++ b/Engine/Plugins/Animation/ControlRig/Source/ControlRigEditor/Private/Editor/ControlRigEditor.cpp
@@ -8,6 +8,7 @@
 #include "Editor/ControlRigEditorMode.h"
 #include "SKismetInspector.h"
 #include "SEnumCombo.h"
+#include "Widgets/Layout/SScrollBox.h"
 #include "Framework/Commands/GenericCommands.h"
 #include "Editor.h"
 #include "Editor/Transactor.h"
@@ -87,6 +88,7 @@
 #include "Editor/SControlRigFunctionLocalizationWidget.h"
 #include "Editor/SControlRigFunctionBulkEditWidget.h"
 #include "Editor/SControlRigBreakLinksWidget.h"
+#include "Graph/SControlRigGraphChangePinType.h"
 #include "SGraphPanel.h"
 #include "Engine/StaticMesh.h"
 #include "Engine/StaticMeshActor.h"
@@ -265,6 +267,7 @@ FControlRigEditor::~FControlRigEditor()
 		RigBlueprint->OnRequestLocalizeFunctionDialog().RemoveAll(this);
 		RigBlueprint->OnRequestBulkEditDialog().Unbind();
 		RigBlueprint->OnRequestBreakLinksDialog().Unbind();
+		RigBlueprint->OnRequestPinTypeSelectionDialog().Unbind();
 		RigBlueprint->OnRequestJumpToHyperlink().Unbind();
 		RigBlueprint->OnReportCompilerMessage().RemoveAll(this);
 
@@ -537,6 +540,7 @@ void FControlRigEditor::InitControlRigEditor(const EToolkitMode::Type Mode, cons
 		InControlRigBlueprint->OnRequestLocalizeFunctionDialog().AddSP(this, &FControlRigEditor::OnRequestLocalizeFunctionDialog);
 		InControlRigBlueprint->OnRequestBulkEditDialog().BindSP(this, &FControlRigEditor::OnRequestBulkEditDialog);
 		InControlRigBlueprint->OnRequestBreakLinksDialog().BindSP(this, &FControlRigEditor::OnRequestBreakLinksDialog);
+		InControlRigBlueprint->OnRequestPinTypeSelectionDialog().BindSP(this, &FControlRigEditor::OnRequestPinTypeSelectionDialog);
 		InControlRigBlueprint->OnRequestJumpToHyperlink().BindSP(this, &FControlRigEditor::HandleJumpToHyperlink);
 	}
 
@@ -5124,6 +5128,56 @@ bool FControlRigEditor::OnRequestBreakLinksDialog(TArray<URigVMLink*> InLinks)
 	return BreakLinksDialog->ShowModal() == EAppReturnType::Ok; 
 }
 
+TRigVMTypeIndex FControlRigEditor::OnRequestPinTypeSelectionDialog(const TArray<TRigVMTypeIndex>& InTypes)
+{
+	if(InTypes.Num() == 0)
+	{
+		return true;
+	}
+
+	TRigVMTypeIndex Answer = INDEX_NONE;
+
+	TSharedPtr< SWindow > Window = SNew(SWindow)
+		.Title(LOCTEXT("SelectPinType", "Select Pin Type"))
+		.SizingRule(ESizingRule::Autosized)
+		.SupportsMaximize(false)
+		.SupportsMinimize(false)
+		[
+			SNew( SBorder )
+			.Padding( 4.f )
+			.BorderImage( FAppStyle::GetBrush( "ToolPanel.GroupBorder" ) )
+			[
+				SNew(SBox)
+				.MaxDesiredHeight(static_cast<float>(300))
+				[
+					SNew(SVerticalBox)
+					+SVerticalBox::Slot()
+					[
+						SNew(SBox)
+						.MaxDesiredHeight(static_cast<float>(300))
+						[
+							SNew(SScrollBox)
+							+SScrollBox::Slot()
+							[
+								SNew(SControlRigChangePinType)
+									.Blueprint(GetControlRigBlueprint())
+									.Types(InTypes)
+									.OnTypeSelected_Lambda([&Answer](const TRigVMTypeIndex& TypeSelected)
+									{
+										Answer = TypeSelected;
+										FSlateApplication::Get().GetActiveModalWindow()->RequestDestroyWindow();
+									})
+							]
+						]
+					]
+				]
+			]
+		];
+
+	GEditor->EditorAddModalWindow(Window.ToSharedRef());
+	return Answer;
+}
+
 void FControlRigEditor::HandleJumpToHyperlink(const UObject* InSubject)
 {
 	UControlRigBlueprint* RigBlueprint = GetControlRigBlueprint();
diff --git a/Engine/Plugins/Animation/ControlRig/Source/ControlRigEditor/Private/Editor/ControlRigEditor.h b/Engine/Plugins/Animation/ControlRig/Source/ControlRigEditor/Private/Editor/ControlRigEditor.h
--- a/Engine/Plugins/Animation/ControlRig/Source/ControlRigEditor/Private/Editor/ControlRigEditor.h
+++ b/Engine/Plugins/Animation/ControlRig/Source/ControlRigEditor/Private/Editor/ControlRigEditor.h
@@ -438,6 +438,7 @@ protected:
 	void OnRequestLocalizeFunctionDialog(FRigVMGraphFunctionIdentifier& InFunction, UControlRigBlueprint* InTargetBlueprint, bool bForce);
 	FRigVMController_BulkEditResult OnRequestBulkEditDialog(UControlRigBlueprint* InBlueprint, URigVMController* InController, URigVMLibraryNode* InFunction, ERigVMControllerBulkEditType InEditType);
 	bool OnRequestBreakLinksDialog(TArray<URigVMLink*> InLinks);
+	TRigVMTypeIndex OnRequestPinTypeSelectionDialog(const TArray<TRigVMTypeIndex>& InTypes);
 	void HandleJumpToHyperlink(const UObject* InSubject);
 	bool UpdateDefaultValueForVariable(FBPVariableDescription& InVariable, bool bUseCDO);
 
diff --git a/Engine/Plugins/Animation/ControlRig/Source/ControlRigEditor/Private/Graph/SControlRigGraphChangePinType.cpp b/Engine/Plugins/Animation/ControlRig/Source/ControlRigEditor/Private/Graph/SControlRigGraphChangePinType.cpp
--- a/Engine/Plugins/Animation/ControlRig/Source/ControlRigEditor/Private/Graph/SControlRigGraphChangePinType.cpp
+++ b/Engine/Plugins/Animation/ControlRig/Source/ControlRigEditor/Private/Graph/SControlRigGraphChangePinType.cpp
@@ -15,12 +15,13 @@
 
 #define LOCTEXT_NAMESPACE "SControlRigGraphChangePinType"
 
-static const FText ControlRigChangePinTypeMultipleValues = LOCTEXT("MultipleValues", "Multiple Values");
+static const FText ControlRigChangePinTypeMultipleValues = LOCTEXT("MultipleValues", " - ");
 
 void SControlRigChangePinType::Construct(const FArguments& InArgs)
 {
-	this->ModelPins = InArgs._ModelPins;
+	this->Types = InArgs._Types;
 	this->Blueprint = InArgs._Blueprint;
+	this->OnTypeSelected = InArgs._OnTypeSelected;
 
 	IPropertyAccessEditor& PropertyAccessEditor = IModularFeatures::Get().GetModularFeature<IPropertyAccessEditor>("PropertyAccessEditor");
 
@@ -103,34 +104,16 @@ FText SControlRigChangePinType::GetBindingText(const FRigVMTemplateArgumentType&
 	return FText();
 }
 
-FText SControlRigChangePinType::GetBindingText(URigVMPin* ModelPin) const
-{
-	if (ModelPin)
-	{
-		FRigVMTemplateArgumentType ArgumentType = ModelPin->GetRootPin()->GetTemplateArgumentType();
-		if(ModelPin->IsArrayElement())
-		{
-			const FRigVMRegistry& Registry = FRigVMRegistry::Get();
-			const TRigVMTypeIndex ArrayTypeIndex = Registry.GetTypeIndex(ArgumentType);
-			const TRigVMTypeIndex ElementTypeIndex = Registry.GetBaseTypeFromArrayTypeIndex(ArrayTypeIndex);
-			ArgumentType = Registry.GetType(ElementTypeIndex);
-		}
-		if(ArgumentType.IsValid())
-		{
-			return GetBindingText(ArgumentType);
-		}
-	}
-	return FText();
-}
-
 FText SControlRigChangePinType::GetBindingText() const
 {
-	if (ModelPins.Num() > 0)
+	if(Types.Num() > 0)
 	{
-		const FText FirstText = GetBindingText(ModelPins[0]);
-		for(int32 Index = 1; Index < ModelPins.Num(); Index++)
+		const FRigVMTemplateArgumentType FirstType = FRigVMRegistry::Get().GetType(Types[0]);
+		const FText FirstText = GetBindingText(FirstType);
+		for(int32 Index = 1; Index < Types.Num(); Index++)
 		{
-			if(!GetBindingText(ModelPins[Index]).EqualTo(FirstText))
+			const FRigVMTemplateArgumentType Type = FRigVMRegistry::Get().GetType(Types[Index]);
+			if(!GetBindingText(Type).EqualTo(FirstText))
 			{
 				return ControlRigChangePinTypeMultipleValues;
 			}
@@ -145,9 +128,10 @@ const FSlateBrush* SControlRigChangePinType::GetBindingImage() const
 	static FName TypeIcon(TEXT("Kismet.VariableList.TypeIcon"));
 	static FName ArrayTypeIcon(TEXT("Kismet.VariableList.ArrayTypeIcon"));
 
-	if (ModelPins.Num() > 0)
+	if(Types.Num() > 0)
 	{
-		if(ModelPins[0]->IsArray())
+		const FRigVMTemplateArgumentType Type = FRigVMRegistry::Get().GetType(Types[0]);
+		if(Type.IsArray())
 		{
 			return FAppStyle::GetBrush(ArrayTypeIcon);
 		}
@@ -157,9 +141,9 @@ const FSlateBrush* SControlRigChangePinType::GetBindingImage() const
 
 FLinearColor SControlRigChangePinType::GetBindingColor() const
 {
-	if(ModelPins.Num() > 0)
+	if(Types.Num() > 0)
 	{
-		const FRigVMTemplateArgumentType Type = ModelPins[0]->GetTemplateArgumentType();
+		const FRigVMTemplateArgumentType Type = FRigVMRegistry::Get().GetType(Types[0]);
 		const FEdGraphPinType PinType = RigVMTypeUtils::PinTypeFromCPPType(Type.CPPType, Type.CPPTypeObject);
 		const UControlRigGraphSchema* Schema = GetDefault<UControlRigGraphSchema>();
 		return Schema->GetPinTypeColor(PinType);
@@ -169,119 +153,58 @@ FLinearColor SControlRigChangePinType::GetBindingColor() const
 
 void SControlRigChangePinType::FillPinTypeMenu(FMenuBuilder& MenuBuilder)
 {
-	if(ModelPins.Num() == 0)
-	{
-		return;
-	}
-
-	URigVMGraph* Model = ModelPins[0]->GetGraph();
-	if(Model == nullptr)
+	if(Types.Num() == 0)
 	{
 		return;
 	}
 
 	struct FArgumentInfo
 	{
-		static FArgumentInfo Make(const FProperty* InProperty, const bool bInIsFilteredOut)
+		static FArgumentInfo Make(const FProperty* InProperty)
 		{
 			FArgumentInfo Info;
 			Info.Property = InProperty;
-			Info.bIsFilteredOut = bInIsFilteredOut;
 			return Info;
 		}
 
-		static FArgumentInfo Make(const FRigVMTemplateArgumentType& InType, const bool bInIsFilteredOut)
+		static FArgumentInfo Make(const FRigVMTemplateArgumentType& InType)
 		{
 			FArgumentInfo Info;
 			Info.Property = nullptr;
 			Info.Type = InType;
-			Info.bIsFilteredOut = bInIsFilteredOut;
 			return Info;
 		}
 
 		const FProperty* Property;
-		bool bIsFilteredOut;
 		FRigVMTemplateArgumentType Type;
 	};
 
 	typedef TPair<TRigVMTypeIndex, FArgumentInfo> FTypePair;
-	TMap<TRigVMTypeIndex, FArgumentInfo> Types;
-	for(URigVMPin* ModelPin : ModelPins)
+	TMap<TRigVMTypeIndex, FArgumentInfo> TypesToInfo;
+	for(TRigVMTypeIndex& TypeIndex : Types)
 	{
-		const bool bIsArrayElement = ModelPin->IsArrayElement();
-		if(bIsArrayElement)
+		if (TypeIndex == INDEX_NONE)
 		{
-			ModelPin = ModelPin->GetParentPin();
+			// skip invalid permutations
+			continue;
 		}
 
-		if(!ModelPin->IsRootPin())
+		if(TypeIndex == RigVMTypeUtils::TypeIndex::Float)
 		{
-			continue;
+			TypeIndex = RigVMTypeUtils::TypeIndex::Double;
 		}
-
-		if(const URigVMTemplateNode* TemplateNode = Cast<URigVMTemplateNode>(ModelPin->GetNode()))
+		else if(TypeIndex == RigVMTypeUtils::TypeIndex::FloatArray)
 		{
-			if(const FRigVMTemplate* Template = TemplateNode->GetTemplate())
-			{
-				const FName ArgumentName = ModelPin->GetFName();
-				if(const FRigVMTemplateArgument* Argument = Template->FindArgument(ArgumentName))
-				{
-					const TArray<TRigVMTypeIndex>& AllArgumentTypes = Argument->GetTypeIndices();
-					const TArray<TRigVMTypeIndex>& FilteredArgumentTypes = Argument->GetSupportedTypeIndices(TemplateNode->GetFilteredPermutationsIndices());
-					
-					for(int32 PermutationIndex = 0; PermutationIndex < Template->NumPermutations(); PermutationIndex++)
-					{
-						TRigVMTypeIndex ArgumentTypeIndex = AllArgumentTypes[PermutationIndex];
-
-						if (ArgumentTypeIndex == INDEX_NONE)
-						{
-							// skip invalid permutations
-							continue;
-						}
-
-						bool bIsFilteredOut = false;
-						if(!FilteredArgumentTypes.Contains(ArgumentTypeIndex))
-						{
-							bIsFilteredOut = true;
-						}
-
-						if(ArgumentTypeIndex == RigVMTypeUtils::TypeIndex::Float)
-						{
-							ArgumentTypeIndex = RigVMTypeUtils::TypeIndex::Double;
-						}
-						else if(ArgumentTypeIndex == RigVMTypeUtils::TypeIndex::FloatArray)
-						{
-							ArgumentTypeIndex = RigVMTypeUtils::TypeIndex::DoubleArray;
-						}
-						
-						if(Types.Contains(ArgumentTypeIndex))
-						{
-							continue;
-						}
-
-						if(const FRigVMFunction* Permutation = Template->GetPermutation(PermutationIndex))
-						{
-							if(UScriptStruct* FunctionStruct = Permutation->Struct)
-							{
-								if(const FProperty* Property = FunctionStruct->FindPropertyByName(ArgumentName))
-								{
-									Types.Add(ArgumentTypeIndex, FArgumentInfo::Make(Property, bIsFilteredOut));
-									continue;
-								}
-							}
-						}
-
-						if(bIsArrayElement)
-						{
-							ArgumentTypeIndex = FRigVMRegistry::Get().GetBaseTypeFromArrayTypeIndex(ArgumentTypeIndex);
-						}
-						
-						const FRigVMTemplateArgumentType ArgumentType = FRigVMRegistry::Get().GetType(ArgumentTypeIndex);
-						Types.Add(ArgumentTypeIndex, FArgumentInfo::Make(ArgumentType, bIsFilteredOut));
-					}
-				}
-			}
+			TypeIndex = RigVMTypeUtils::TypeIndex::DoubleArray;
+		}
+		
+		if(TypesToInfo.Contains(TypeIndex))
+		{
+			continue;
 		}
+
+		const FRigVMTemplateArgumentType ArgumentType = FRigVMRegistry::Get().GetType(TypeIndex);
+		TypesToInfo.Add(TypeIndex, FArgumentInfo::Make(ArgumentType));
 	}
 
 	// sort the types and put them into an array
@@ -300,7 +223,7 @@ void SControlRigChangePinType::FillPinTypeMenu(FMenuBuilder& MenuBuilder)
 		TEXT("FMatrix")
 	};
 	TArray<FTypePair> SortedTypes;
-	for(const FTypePair& Pair : Types)
+	for(const FTypePair& Pair : TypesToInfo)
 	{
 		SortedTypes.Add(Pair);
 	}
@@ -338,7 +261,6 @@ void SControlRigChangePinType::FillPinTypeMenu(FMenuBuilder& MenuBuilder)
 		{
 			const TRigVMTypeIndex& TypeIndex = SortedTypes[SortedIndex].Key;
 			const FRigVMTemplateArgumentType Type = FRigVMRegistry::Get().GetType(TypeIndex);
-			const bool bIsFilteredOut = SortedTypes[SortedIndex].Value.bIsFilteredOut;
 			if (Type.CPPTypeObject != nullptr && !IsValid(Type.CPPTypeObject))
 			{
 				continue;
@@ -389,7 +311,7 @@ void SControlRigChangePinType::FillPinTypeMenu(FMenuBuilder& MenuBuilder)
 					[
 						SNew(SImage)
 						.Image(FBlueprintEditorUtils::GetIconFromPin(PinType, true))
-						.ColorAndOpacity(Schema->GetPinTypeColor(PinType) * (bIsFilteredOut ? 0.5f : 1.f))
+						.ColorAndOpacity(Schema->GetPinTypeColor(PinType))
 					]
 					+SHorizontalBox::Slot()
 					.AutoWidth()
@@ -398,9 +320,8 @@ void SControlRigChangePinType::FillPinTypeMenu(FMenuBuilder& MenuBuilder)
 					[
 						SNew(STextBlock)
 						.Text(GetBindingText(Type))
-						.ColorAndOpacity(FLinearColor::White * (bIsFilteredOut ? 0.5f : 1.f))
-						.ToolTipText(bIsFilteredOut ? LOCTEXT("WildcardAvailableTypeTooltip","Will break connections if resolved to this type.") :
-							LOCTEXT("WildcardAvailableFilterTypeTooltip", "Available filtered type"))
+						.ColorAndOpacity(FLinearColor::White)
+						.ToolTipText(LOCTEXT("WildcardAvailableFilterTypeTooltip", "Available filtered type"))
 					]);
 		}
 	}
@@ -409,56 +330,12 @@ void SControlRigChangePinType::FillPinTypeMenu(FMenuBuilder& MenuBuilder)
 
 void SControlRigChangePinType::HandlePinTypeChanged(FRigVMTemplateArgumentType InType)
 {
-	if(ModelPins.IsEmpty() || (Blueprint == nullptr))
+	if (OnTypeSelected.IsBound())
 	{
+		OnTypeSelected.Execute(FRigVMRegistry::Get().GetTypeIndex(InType));
 		return;
 	}
-
-	for(URigVMPin* ModelPin : ModelPins)
-	{
-		URigVMGraph* Model = ModelPin->GetGraph();
-		if(Model == nullptr)
-		{
-			continue;
-		}
-
-		URigVMController* Controller = Blueprint->GetOrCreateController(Model);
-		if(Controller == nullptr)
-		{
-			continue;
-		}
-
-		TRigVMTypeIndex TypeIndex = FRigVMRegistry::Get().GetTypeIndex(InType);
-		if(ModelPin->IsArrayElement())
-		{
-			TypeIndex = FRigVMRegistry::Get().GetArrayTypeFromBaseTypeIndex(TypeIndex);
-			ModelPin = ModelPin->GetParentPin();
-		}
-
-		if(!ModelPin->IsRootPin())
-		{
-			continue;
-		}
-		
-		Controller->ResolveWildCardPin(ModelPin->GetRootPin()->GetPinPath(), TypeIndex, true, true);
-	}
 }
 
-/////////////////////////////////////////////////////////////////////////////////////////////////////////////
-
-void SControlRigGraphChangePinType::Construct(const FArguments& InArgs, UEdGraphPin* InGraphPinObj)
-{
-	this->ModelPins = InArgs._ModelPins;
-	this->Blueprint = InArgs._Blueprint;
-
-	SGraphPin::Construct(SGraphPin::FArguments(), InGraphPinObj);
-}
-
-TSharedRef<SWidget>	SControlRigGraphChangePinType::GetDefaultValueWidget()
-{
-	return SNew(SControlRigChangePinType)
-		.Blueprint(Blueprint)
-		.ModelPins(ModelPins);
-}
 
 #undef LOCTEXT_NAMESPACE
diff --git a/Engine/Plugins/Animation/ControlRig/Source/ControlRigEditor/Private/Graph/SControlRigGraphChangePinType.h b/Engine/Plugins/Animation/ControlRig/Source/ControlRigEditor/Private/Graph/SControlRigGraphChangePinType.h
--- a/Engine/Plugins/Animation/ControlRig/Source/ControlRigEditor/Private/Graph/SControlRigGraphChangePinType.h
+++ b/Engine/Plugins/Animation/ControlRig/Source/ControlRigEditor/Private/Graph/SControlRigGraphChangePinType.h
@@ -10,17 +10,21 @@
 #include "IPropertyAccessEditor.h"
 #include "ControlRigBlueprint.h"
 
+DECLARE_DELEGATE_OneParam(FOnTypeSelected, TRigVMTypeIndex);
+
 class SControlRigChangePinType : public SCompoundWidget
 {
 public:
 
 	SLATE_BEGIN_ARGS(SControlRigChangePinType)
-	: _ModelPins()
+	: _Types()
     , _Blueprint(nullptr)
+	, _OnTypeSelected(nullptr)
 	{}
 
-		SLATE_ARGUMENT(TArray<URigVMPin*>, ModelPins)
+		SLATE_ARGUMENT(TArray<TRigVMTypeIndex>, Types)
 		SLATE_ARGUMENT(UControlRigBlueprint*, Blueprint)
+		SLATE_EVENT(FOnTypeSelected, OnTypeSelected)
 
 	SLATE_END_ARGS()
 
@@ -39,30 +43,8 @@ protected:
 	void FillPinTypeMenu( FMenuBuilder& MenuBuilder );
 	void HandlePinTypeChanged(FRigVMTemplateArgumentType InType);
 
-	TArray<URigVMPin*> ModelPins;
+	TArray<TRigVMTypeIndex> Types;
 	UControlRigBlueprint* Blueprint;
+	FOnTypeSelected OnTypeSelected;
 	FPropertyBindingWidgetArgs BindingArgs;
 };
-
-class SControlRigGraphChangePinType : public SGraphPin
-{
-public:
-
-	SLATE_BEGIN_ARGS(SControlRigGraphChangePinType){}
-
-		SLATE_ARGUMENT(TArray<URigVMPin*>, ModelPins)
-		SLATE_ARGUMENT(UControlRigBlueprint*, Blueprint)
-
-	SLATE_END_ARGS()
-
-	void Construct(const FArguments& InArgs, UEdGraphPin* InGraphPinObj);
-
-protected:
-
-	//~ Begin SGraphPin Interface
-	virtual TSharedRef<SWidget>	GetDefaultValueWidget() override;
-	//~ End SGraphPin Interface
-
-	TArray<URigVMPin*> ModelPins;
-	UControlRigBlueprint* Blueprint;
-};
diff --git a/Engine/Plugins/Experimental/AnimNext/Source/AnimNextInterfaceGraphUncookedOnly/Private/AnimNextInterfaceGraph_EditorData.cpp b/Engine/Plugins/Experimental/AnimNext/Source/AnimNextInterfaceGraphUncookedOnly/Private/AnimNextInterfaceGraph_EditorData.cpp
--- a/Engine/Plugins/Experimental/AnimNext/Source/AnimNextInterfaceGraphUncookedOnly/Private/AnimNextInterfaceGraph_EditorData.cpp
+++ b/Engine/Plugins/Experimental/AnimNext/Source/AnimNextInterfaceGraphUncookedOnly/Private/AnimNextInterfaceGraph_EditorData.cpp
@@ -229,7 +229,6 @@ void UAnimNextInterfaceGraph_EditorData::RefreshAllModels(EAnimNextInterfaceGrap
 	{
 		URigVMGraph* GraphToClean = SortedGraphsToClean[GraphIndex];
 		URigVMController* Controller = RigVMClient.GetOrCreateController(GraphToClean);
-		TGuardValue<bool> RecomputeGuard(Controller->bSuspendRecomputingOuterTemplateFilters, true);
 		//TGuardValue<bool> GuardEditGraph(GraphToClean->bEditable, true);
 		FRigVMControllerNotifGuard NotifGuard(Controller, true);
 
@@ -237,14 +236,6 @@ void UAnimNextInterfaceGraph_EditorData::RefreshAllModels(EAnimNextInterfaceGrap
 		{
 			Controller->RemoveUnusedOrphanedPins(ModelNode);
 		}
-
-		if (bIsPostLoad)
-		{
-			if (URigVMLibraryNode* LibraryNode = GraphToClean->GetTypedOuter<URigVMLibraryNode>())
-			{
-				Controller->UpdateLibraryTemplate(LibraryNode, false);
-			}
-		}
 	}
 }
 
diff --git a/Engine/Source/Developer/RigVMDeveloper/Private/RigVMModel/Nodes/RigVMDispatchNode.cpp b/Engine/Source/Developer/RigVMDeveloper/Private/RigVMModel/Nodes/RigVMDispatchNode.cpp
--- a/Engine/Source/Developer/RigVMDeveloper/Private/RigVMModel/Nodes/RigVMDispatchNode.cpp
+++ b/Engine/Source/Developer/RigVMDeveloper/Private/RigVMModel/Nodes/RigVMDispatchNode.cpp
@@ -9,7 +9,7 @@ FString URigVMDispatchNode::GetNodeTitle() const
 {
 	if (const FRigVMDispatchFactory* Factory = GetFactory())
 	{
-		return Factory->GetNodeTitle(GetFilteredTypes());
+		return Factory->GetNodeTitle(GetTemplatePinTypeMap());
 	}
 	return Super::GetNodeTitle();
 }
@@ -18,7 +18,7 @@ FText URigVMDispatchNode::GetToolTipText() const
 {
 	if (const FRigVMDispatchFactory* Factory = GetFactory())
 	{
-		return Factory->GetNodeTooltip(GetFilteredTypes());
+		return Factory->GetNodeTooltip(GetTemplatePinTypeMap());
 	}
 	return URigVMNode::GetToolTipText();
 }
@@ -200,20 +200,6 @@ const UScriptStruct* URigVMDispatchNode::GetFactoryStruct() const
 	return nullptr;
 }
 
-const FRigVMTemplateTypeMap& URigVMDispatchNode::GetFilteredTypes() const
-{
-	const TArray<int32> PermutationIndices = GetFilteredPermutationsIndices();
-	if(PermutationIndices.Num() == 1)
-	{
-		TypesFromPins = GetTypesForPermutation(PermutationIndices[0]);
-	}
-	else
-	{
-		TypesFromPins.Reset();
-	}
-	return TypesFromPins;
-}
-
 void URigVMDispatchNode::InvalidateCache()
 {
 	Super::InvalidateCache();
@@ -263,7 +249,7 @@ FRigVMStructUpgradeInfo URigVMDispatchNode::GetUpgradeInfo() const
 {
 	if(const FRigVMDispatchFactory* Factory = GetFactory())
 	{
-		return Factory->GetUpgradeInfo(GetFilteredTypes(), GetDispatchContext());
+		return Factory->GetUpgradeInfo(GetTemplatePinTypeMap(), GetDispatchContext());
 	}
 	return FRigVMStructUpgradeInfo();
 }
diff --git a/Engine/Source/Developer/RigVMDeveloper/Private/RigVMModel/Nodes/RigVMFunctionEntryNode.cpp b/Engine/Source/Developer/RigVMDeveloper/Private/RigVMModel/Nodes/RigVMFunctionEntryNode.cpp
--- a/Engine/Source/Developer/RigVMDeveloper/Private/RigVMModel/Nodes/RigVMFunctionEntryNode.cpp
+++ b/Engine/Source/Developer/RigVMDeveloper/Private/RigVMModel/Nodes/RigVMFunctionEntryNode.cpp
@@ -5,24 +5,6 @@
 
 #include UE_INLINE_GENERATED_CPP_BY_NAME(RigVMFunctionEntryNode)
 
-const FRigVMTemplate* URigVMFunctionEntryNode::GetTemplate() const
-{
-	if (URigVMLibraryNode* LibraryNode = GetTypedOuter<URigVMLibraryNode>())
-	{
-		return LibraryNode->GetTemplate();
-	}
-	return nullptr;
-}
-
-FName URigVMFunctionEntryNode::GetNotation() const
-{
-	if (URigVMLibraryNode* LibraryNode = GetTypedOuter<URigVMLibraryNode>())
-	{
-		return LibraryNode->GetNotation();
-	}
-	return NAME_None;
-}
-
 FLinearColor URigVMFunctionEntryNode::GetNodeColor() const
 {
 	if(URigVMGraph* RootGraph = GetRootGraph())
@@ -48,7 +30,7 @@ FString URigVMFunctionEntryNode::GetNodeTitle() const
 
 FText URigVMFunctionEntryNode::GetToolTipText() const
 {
-	return FText::FromName(GetTemplate()->GetNotation());
+	return FText::FromName(GetGraph()->GetOuter()->GetFName());
 }
 
 FText URigVMFunctionEntryNode::GetToolTipTextForPin(const URigVMPin* InPin) const
diff --git a/Engine/Source/Developer/RigVMDeveloper/Private/RigVMModel/Nodes/RigVMFunctionReturnNode.cpp b/Engine/Source/Developer/RigVMDeveloper/Private/RigVMModel/Nodes/RigVMFunctionReturnNode.cpp
--- a/Engine/Source/Developer/RigVMDeveloper/Private/RigVMModel/Nodes/RigVMFunctionReturnNode.cpp
+++ b/Engine/Source/Developer/RigVMDeveloper/Private/RigVMModel/Nodes/RigVMFunctionReturnNode.cpp
@@ -5,24 +5,6 @@
 
 #include UE_INLINE_GENERATED_CPP_BY_NAME(RigVMFunctionReturnNode)
 
-const FRigVMTemplate* URigVMFunctionReturnNode::GetTemplate() const
-{
-	if (URigVMLibraryNode* LibraryNode = GetTypedOuter<URigVMLibraryNode>())
-	{
-		return LibraryNode->GetTemplate();
-	}
-	return nullptr;
-}
-
-FName URigVMFunctionReturnNode::GetNotation() const
-{
-	if (URigVMLibraryNode* LibraryNode = GetTypedOuter<URigVMLibraryNode>())
-	{
-		return LibraryNode->GetNotation();
-	}
-	return NAME_None;
-}
-
 FLinearColor URigVMFunctionReturnNode::GetNodeColor() const
 {
 	if(URigVMGraph* RootGraph = GetRootGraph())
@@ -48,7 +30,7 @@ FString URigVMFunctionReturnNode::GetNodeTitle() const
 
 FText URigVMFunctionReturnNode::GetToolTipText() const
 {
-	return FText::FromName(GetTemplate()->GetNotation());
+	return FText::FromName(GetGraph()->GetOuter()->GetFName());
 }
 
 FText URigVMFunctionReturnNode::GetToolTipTextForPin(const URigVMPin* InPin) const
diff --git a/Engine/Source/Developer/RigVMDeveloper/Private/RigVMModel/Nodes/RigVMLibraryNode.cpp b/Engine/Source/Developer/RigVMDeveloper/Private/RigVMModel/Nodes/RigVMLibraryNode.cpp
--- a/Engine/Source/Developer/RigVMDeveloper/Private/RigVMModel/Nodes/RigVMLibraryNode.cpp
+++ b/Engine/Source/Developer/RigVMDeveloper/Private/RigVMModel/Nodes/RigVMLibraryNode.cpp
@@ -13,10 +13,6 @@
 const TArray<URigVMNode*> URigVMLibraryNode::EmptyNodes;
 const TArray<URigVMLink*> URigVMLibraryNode::EmptyLinks;
 
-URigVMLibraryNode::URigVMLibraryNode()
-{
-	Template.Notation = *FString::Printf(TEXT("%s()"), *GetName());
-}
 
 bool URigVMLibraryNode::IsDefinedAsConstant() const
 {
@@ -39,19 +35,9 @@ bool URigVMLibraryNode::IsDefinedAsVarying() const
 	return false;
 }
 
-const FRigVMTemplate* URigVMLibraryNode::GetTemplate() const
-{
-	return &Template;
-}
-
-FName URigVMLibraryNode::GetNotation() const
-{
-	return Template.GetNotation();
-}
-
 FText URigVMLibraryNode::GetToolTipText() const
 {
-	return FText::FromName(Template.GetNotation());
+	return FText::FromName(GetFunctionIdentifier().HostObject.GetLongPackageFName());
 }
 
 TArray<int32> URigVMLibraryNode::GetInstructionsForVMImpl(URigVM* InVM, const FRigVMASTProxy& InProxy) const
diff --git a/Engine/Source/Developer/RigVMDeveloper/Private/RigVMModel/Nodes/RigVMTemplateNode.cpp b/Engine/Source/Developer/RigVMDeveloper/Private/RigVMModel/Nodes/RigVMTemplateNode.cpp
--- a/Engine/Source/Developer/RigVMDeveloper/Private/RigVMModel/Nodes/RigVMTemplateNode.cpp
+++ b/Engine/Source/Developer/RigVMDeveloper/Private/RigVMModel/Nodes/RigVMTemplateNode.cpp
@@ -68,7 +68,6 @@ URigVMTemplateNode::URigVMTemplateNode()
 void URigVMTemplateNode::PreSave(FObjectPreSaveContext SaveContext)
 {
 	Super::PreSave(SaveContext);
-	ConvertPreferredTypesToString();
 }
 
 void URigVMTemplateNode::PostLoad()
@@ -100,7 +99,7 @@ void URigVMTemplateNode::PostLoad()
 	{
 		const FRigVMRegistry& Registry = FRigVMRegistry::Get();
 
-		PreferredPermutationPairs.Reserve(PreferredPermutationTypes_DEPRECATED.Num());
+		PreferredPermutationPairs_DEPRECATED.Reserve(PreferredPermutationTypes_DEPRECATED.Num());
 
 		for(const FString& PreferredPermutation : PreferredPermutationTypes_DEPRECATED)
 		{
@@ -109,13 +108,12 @@ void URigVMTemplateNode::PostLoad()
 
 			CPPType = RigVMTypeUtils::PostProcessCPPType(CPPType);
 			const FRigVMTemplateArgumentType Type = Registry.FindTypeFromCPPType(CPPType);
-			PreferredPermutationPairs.Emplace(*ArgName, Registry.GetTypeIndex(Type));
+			PreferredPermutationPairs_DEPRECATED.Emplace(*ArgName, Registry.GetTypeIndex(Type));
 		}
 
 		PreferredPermutationTypes_DEPRECATED.Reset();
 	}
 
-	ConvertPreferredTypesToTypeIndex();
 	InvalidateCache();
 
 	// the template notation may have changed
@@ -125,54 +123,6 @@ void URigVMTemplateNode::PostLoad()
 	}
 }
 
-void URigVMTemplateNode::ConvertPreferredTypesToString()
-{
-	// we rely on the strings being serialized
-	for(FRigVMTemplatePreferredType& PreferredType : PreferredPermutationPairs)
-	{
-		PreferredType.UpdateStringFromIndex();
-	}
-}
-
-void URigVMTemplateNode::ConvertPreferredTypesToTypeIndex()
-{
-	const FRigVMTemplate* Template = GetTemplate();
-	
-	// we rely on the strings being serialized - so on load we need to update the type index again 
-	for(FRigVMTemplatePreferredType& PreferredType : PreferredPermutationPairs)
-	{
-		PreferredType.UpdateIndexFromString();
-
-		if(Template)
-		{
-			if(const FRigVMTemplateArgument* Argument = Template->FindArgument(PreferredType.GetArgument()))
-			{
-				TRigVMTypeIndex ResolvedTypeIndex = INDEX_NONE;
-				if(Argument->SupportsTypeIndex(PreferredType.GetTypeIndex(), &ResolvedTypeIndex))
-				{
-					if(ResolvedTypeIndex != PreferredType.GetTypeIndex())
-					{
-						PreferredType.TypeIndex = ResolvedTypeIndex;
-						PreferredType.UpdateStringFromIndex();
-					}
-				}
-			}
-		}
-	}
-}
-
-TRigVMTypeIndex URigVMTemplateNode::GetPreferredType(const FName& ArgumentName) const
-{
-	for (const FRigVMTemplatePreferredType& PreferredType : PreferredPermutationPairs)
-	{
-		if (PreferredType.Argument == ArgumentName)
-		{
-			return PreferredType.GetTypeIndex();
-		}
-	}
-	return INDEX_NONE;
-}
-
 UScriptStruct* URigVMTemplateNode::GetScriptStruct() const
 {
 	if(const FRigVMFunction* Function = GetResolvedFunction())
@@ -216,7 +166,7 @@ FText URigVMTemplateNode::GetToolTipText() const
 {
 	if(const FRigVMTemplate* Template = GetTemplate())
 	{
-		const TArray<int32> PermutationIndices = GetFilteredPermutationsIndices();
+		const TArray<int32> PermutationIndices = GetResolvedPermutationIndices(false);
 		return Template->GetTooltipText(PermutationIndices);
 	}
 	return Super::GetToolTipText();
@@ -234,12 +184,9 @@ FText URigVMTemplateNode::GetToolTipTextForPin(const URigVMPin* InPin) const
 			{
 				FString Tooltip;
 				
-				if(FilteredPermutations.Num() == GetTemplate()->NumPermutations())
+				if(Arg->GetTypeIndices().Num() > 100)
 				{
-					if(Arg->GetTypeIndices().Num() > 100)
-					{
-						Tooltip = TEXT("Supports any type.");
-					}
+					Tooltip = TEXT("Supports any type.");
 				}
 				
 				if(Tooltip.IsEmpty())
@@ -255,10 +202,6 @@ FText URigVMTemplateNode::GetToolTipTextForPin(const URigVMPin* InPin) const
 						}
 						TypesPrinted.Add(TypeIndex);
 						FString Type = FRigVMRegistry::Get().GetType(TypeIndex).CPPType.ToString();
-						if (!FilteredPermutations.Contains(Index))
-						{
-							Type += TEXT(" : Breaks Connections");
-						}
 						SupportedTypesJoined += Type + TEXT("\n");
 					}
 					Tooltip = TEXT("Supported Types:\n\n") + SupportedTypesJoined;
@@ -292,14 +235,11 @@ TArray<URigVMPin*> URigVMTemplateNode::GetAggregatePins(const ERigVMPinDirection
 #if UE_RIGVM_AGGREGATE_NODES_ENABLED
 	if (const FRigVMTemplate* Template = GetTemplate())
 	{
-		TArray<int32> Permutations = FilteredPermutations;
-		if (Permutations.IsEmpty())
+		TArray<int32> Permutations;
+		Permutations.Reserve(Template->NumPermutations());
+		for (int32 i=0; i<Template->NumPermutations(); ++i)
 		{
-			Permutations.Reserve(Template->NumPermutations());
-			for (int32 i=0; i<Template->NumPermutations(); ++i)
-			{
-				Permutations.Add(i);
-			}
+			Permutations.Add(i);
 		}
 		
 		if (Permutations.IsEmpty())
@@ -358,12 +298,7 @@ FName URigVMTemplateNode::GetNextAggregateName(const FName& InLastAggregatePinNa
 #if UE_RIGVM_AGGREGATE_NODES_ENABLED
 	if (const FRigVMTemplate* Template = GetTemplate())
 	{
-		if (FilteredPermutations.IsEmpty())
-		{
-			return NextName;
-		}
-
-		if (const FRigVMFunction* Permutation = Template->GetPermutation(FilteredPermutations[0]))
+		if (const FRigVMFunction* Permutation = Template->GetPermutation(0))
 		{
 			const TSharedPtr<FStructOnScope> StructOnScope = MakeShareable(new FStructOnScope(Permutation->Struct));
 			if (const UScriptStruct* ScriptStruct = Cast<UScriptStruct>(StructOnScope->GetStruct()))
@@ -485,83 +420,50 @@ bool URigVMTemplateNode::SupportsType(const URigVMPin* InPin, TRigVMTypeIndex In
 	return false;
 }
 
-bool URigVMTemplateNode::FilteredSupportsType(const URigVMPin* InPin, TRigVMTypeIndex InTypeIndex, TRigVMTypeIndex* OutTypeIndex, bool bAllowFloatingPointCasts)
+TArray<const FRigVMFunction*> URigVMTemplateNode::GetResolvedPermutations() const
 {
-	if (OutTypeIndex)
-	{
-		*OutTypeIndex = INDEX_NONE; 
-	}
-
-	const URigVMPin* RootPin = InPin;
-	bool bIsArrayElement = false;
-	bool bIsStructElement = false;
-	if (URigVMPin* ParentPin = InPin->GetParentPin())
-	{
-		RootPin = ParentPin;
-		if (ParentPin->IsArray())
-		{
-			bIsArrayElement = true;
-		}
-		else if (ParentPin->IsStruct())
-		{
-			bIsStructElement = true;
-		}
-	}
-
-	if (bIsStructElement)
+	TArray<const FRigVMFunction*> Functions;
+	if (const FRigVMTemplate* Template = GetTemplate())
 	{
-		const bool bResult = InPin->GetTypeIndex() == InTypeIndex;
-		if(bResult)
+		TArray<int32> ResolvedPermutations = GetResolvedPermutationIndices(true);
+		for(const int32 Index : ResolvedPermutations)
 		{
-			if(OutTypeIndex)
-			{
-				*OutTypeIndex = InTypeIndex; 
-			}
+			Functions.Add(Template->GetPermutation(Index));
 		}
-		return bResult;
-	}
-
-	const FRigVMTemplateArgument* Argument = GetTemplate()->FindArgument(RootPin->GetFName());
-	if (Argument == nullptr)
-	{
-		return false;
-	}
-
-	const FRigVMRegistry& Registry = FRigVMRegistry::Get();
-
-	TRigVMTypeIndex RootTypeIndex = InTypeIndex;
-	if (bIsArrayElement)
-	{
-		RootTypeIndex = Registry.GetArrayTypeFromBaseTypeIndex(RootTypeIndex);
-	}
-
-	if (FilteredPermutations.Num() == GetTemplate()->NumPermutations())
-	{
-		return GetTemplate()->ArgumentSupportsTypeIndex(RootPin->GetFName(), InTypeIndex, OutTypeIndex);
 	}
+	return Functions;
+}
 
-	const TArray<TRigVMTypeIndex>& TypeIndices = Argument->GetTypeIndices();
-	for (const int32& PermutationIndex : FilteredPermutations)
+TArray<int32> URigVMTemplateNode::GetResolvedPermutationIndices(bool bAllowFloatingPointCasts) const
+{
+	TArray<int32> ResolvedPermutations;
+	if (const FRigVMTemplate* Template = GetTemplate())
 	{
-		const TRigVMTypeIndex& FilteredTypeIndex = TypeIndices[PermutationIndex];
-		if (Registry.CanMatchTypes(FilteredTypeIndex, RootTypeIndex, bAllowFloatingPointCasts))
-		{
-			return true;
-		}
+		FRigVMTemplateTypeMap PinTypeMap = GetTemplatePinTypeMap();
+		Template->Resolve(PinTypeMap, ResolvedPermutations, bAllowFloatingPointCasts);
 	}
-
-	return false;
+	return ResolvedPermutations;
 }
 
-TArray<const FRigVMFunction*> URigVMTemplateNode::GetResolvedPermutations() const
+FRigVMTemplateTypeMap URigVMTemplateNode::GetTemplatePinTypeMap(bool bIncludeHiddenPins, bool bIncludeExecutePins) const
 {
-	TArray<int32> Indices = GetFilteredPermutationsIndices();
-	TArray<const FRigVMFunction*> Functions;
-	for(const int32 Index : Indices)
+	FRigVMTemplateTypeMap PinTypeMap;
+	for (const URigVMPin* Pin : GetPins())
 	{
-		Functions.Add(GetTemplate()->GetPermutation(Index));
+		if (!bIncludeHiddenPins && Pin->GetDirection() == ERigVMPinDirection::Hidden)
+		{
+			continue;
+		}
+		if (!bIncludeExecutePins && Pin->IsExecuteContext())
+		{
+			continue;
+		}
+		if (!Pin->IsWildCard())
+		{
+			PinTypeMap.Add(Pin->GetFName(), Pin->GetTypeIndex());
+		}
 	}
-	return Functions;
+	return PinTypeMap;
 }
 
 const FRigVMTemplate* URigVMTemplateNode::GetTemplate() const
@@ -602,13 +504,32 @@ bool URigVMTemplateNode::IsFullyUnresolved() const
 {
 	check(GetTemplate());
 
+	if (const FRigVMTemplate* Template = GetTemplate())
+	{
+		for (int32 i=0; i<Template->NumArguments(); ++i)
+		{
+			const FRigVMTemplateArgument* Argument = Template->GetArgument(i);
+			if (!Argument->IsSingleton())
+			{
+				if (const URigVMPin* Pin = FindPin(Argument->GetName().ToString()))
+				{
+					if (!Pin->IsWildCard())
+					{
+						return false;
+					}
+				}
+			}
+		}
+	}
+
 	// all permutations are available means we haven't resolved any wildcard pin
-	return GetFilteredPermutationsIndices().Num() == GetTemplate()->NumPermutations();
+	return true;
 }
 
 FString URigVMTemplateNode::GetInitialDefaultValueForPin(const FName& InRootPinName, const TArray<int32>& InPermutationIndices) const
 {
-	if(GetTemplate() == nullptr)
+	const FRigVMTemplate* Template = GetTemplate();
+	if(Template == nullptr)
 	{
 		return FString();
 	}
@@ -622,11 +543,10 @@ FString URigVMTemplateNode::GetInitialDefaultValueForPin(const FName& InRootPinN
 	TArray<int32> PermutationIndices = InPermutationIndices;
 	if(PermutationIndices.IsEmpty())
 	{
-		PermutationIndices = GetFilteredPermutationsIndices();
+		PermutationIndices = GetResolvedPermutationIndices(false);
 	}
 
 	FString DefaultValue;
-	const FRigVMTemplate* Template = GetTemplate();
 	const FRigVMDispatchFactory* Factory = Template->GetDispatchFactory(); 
 	const FRigVMTemplateArgument* Argument = Template->FindArgument(InRootPinName);
 
@@ -737,9 +657,11 @@ FName URigVMTemplateNode::GetDisplayNameForPin(const FName& InRootPinName,
 	if(const FRigVMTemplate* Template = GetTemplate())
 	{
 		const TArray<int32>* PermutationIndicesPtr = &InPermutationIndices;
+		TArray<int32> AllPermutations;
 		if(PermutationIndicesPtr->IsEmpty())
 		{
-			PermutationIndicesPtr = &GetFilteredPermutationsIndices();
+			AllPermutations = GetResolvedPermutationIndices(false);
+			PermutationIndicesPtr = &AllPermutations;
 		}
 
 		const FText DisplayNameText = Template->GetDisplayNameForArgument(InRootPinName, *PermutationIndicesPtr);
@@ -766,77 +688,6 @@ FName URigVMTemplateNode::GetDisplayNameForPin(const FName& InRootPinName,
 	return NAME_None;
 }
 
-const TArray<int32>& URigVMTemplateNode::GetFilteredPermutationsIndices() const
-{
-	return FilteredPermutations;
-}
-
-TArray<TRigVMTypeIndex> URigVMTemplateNode::GetFilteredTypesForPin(URigVMPin* InPin) const
-{
-	ensureMsgf(InPin->GetNode() == this, TEXT("GetFilteredTypesForPin of %s with pin from another node %s"), *GetNodePath(), *InPin->GetPinPath(true));
-	TArray<TRigVMTypeIndex> FilteredTypes;
-
-	if (IsSingleton())
-	{
-		return {InPin->GetTypeIndex()};
-	}
-
-	if (FilteredPermutations.IsEmpty())
-	{
-		return FilteredTypes;	
-	}
-
-	if (InPin->IsStructMember())
-	{
-		return {InPin->GetTypeIndex()};
-	}
-
-	if (!PreferredPermutationPairs.IsEmpty())
-	{
-		for (const FRigVMTemplatePreferredType& PreferredType : PreferredPermutationPairs)
-		{
-			if (InPin->GetFName() == PreferredType.GetArgument())
-			{
-				const FRigVMTemplateArgument* Argument = GetTemplate()->FindArgument(PreferredType.GetArgument());
-				for (const TRigVMTypeIndex& TypeIndex : Argument->GetTypeIndices())
-				{
-					if (TypeIndex == PreferredType.GetTypeIndex())
-					{
-						if (InPin->IsArrayElement())
-						{
-							return {FRigVMRegistry::Get().GetArrayTypeFromBaseTypeIndex(TypeIndex)};
-						}
-						else
-						{
-							return {TypeIndex};
-						}
-					}
-				}	
-			}
-		}
-	}
-	
-	const URigVMPin* RootPin = InPin;
-	bool bIsArrayElement = false;
-	if (URigVMPin* ParentPin = InPin->GetParentPin())
-	{
-		RootPin = ParentPin;
-		bIsArrayElement = true;
-	}
-	if (const FRigVMTemplateArgument* Argument = GetTemplate()->FindArgument(RootPin->GetFName()))
-	{
-		FilteredTypes = Argument->GetSupportedTypeIndices(FilteredPermutations);
-		if (bIsArrayElement)
-		{
-			for (TRigVMTypeIndex& ArrayType : FilteredTypes)
-			{
-				ArrayType = FRigVMRegistry::Get().GetBaseTypeFromArrayTypeIndex(ArrayType);
-			}
-		}
-	}
-	return FilteredTypes;
-}
-
 TRigVMTypeIndex URigVMTemplateNode::TryReduceTypesToSingle(const TArray<TRigVMTypeIndex>& InTypes, const TRigVMTypeIndex PreferredType) const
 {
 	if (InTypes.IsEmpty())
@@ -865,127 +716,6 @@ TRigVMTypeIndex URigVMTemplateNode::TryReduceTypesToSingle(const TArray<TRigVMTy
 	return InTypes[0];
 }
 
-TArray<int32> URigVMTemplateNode::GetNewFilteredPermutations(URigVMPin* InPin, URigVMPin* LinkedPin)
-{	
-	TArray<int32> NewFilteredPermutations;
-	if (InPin->GetNode() != this)
-	{
-		return NewFilteredPermutations;
-	}
-
-	NewFilteredPermutations.Reserve(FilteredPermutations.Num());
-	
-	bool bIsArrayElement = false;
-	bool bIsStructElement = false;
-	URigVMPin* RootPin = InPin;
-	if (URigVMPin* ParentPin = InPin->GetParentPin())
-	{
-		RootPin = ParentPin;
-		bIsArrayElement = RootPin->IsArray();
-		bIsStructElement = RootPin->IsStruct();
-	}
-
-	if (bIsStructElement)
-	{
-		if (FRigVMRegistry::Get().CanMatchTypes(InPin->GetTypeIndex(), LinkedPin->GetTypeIndex(), true))
-		{
-			return FilteredPermutations;
-		}
-	}
-
-	TArray<int32> PermutationsToTry = FilteredPermutations;
-
-	// Reduce permutations to the ones respecting the preferred types
-	const TArray<int32> PreferredPermutations = FindPermutationsForTypes(PreferredPermutationPairs, true);
-	PermutationsToTry = PermutationsToTry.FilterByPredicate([&](const int32& OtherPermutation) { return PreferredPermutations.Contains(OtherPermutation); });
-	
-	bool bLinkedIsTemplate = false;
-	if (URigVMTemplateNode* LinkedTemplate = Cast<URigVMTemplateNode>(LinkedPin->GetNode()))
-	{
-		if (!LinkedTemplate->IsSingleton() && !LinkedPin->IsStructMember())
-		{
-			bLinkedIsTemplate = true;
-			if (const FRigVMTemplateArgument* Argument = GetTemplate()->FindArgument(RootPin->GetFName()))
-			{
-				for(int32 PermutationIndex : PermutationsToTry)
-				{
-					TRigVMTypeIndex TypeIndex = Argument->GetTypeIndices()[PermutationIndex];
-					if (bIsArrayElement)
-					{
-						TypeIndex = FRigVMRegistry::Get().GetBaseTypeFromArrayTypeIndex(TypeIndex);
-					}
-					if (LinkedTemplate->FilteredSupportsType(LinkedPin, TypeIndex))
-					{
-						NewFilteredPermutations.Add(PermutationIndex);
-					}
-				}
-			}
-		}
-	}
-	
-	if (!bLinkedIsTemplate)
-	{
-		if (const FRigVMTemplateArgument* Argument = GetTemplate()->FindArgument(RootPin->GetFName()))
-		{
-			const TRigVMTypeIndex LinkedTypeIndex = LinkedPin->GetTypeIndex();
-			for(int32 PermutationIndex : PermutationsToTry)
-			{
-				TRigVMTypeIndex TypeIndex = Argument->GetTypeIndices()[PermutationIndex];
-				if (bIsArrayElement)
-				{
-					TypeIndex = FRigVMRegistry::Get().GetBaseTypeFromArrayTypeIndex(TypeIndex);
-				}
-
-				if (FRigVMRegistry::Get().CanMatchTypes(TypeIndex, LinkedTypeIndex, true))
-				{
-					NewFilteredPermutations.Add(PermutationIndex);
-				}
-			}
-		}
-	}
-	return NewFilteredPermutations;
-}
-
-TArray<int32> URigVMTemplateNode::GetNewFilteredPermutations(URigVMPin* InPin, const TArray<TRigVMTypeIndex>& InTypeIndices)
-{
-	TArray<int32> NewFilteredPermutations;
-	NewFilteredPermutations.Reserve(FilteredPermutations.Num());
-
-	URigVMPin* RootPin = InPin;
-	TArray<TRigVMTypeIndex> RootTypes = InTypeIndices;
-	if (URigVMPin* ParentPin = InPin->GetParentPin())
-	{
-		RootPin = ParentPin;
-		for (TRigVMTypeIndex& TypeIndex : RootTypes)
-		{
-			TypeIndex = FRigVMRegistry::Get().GetArrayTypeFromBaseTypeIndex(TypeIndex);
-		}
-	}
-
-	TArray<int32> PermutationsToTry = FilteredPermutations;
-
-	// Reduce permutations to the ones respecting the preferred types
-	const TArray<int32> PreferredPermutations = FindPermutationsForTypes(PreferredPermutationPairs, true);
-	PermutationsToTry = PermutationsToTry.FilterByPredicate([&](const int32& OtherPermutation) { return PreferredPermutations.Contains(OtherPermutation); });
-	
-	if (const FRigVMTemplateArgument* Argument = GetTemplate()->FindArgument(RootPin->GetFName()))
-	{
-		const TArray<TRigVMTypeIndex>& TypeIndices = Argument->GetTypeIndices();
-		for(int32 PermutationIndex : PermutationsToTry)
-		{
-			for (const TRigVMTypeIndex& RootTypeIndex : RootTypes)
-			{				
-				if (FRigVMRegistry::Get().CanMatchTypes(TypeIndices[PermutationIndex], RootTypeIndex, true))
-				{
-					NewFilteredPermutations.Add(PermutationIndex);
-					break;
-				}
-			}			
-		}
-	}
-	return NewFilteredPermutations;
-}
-
 TArray<int32> URigVMTemplateNode::FindPermutationsForTypes(const TArray<FRigVMTemplatePreferredType>& ArgumentTypes, bool bAllowCasting) const
 {
 	TArray<int32> Permutations;
@@ -1038,28 +768,6 @@ TArray<int32> URigVMTemplateNode::FindPermutationsForTypes(const TArray<FRigVMTe
 	return Permutations;
 }
 
-TArray<FRigVMTemplatePreferredType> URigVMTemplateNode::GetPreferredTypesForPermutation(const int32 InPermutationIndex) const
-{
-	TArray<FRigVMTemplatePreferredType> ArgTypes;
-	if (const FRigVMTemplate* Template = GetTemplate())
-	{
-		for (int32 ArgIndex = 0; ArgIndex < Template->NumArguments(); ++ArgIndex)
-		{
-			const FRigVMTemplateArgument* Argument = Template->GetArgument(ArgIndex);
-			if (Argument->GetTypeIndices().Num() > InPermutationIndex)
-			{
-				ArgTypes.Emplace(Argument->GetName(), Argument->GetTypeIndices()[InPermutationIndex]);
-			}
-			else
-			{
-				return {};
-			}
-		}
-	}
-
-	return ArgTypes;
-}
-
 FRigVMTemplateTypeMap URigVMTemplateNode::GetTypesForPermutation(const int32 InPermutationIndex) const
 {
 	if (const FRigVMTemplate* Template = GetTemplate())
@@ -1070,64 +778,6 @@ FRigVMTemplateTypeMap URigVMTemplateNode::GetTypesForPermutation(const int32 InP
 	return FRigVMTemplateTypeMap();
 }
 
-bool URigVMTemplateNode::PinNeedsFilteredTypesUpdate(URigVMPin* InPin, const TArray<TRigVMTypeIndex>& InTypeIndices)
-{
-	const TArray<int32> NewFilteredPermutations = GetNewFilteredPermutations(InPin, InTypeIndices);
-	if (NewFilteredPermutations.Num() == FilteredPermutations.Num())
-	{
-		return false;
-	}
-	return true;
-}
-
-bool URigVMTemplateNode::PinNeedsFilteredTypesUpdate(URigVMPin* InPin, URigVMPin* LinkedPin)
-{
-	const TArray<int32> NewFilteredPermutations = GetNewFilteredPermutations(InPin, LinkedPin);
-	if (NewFilteredPermutations.Num() == FilteredPermutations.Num())
-	{
-		return false;
-	}
-
-	return true;
-}
-
-bool URigVMTemplateNode::UpdateFilteredPermutations(URigVMPin* InPin, URigVMPin* LinkedPin)
-{
-	ensureMsgf(InPin->GetNode() == this, TEXT("Updating filtered permutations of %s with pin from another node %s"), *GetNodePath(), *InPin->GetPinPath(true));
-	ensureMsgf(LinkedPin->GetNode() != this, TEXT("Updating filtered permutations of %s with linked pin from same node %s"), *GetNodePath(), *LinkedPin->GetPinPath(true));
-
-	const TArray<int32> NewFilteredPermutations = GetNewFilteredPermutations(InPin, LinkedPin);	
-	if (NewFilteredPermutations.Num() == FilteredPermutations.Num())
-	{
-		return false;
-	}
-
-	if (NewFilteredPermutations.IsEmpty())
-	{
-		return false;
-	}
-	
-	FilteredPermutations = NewFilteredPermutations;
-	return true;
-}
-
-bool URigVMTemplateNode::UpdateFilteredPermutations(URigVMPin* InPin, const TArray<TRigVMTypeIndex>& InTypeIndices)
-{
-	const TArray<int32> NewFilteredPermutations = GetNewFilteredPermutations(InPin, InTypeIndices);
-	if (NewFilteredPermutations.Num() == FilteredPermutations.Num())
-	{
-		return false;
-	}
-
-	if (NewFilteredPermutations.IsEmpty())
-	{
-		return false;
-	}
-	
-	FilteredPermutations = NewFilteredPermutations;
-	return true;
-}
-
 void URigVMTemplateNode::InvalidateCache()
 {
 	Super::InvalidateCache();
@@ -1140,74 +790,3 @@ void URigVMTemplateNode::InvalidateCache()
 		ResolvedFunctionName.Reset();
 	}
 }
-
-void URigVMTemplateNode::InitializeFilteredPermutations()
-{
-	if (const FRigVMTemplate* Template = GetTemplate())
-	{
-		if (!PreferredPermutationPairs.IsEmpty())
-		{
-			FilteredPermutations = FindPermutationsForTypes(PreferredPermutationPairs, true);
-		}
-		else
-		{
-			FilteredPermutations.SetNumUninitialized(Template->NumPermutations());
-			for (int32 i=0; i<FilteredPermutations.Num(); ++i)
-			{
-				FilteredPermutations[i] = i;
-			}
-		}
-	}	
-}
-
-void URigVMTemplateNode::InitializeFilteredPermutationsFromTypes(bool bAllowCasting)
-{
-	if (IsSingleton())
-	{
-		return;
-	}
-
-	if (IsA<URigVMFunctionEntryNode>() || IsA<URigVMFunctionReturnNode>())
-	{
-		TArray<FRigVMTemplatePreferredType> ArgTypes;
-		for (URigVMPin* Pin : GetPins())
-		{
-			if (!Pin->IsWildCard())
-			{
-				ArgTypes.Emplace(Pin->GetFName(), Pin->GetTypeIndex());
-			}
-		}
-		
-		PreferredPermutationPairs = ArgTypes;
-	}
-	else
-	{
-		if (const FRigVMTemplate* Template = GetTemplate())
-		{
-			TArray<FRigVMTemplatePreferredType> ArgTypes;
-			for (int32 ArgIndex = 0; ArgIndex < Template->NumArguments(); ++ArgIndex)
-			{
-				const FRigVMTemplateArgument* Argument = Template->GetArgument(ArgIndex);
-				if (URigVMPin* Pin = FindPin(Argument->GetName().ToString()))
-				{
-					if (!Pin->IsWildCard())
-					{
-						ArgTypes.Emplace(Argument->GetName(), Pin->GetTypeIndex());
-					}
-				}
-			}
-
-			const TArray<int32> Permutations = FindPermutationsForTypes(ArgTypes, bAllowCasting);
-			if (!Permutations.IsEmpty())
-			{
-				FilteredPermutations = Permutations;
-				PreferredPermutationPairs = ArgTypes;
-			}
-			else
-			{
-				InitializeFilteredPermutations();
-			}
-		}
-	}
-}
-
diff --git a/Engine/Source/Developer/RigVMDeveloper/Private/RigVMModel/RigVMClient.cpp b/Engine/Source/Developer/RigVMDeveloper/Private/RigVMModel/RigVMClient.cpp
--- a/Engine/Source/Developer/RigVMDeveloper/Private/RigVMModel/RigVMClient.cpp
+++ b/Engine/Source/Developer/RigVMDeveloper/Private/RigVMModel/RigVMClient.cpp
@@ -678,6 +678,7 @@ FRigVMClientPatchResult FRigVMClient::PatchModelsOnLoad()
 		Result.Merge(Controller->PatchBranchNodesOnLoad());
 		Result.Merge(Controller->PatchIfSelectNodesOnLoad());
 		Result.Merge(Controller->PatchArrayNodesOnLoad());
+		Result.Merge(Controller->PatchInvalidLinksOnWildcards());
 	}
 
 	return Result;
diff --git a/Engine/Source/Developer/RigVMDeveloper/Private/RigVMModel/RigVMController.cpp b/Engine/Source/Developer/RigVMDeveloper/Private/RigVMModel/RigVMController.cpp
--- a/Engine/Source/Developer/RigVMDeveloper/Private/RigVMModel/RigVMController.cpp
+++ b/Engine/Source/Developer/RigVMDeveloper/Private/RigVMModel/RigVMController.cpp
@@ -82,7 +82,6 @@ void FRigVMClientPatchResult::Merge(const FRigVMClientPatchResult& InOther)
 
 URigVMController::URigVMController()
 	: bValidatePinDefaults(true)
-	, bSuspendRecomputingOuterTemplateFilters(false)
 	, bSuspendNotifications(false)
 	, bReportWarningsAndErrors(true)
 	, bIgnoreRerouteCompactnessChanges(false)
@@ -91,7 +90,6 @@ URigVMController::URigVMController()
 	, bIsTransacting(false)
 	, bIsRunningUnitTest(false)
 	, bIsFullyResolvingTemplateNode(false)
-	, bSuspendRecomputingTemplateFilters(false)
 	, bSuspendTemplateComputation(false)
 #if WITH_EDITOR
 	, bRegisterTemplateNodeUsage(true)
@@ -102,7 +100,6 @@ URigVMController::URigVMController()
 URigVMController::URigVMController(const FObjectInitializer& ObjectInitializer)
 	: Super(ObjectInitializer)
 	, bValidatePinDefaults(true)
-	, bSuspendRecomputingOuterTemplateFilters(false)
 	, bSuspendNotifications(false)
 	, bReportWarningsAndErrors(true)
 	, bIgnoreRerouteCompactnessChanges(false)
@@ -111,7 +108,6 @@ URigVMController::URigVMController(const FObjectInitializer& ObjectInitializer)
 	, bIsTransacting(false)
 	, bIsRunningUnitTest(false)
 	, bIsFullyResolvingTemplateNode(false)
-	, bSuspendRecomputingTemplateFilters(false)
 	, bSuspendTemplateComputation(false)
 #if WITH_EDITOR
 	, bRegisterTemplateNodeUsage(true)
@@ -1253,18 +1249,22 @@ URigVMUnitNode* URigVMController::AddUnitNode(UScriptStruct* InScriptStruct, con
 			return nullptr;
 		}
 
-		TArray<int32> OldPermutations = TemplateNode->FilteredPermutations;
 		int32 PermutationIndex = Template->FindPermutation(Function);
-		TemplateNode->FilteredPermutations = {PermutationIndex};
-		const TArray<FRigVMTemplatePreferredType> NewPreferredPermutationTypes = TemplateNode->GetPreferredTypesForPermutation(PermutationIndex);
-		if (bSetupUndoRedo)
+		FRigVMTemplateTypeMap Types = Template->GetTypesForPermutation(PermutationIndex);
+		for (TPair<FName, TRigVMTypeIndex>& Pair : Types)
 		{
-			FRigVMSetTemplateFilteredPermutationsAction Action(TemplateNode, OldPermutations);
-			ActionStack->AddAction(Action);
-			ActionStack->AddAction(FRigVMSetPreferredTemplatePermutationsAction(TemplateNode, NewPreferredPermutationTypes));
-		}		
-		TemplateNode->PreferredPermutationPairs = NewPreferredPermutationTypes;
-		UpdateTemplateNodePinTypes(TemplateNode, bSetupUndoRedo);
+			if (URigVMPin* Pin = TemplateNode->FindPin(Pair.Key.ToString()))
+			{
+				if (Pin->IsWildCard())
+				{
+					ResolveWildCardPin(Pin, Pair.Value, bSetupUndoRedo);
+				}
+			}
+			if (!TemplateNode->HasWildCardPin())
+			{
+				break;
+			}
+		}
 
 		if (UnitNodeCreatedContext.IsValid())
 		{
@@ -2211,7 +2211,7 @@ bool URigVMController::UnresolveTemplateNodes(const TArray<FName>& InNodeNames,
 	{
 		if (URigVMNode* Node = GetGraph()->FindNodeByName(NodeName))
 		{
-			Nodes.Add(Node);			
+			Nodes.Add(Node);
 		}
 	}
 
@@ -2240,14 +2240,13 @@ bool URigVMController::UnresolveTemplateNodes(const TArray<FName>& InNodeNames,
 	return false;
 }
 
-bool URigVMController::UnresolveTemplateNodes(const TArray<URigVMNode*>& InNodes, bool bSetupUndoRedo, bool bBreakLinks)
+bool URigVMController::UnresolveTemplateNodes(const TArray<URigVMNode*>& InNodes, bool bSetupUndoRedo)
 {
 	if (!IsValidGraph() || InNodes.IsEmpty())
 	{
 		return false;
 	}
 
-
 	if (!bIsTransacting && !IsGraphEditable())
 	{
 		return false;
@@ -2275,68 +2274,80 @@ bool URigVMController::UnresolveTemplateNodes(const TArray<URigVMNode*>& InNodes
 	}
 
 	{
-		TGuardValue<bool> GuardRecompute(bSuspendRecomputingTemplateFilters, true);
+		const FRigVMRegistry& Registry = FRigVMRegistry::Get();
 		for (URigVMNode* Node : InNodes)
 		{
 			EjectAllInjectedNodes(Node, bSetupUndoRedo);
-
-			if (URigVMTemplateNode* TemplateNode = Cast<URigVMTemplateNode>(Node))
+			TArray<URigVMLink*> Links = Node->GetLinks();
+			for (int32 i=0; i<Links.Num(); ++i)
 			{
-				if (!TemplateNode->PreferredPermutationPairs.IsEmpty())
+				URigVMLink* Link = Links[i];
+				URigVMPin* SourcePin = Link->GetSourcePin();
+				URigVMPin* TargetPin = Link->GetTargetPin();
+
+				URigVMNode* OtherNode = SourcePin->GetNode() == Node ? TargetPin->GetNode() : SourcePin->GetNode();
+				if (!InNodes.Contains(OtherNode))
 				{
-					if (bSetupUndoRedo)
+					const URigVMPin* PinOnNode = SourcePin->GetNode() == Node ? SourcePin : TargetPin;
+					if(PinOnNode->IsExecuteContext())
 					{
-						ActionStack->AddAction(FRigVMSetPreferredTemplatePermutationsAction(TemplateNode, {}));
+						continue;
 					}
-					TemplateNode->PreferredPermutationPairs = {};
-				}	
+					
+					if (const URigVMTemplateNode* TemplateNode = Cast<URigVMTemplateNode>(Node))
+					{
+						if(const FRigVMTemplate* Template = TemplateNode->GetTemplate())
+						{
+							const URigVMPin* RootPin = PinOnNode->GetRootPin();
+							if(const FRigVMTemplateArgument* Argument = Template->FindArgument(RootPin->GetFName()))
+							{
+								if(Argument->IsSingleton())
+								{
+									continue;
+								}
+							}
+						}
+					}
+					
+					BreakLink(SourcePin, TargetPin, bSetupUndoRedo);
+				}
 			}
 
-			if (bBreakLinks)
+			if (URigVMTemplateNode* TemplateNode = Cast<URigVMTemplateNode>(Node))
 			{
-				TArray<URigVMLink*> Links = Node->GetLinks();
-				for (int32 i=0; i<Links.Num(); ++i)
+				TemplateNode->InvalidateCache();
+				TemplateNode->ResolvedFunctionName.Reset();
+				TemplateNode->ResolvedPermutation = INDEX_NONE;
+
+				if (const FRigVMTemplate* Template = TemplateNode->GetTemplate())
 				{
-					URigVMLink* Link = Links[i];
-					URigVMPin* SourcePin = Link->GetSourcePin();
-					URigVMPin* TargetPin = Link->GetTargetPin();
-			
-					URigVMNode* OtherNode = SourcePin->GetNode() == Node ? TargetPin->GetNode() : SourcePin->GetNode();
-					if (!InNodes.Contains(OtherNode))
+					for (int32 i=0; i<Template->NumArguments(); ++i)
 					{
-						const URigVMPin* PinOnNode = SourcePin->GetNode() == Node ? SourcePin : TargetPin;
-						if(PinOnNode->IsExecuteContext())
-						{
-							continue;
-						}
-						
-						if (const URigVMTemplateNode* TemplateNode = Cast<URigVMTemplateNode>(Node))
+						const FRigVMTemplateArgument* Argument = Template->GetArgument(i);
+						if (!Argument->IsSingleton())
 						{
-							if(const FRigVMTemplate* Template = TemplateNode->GetTemplate())
+							if (URigVMPin* Pin = TemplateNode->FindPin(Argument->Name.ToString()))
 							{
-								const URigVMPin* RootPin = PinOnNode->GetRootPin();
-								if(const FRigVMTemplateArgument* Argument = Template->FindArgument(RootPin->GetFName()))
+								TRigVMTypeIndex OldTypeIndex = Pin->GetTypeIndex();
+								TRigVMTypeIndex NewTypeIndex = RigVMTypeUtils::TypeIndex::WildCard;
+								while (Registry.IsArrayType(OldTypeIndex))
 								{
-									if(Argument->IsSingleton())
-									{
-										continue;
-									}
+									OldTypeIndex = Registry.GetBaseTypeFromArrayTypeIndex(OldTypeIndex);
+									NewTypeIndex = Registry.GetArrayTypeFromBaseTypeIndex(NewTypeIndex);
 								}
+								ChangePinType(Pin, NewTypeIndex, bSetupUndoRedo, false, true, false);
 							}
 						}
-						
-						BreakLink(SourcePin, TargetPin, bSetupUndoRedo);
-					}			
+					}
+					UpdateTemplateNodePinTypes(TemplateNode, bSetupUndoRedo);
 				}
 			}
 		}
 	}
-	
-	RecomputeAllTemplateFilteredPermutations(bSetupUndoRedo);
 
 	if(bSetupUndoRedo)
 	{
-		ActionStack->EndAction(Action);		
+		ActionStack->EndAction(Action);
 	}
 
 	return true;
@@ -2892,7 +2903,6 @@ URigVMRerouteNode* URigVMController::AddRerouteNodeOnPin(const FString& InPinPat
 
 	Notify(ERigVMGraphNotifType::NodeAdded, Node);
 
-	Node->InitializeFilteredPermutations();
 	if (bAsInput)
 	{
 		AddLink(ValuePin, PinForLink, bSetupUndoRedo);
@@ -3673,13 +3683,6 @@ FString URigVMController::ExportNodesToText(const TArray<FName>& InNodeNames)
 					AllNodeNames.AddUnique(Injection->Node->GetFName());
 				}
 			}
-
-			if(URigVMTemplateNode* TemplateNode = Cast<URigVMTemplateNode>(Node))
-			{
-				// make sure the preferred types are saved
-				// as strings - not just as indices
-				TemplateNode->ConvertPreferredTypesToString();
-			}
 		}
 	}
 
@@ -3935,14 +3938,6 @@ TArray<FName> URigVMController::ImportNodesFromText(const FString& InText, bool
 
 		for(int32 SubNodeIndex=0; SubNodeIndex < SubNodes.Num(); SubNodeIndex++)
 		{
-			if(URigVMTemplateNode* TemplateNode = Cast<URigVMTemplateNode>(SubNodes[SubNodeIndex]))
-			{
-				// make sure the update preferred type pairs
-				// from string to index - since indices may
-				// have changed between sessions
-				TemplateNode->ConvertPreferredTypesToTypeIndex();
-			}
-		
 			if(URigVMCollapseNode* CollapseNode = Cast<URigVMCollapseNode>(SubNodes[SubNodeIndex]))
 			{
 				{
@@ -3958,7 +3953,6 @@ TArray<FName> URigVMController::ImportNodesFromText(const FString& InText, bool
 
 		if (URigVMLibraryNode* LibraryNode = Cast<URigVMLibraryNode>(CreatedNode))
 		{
-			TGuardValue<bool> GuardRecompute(bSuspendRecomputingOuterTemplateFilters, true);			
 			TFunction<void(URigVMLibraryNode*)> Recompute = [&](URigVMLibraryNode* OuterNode)
 			{
 				for (URigVMNode* Node : OuterNode->GetContainedNodes())
@@ -3974,16 +3968,6 @@ TArray<FName> URigVMController::ImportNodesFromText(const FString& InText, bool
 				{
 					return;
 				}
-				
-				{
-					FRigVMControllerGraphGuard GraphGuard(this, OuterNode->GetGraph(), false);
-					UpdateLibraryTemplate(OuterNode, false);
-				}
-				{
-					FRigVMControllerGraphGuard GraphGuard(this, OuterNode->GetContainedGraph(), false);
-					RecomputeAllTemplateFilteredPermutations(false);
-				}
-				
 			};
 			Recompute(LibraryNode);
 		}
@@ -4004,15 +3988,6 @@ TArray<FName> URigVMController::ImportNodesFromText(const FString& InText, bool
 				}
 			}
 
-			if (URigVMTemplateNode* TemplateNode = Cast<URigVMTemplateNode>(SubNode))
-			{
-				if (!TemplateNode->IsSingleton())
-				{
-					TemplateNode->InitializeFilteredPermutations();
-					UpdateTemplateNodePinTypes(TemplateNode, false, false);					
-				}
-			}
-
 			if (URigVMFunctionReferenceNode* FunctionRefNode = Cast<URigVMFunctionReferenceNode>(SubNode))
 			{
 				if(URigVMBuildData* BuildData = URigVMBuildData::Get())
@@ -4093,13 +4068,10 @@ TArray<FName> URigVMController::ImportNodesFromText(const FString& InText, bool
 									FString NewSourceRootPinPath = URigVMPin::JoinPinPath(NewSourceNodeName->ToString(), OldRootPin->GetName());
 									if (URigVMPin* NewRootPin = Graph->FindPin(NewSourceRootPinPath))
 									{
-										if (PrepareTemplatePinForType(NewRootPin, {OldRootPin->GetTypeIndex()}, bSetupUndoRedo))
-										{
-											SourcePin = CreatedLink->GetSourcePin();
-										}										
+										SourcePin = CreatedLink->GetSourcePin();
 									}
 								}
-							}							
+							}
 						}
 					}
 					if (TargetPin == nullptr)
@@ -4116,13 +4088,10 @@ TArray<FName> URigVMController::ImportNodesFromText(const FString& InText, bool
 									FString NewTargetRootPinPath = URigVMPin::JoinPinPath(NewTargetNodeName->ToString(), OldRootPin->GetName());
 									if (URigVMPin* NewRootPin = Graph->FindPin(NewTargetRootPinPath))
 									{
-										if (PrepareTemplatePinForType(NewRootPin, {OldRootPin->GetTypeIndex()}, bSetupUndoRedo))
-										{
-											TargetPin = CreatedLink->GetTargetPin();
-										}
+										TargetPin = CreatedLink->GetTargetPin();
 									}
 								}
-							}							
+							}
 						}
 					}
 					
@@ -4729,7 +4698,6 @@ URigVMCollapseNode* URigVMController::CollapseNodes(const TArray<URigVMNode*>& I
 		return AllLinks.Find(A) < AllLinks.Find(B);
 	});
 
-	TGuardValue<bool> RecomputeOuterTemplateFilter(bSuspendRecomputingOuterTemplateFilters, true);
 	FRigVMControllerCompileBracketScope CompileScope(this);
 	FRigVMCollapseNodesAction CollapseAction;
 
@@ -4880,22 +4848,6 @@ URigVMCollapseNode* URigVMController::CollapseNodes(const TArray<URigVMNode*>& I
 
 			Notify(ERigVMGraphNotifType::NodeAdded, ReturnNode);
 		}
-
-		for (int32 i=0; i<2; ++i)
-		{
-			URigVMTemplateNode* InterfaceNode = (i==0) ? (URigVMTemplateNode*) EntryNode : (URigVMTemplateNode*) ReturnNode;
-			if (InterfaceNode)
-			{
-				for (URigVMPin* Pin : InterfaceNode->GetPins())
-				{
-					if (!Pin->IsWildCard())
-					{
-						AddPreferredType(InterfaceNode, Pin->GetFName(), Pin->GetTypeIndex(), false);
-						AddArgumentForPin(Pin, nullptr, false);
-					}
-				}
-			}
-		}
 	}
 
 	// create the new nodes within the collapse node
@@ -5135,8 +5087,6 @@ URigVMCollapseNode* URigVMController::CollapseNodes(const TArray<URigVMNode*>& I
 		}
 	}
 
-	RecomputeAllTemplateFilteredPermutations(bSuspendNotifications);
-	
 	if (bSetupUndoRedo)
 	{
 		ActionStack->EndAction(CollapseAction);
@@ -5952,7 +5902,6 @@ URigVMFunctionReferenceNode* URigVMController::PromoteCollapseNodeToFunctionRefe
 
 	FRigVMControllerCompileBracketScope CompileScope(this);
 	URigVMFunctionReferenceNode* FunctionRefNode = nullptr;
-	TGuardValue<bool> GuardRecomputeOuter(bSuspendRecomputingOuterTemplateFilters, true);
 
 	// Create Function
 	URigVMLibraryNode* FunctionDefinition = nullptr;
@@ -6036,7 +5985,6 @@ URigVMFunctionReferenceNode* URigVMController::PromoteCollapseNodeToFunctionRefe
 	// Remove collapse node, add function reference, and add external links
 	if (FunctionDefinition)
 	{
-		TGuardValue<bool> GuardRecomputeInner(bSuspendRecomputingTemplateFilters, true);
  		FString NodeName = InCollapseNode->GetName();
 		FVector2D NodePosition = InCollapseNode->GetPosition();
 		TMap<FString, FPinState> PinStates = GetPinStates(InCollapseNode);
@@ -6168,19 +6116,6 @@ URigVMCollapseNode* URigVMController::PromoteFunctionReferenceNodeToCollapseNode
 	URigVMCollapseNode* CollapseNode = DuplicateObject<URigVMCollapseNode>(FunctionDefinition, Graph, *NodeName);
 	if(CollapseNode)
 	{
-		// Copy the library template
-		{
-			FRigVMByteArray ArchiveBytes;
-			FMemoryWriter ArchiveWriter(ArchiveBytes);
-			FRigVMTemplate Template = FunctionDefinition->Template;
-			Template.Serialize(ArchiveWriter);
-
-			FMemoryReader ArchiveReader(ArchiveBytes);
-			CollapseNode->Template.Serialize(ArchiveReader);
-
-			CollapseNode->InitializeFilteredPermutations();
-		}
-	
 		{
 			FRigVMControllerGraphGuard Guard(this, CollapseNode->GetContainedGraph(), false);
 			TGuardValue<bool> SuspendNotifications(bSuspendNotifications, true);
@@ -6634,13 +6569,6 @@ bool URigVMController::RenameNode(URigVMNode* InNode, const FName& InNewName, bo
 		ActionStack->BeginAction(Action);
 	}
 
-	if (URigVMCollapseNode* CollapseNode = Cast<URigVMCollapseNode>(InNode))
-	{
-		const FString& OldNotation = CollapseNode->Template.Notation.ToString();
-		const FString NewNotation = ValidNewName.ToString() + OldNotation.RightChop(OldNotation.Find(TEXT("(")));
-		CollapseNode->Template.Notation = *NewNotation;
-	}
-
 	// loop over all links and remove them
 	TArray<URigVMLink*> Links = InNode->GetLinks();
 	for (URigVMLink* Link : Links)
@@ -8201,7 +8129,7 @@ FString URigVMController::AddAggregatePin(URigVMNode* InNode, const FString& InP
 			{
 				ActionStack->CancelAction(Action, this);
 			}
-			return FString();	
+			return FString();
 		}
 
 		if (!Arg1 || !Arg2 || !ArgOpposite)
@@ -8237,6 +8165,41 @@ FString URigVMController::AddAggregatePin(URigVMNode* InNode, const FString& InP
 			return FString();
 		}
 
+		// We must resolve the type before we continue
+		if (Arg1->IsWildCard())
+		{
+			TRigVMTypeIndex AnswerType = INDEX_NONE;
+			if (RequestPinTypeSelectionDelegate.IsBound())
+			{
+				if (URigVMTemplateNode* TemplateNode = Cast<URigVMTemplateNode>(InNode))
+				{
+					if (const FRigVMTemplate* Template = TemplateNode->GetTemplate())
+					{
+						if (const FRigVMTemplateArgument* Argument = Template->FindArgument(Arg1->GetFName()))
+						{
+							const TArray<TRigVMTypeIndex>& Types = Argument->GetTypeIndices();
+							AnswerType = RequestPinTypeSelectionDelegate.Execute(Types);
+						}
+					}
+				}
+			}
+
+			if (AnswerType != INDEX_NONE)
+			{
+				ResolveWildCardPin(Arg1, AnswerType, bSetupUndoRedo);
+			}
+
+		}
+
+		if (Arg1->IsWildCard() || Arg2->IsWildCard() || ArgOpposite->IsWildCard())
+		{
+			if(bSetupUndoRedo)
+			{
+				ActionStack->CancelAction(Action, this);
+			}
+			return FString();
+		}
+
 		const FName PreviousNodeName = InNode->GetFName();
 		URigVMCollapseNode* CollapseNode = CollapseNodes({InNode}, InNode->GetName(), bSetupUndoRedo, true);
 		if (!CollapseNode)
@@ -8322,7 +8285,6 @@ FString URigVMController::AddAggregatePin(URigVMNode* InNode, const FString& InP
 	{
 		FRigVMControllerGraphGuard GraphGuard(this, AggregateNode->GetContainedGraph(), bSetupUndoRedo);
 		TGuardValue<bool> GuardEditGraph(GetGraph()->bEditable, true);
-		TGuardValue<bool> GuardOuterTemplateRecompute(bSuspendRecomputingOuterTemplateFilters, true);
 
 		URigVMNode* InnerNode = (AggregateNode == nullptr) ? InNode : AggregateNode->GetFirstInnerNode();
 
@@ -8435,9 +8397,6 @@ FString URigVMController::AddAggregatePin(URigVMNode* InNode, const FString& InP
 		return FString();
 	}
 
-	RecomputeAllTemplateFilteredPermutations(bSetupUndoRedo);
-	UpdateTemplateNodePinTypes(AggregateNode, bSetupUndoRedo);
-
 	if (bSetupUndoRedo)
 	{
 		ActionStack->EndAction(Action);
@@ -8540,7 +8499,6 @@ bool URigVMController::RemoveAggregatePin(URigVMPin* InPin, bool bSetupUndoRedo,
 							AggregateNode->InvalidateCache();
 						}
 
-						TGuardValue<bool> GuardOuterTemplateRecompute(bSuspendRecomputingOuterTemplateFilters, true);
 						const FString FirstAggregatePin = AggregateNode->GetFirstAggregatePin()->GetName();
 						const FString SecondAggregatePin = AggregateNode->GetSecondAggregatePin()->GetName();
 						FString OtherArg = TargetPin->GetName() == FirstAggregatePin ? SecondAggregatePin : FirstAggregatePin;
@@ -8548,7 +8506,7 @@ bool URigVMController::RemoveAggregatePin(URigVMPin* InPin, bool bSetupUndoRedo,
 						RewireLinks(NodeToRemove->FindPin(OtherArg), NextNodePin, true, bSetupUndoRedo);
 						RemoveNode(NodeToRemove, bSetupUndoRedo);
 						RemoveExposedPin(*InPin->GetName(), bSetupUndoRedo);
-						bSuccess = true;						
+						bSuccess = true;
 					}
 				}
 			}
@@ -8576,7 +8534,6 @@ bool URigVMController::RemoveAggregatePin(URigVMPin* InPin, bool bSetupUndoRedo,
 							AggregateNode->InvalidateCache();
 						}
 
-						TGuardValue<bool> GuardOuterTemplateRecompute(bSuspendRecomputingOuterTemplateFilters, true);
 						const FString FirstAggregatePin = AggregateNode->GetFirstAggregatePin()->GetName();
 						const FString SecondAggregatePin = AggregateNode->GetSecondAggregatePin()->GetName();
 						FString OtherArg = SourcePin->GetName() == FirstAggregatePin ? SecondAggregatePin : FirstAggregatePin;
@@ -8584,7 +8541,7 @@ bool URigVMController::RemoveAggregatePin(URigVMPin* InPin, bool bSetupUndoRedo,
 						RewireLinks(NodeToRemove->FindPin(OtherArg), NextNodePin, false, bSetupUndoRedo);
 						RemoveNode(NodeToRemove, bSetupUndoRedo);
 						RemoveExposedPin(*InPin->GetName(), bSetupUndoRedo);
-						bSuccess = true;	
+						bSuccess = true;
 					}
 				}
 			}			
@@ -9728,81 +9685,6 @@ bool URigVMController::AddLink(URigVMPin* OutputPin, URigVMPin* InputPin, bool b
 	{
 		if(!FRigVMRegistry::Get().CanMatchTypes(OutputPin->GetTypeIndex(), InputPin->GetTypeIndex(), true))
 		{
-			// returns true if this template node can resolve to the type without breaking any links
-			auto CanTemplateResolveToType = [](const TRigVMTypeIndex& InTypeIndex, URigVMPin* InPin) -> bool
-			{
-				if(const URigVMTemplateNode* TemplateNode = Cast<URigVMTemplateNode>(InPin->GetNode()))
-				{
-					if(const FRigVMTemplate* Template = TemplateNode->GetTemplate())
-					{
-						if(!InPin->IsRootPin())
-						{
-							return false;
-						}
-						
-						if(const FRigVMTemplateArgument* Argument = Template->FindArgument(InPin->GetFName()))
-						{
-							if(!Argument->SupportsTypeIndex(InTypeIndex))
-							{
-								return false;
-							}
-						}
-						
-						// we only rely on the cast if the number of links on non-singleton arguments is > 0
-						const TArray<URigVMLink*> Links = InPin->GetNode()->GetLinks();
-						int32 NumLinksRemaining = Links.Num();
-						for(const URigVMLink* Link : Links)
-						{
-							const URigVMPin* SourcePin = Link->GetSourcePin();
-							const URigVMPin* TargetPin = Link->GetTargetPin();
-
-							if(SourcePin && TargetPin)
-							{
-								if(SourcePin->IsExecuteContext())
-								{
-									NumLinksRemaining--;
-									continue;
-								}
-								
-								const URigVMPin* PinOnTemplateNode = SourcePin->GetNode() == TemplateNode ? SourcePin : TargetPin;
-								if(PinOnTemplateNode->IsRootPin() && PinOnTemplateNode != InPin)
-								{
-									const FName ArgumentName = PinOnTemplateNode->GetFName();
-									if(const FRigVMTemplateArgument* Argument = Template->FindArgument(ArgumentName))
-									{
-										if(Argument->IsSingleton())
-										{
-											NumLinksRemaining--;
-										}
-									}
-								}
-								else
-								{
-									NumLinksRemaining--;
-								}
-							}
-						}
-						
-						return NumLinksRemaining == 0;
-					}
-				}
-				
-				return false;
-			};
-
-			if(bCreateCastNode)
-			{
-				// todo: this should take into account links which have been detached
-				// during load we don't filter out the links and assume the template node can cast
-				// correctly to each type. we need to look at the links and their map rather then
-				// URigVMPin::GetLinks
-				if(CanTemplateResolveToType(InputPin->GetTypeIndex(), OutputPin) || 
-					CanTemplateResolveToType(OutputPin->GetTypeIndex(), InputPin))
-				{
-					bCreateCastNode = false;
-				}
-			}
-			
 			if(bCreateCastNode)
 			{
 				const FRigVMFunction* CastFunction = RigVMTypeUtils::GetCastForTypeIndices(OutputPin->GetTypeIndex(), InputPin->GetTypeIndex());
@@ -9905,41 +9787,6 @@ bool URigVMController::AddLink(URigVMPin* OutputPin, URigVMPin* InputPin, bool b
 		BreakAllLinks(OutputPin, false, bSetupUndoRedo);
 	}
 
-	bool bAddedTemporaryOutputPreferredType = false, bAddedTemporaryInputPreferredType = false;
-	for (int32 i=0; i<2; ++i)
-	{
-		URigVMPin* EdgePin = (i==0) ? OutputPin : InputPin;
-		bool& bAddedTemporaryPreferredType = (i==0) ? bAddedTemporaryOutputPreferredType : bAddedTemporaryInputPreferredType;
-		if (EdgePin->IsStructMember())
-		{
-			// We need to make sure the pin is not removed if any links are broken
-			// We temporarily add a preferred type to make sure that doesn't happen
-			if (URigVMTemplateNode* TemplateNode = Cast<URigVMTemplateNode>(EdgePin->GetNode()))
-			{
-				if (TemplateNode->IsSingleton())
-				{
-					continue;
-				}
-				
-				bAddedTemporaryPreferredType = true;
-				URigVMPin* RootPin = EdgePin->GetRootPin();
-				for (FRigVMTemplatePreferredType& Preferred : TemplateNode->PreferredPermutationPairs)
-				{
-					if (Preferred.GetArgument() == RootPin->GetFName())
-					{
-						bAddedTemporaryPreferredType = false;
-						break;
-					}
-				}
-
-				if (bAddedTemporaryPreferredType)
-				{
-					AddPreferredType(TemplateNode, RootPin->GetFName(), RootPin->GetTypeIndex(), bSetupUndoRedo);					
-				}
-			}
-		}
-	}
-
 	BreakAllLinks(InputPin, true, bSetupUndoRedo);
 	if (bSetupUndoRedo)
 	{
@@ -10020,54 +9867,6 @@ bool URigVMController::AddLink(URigVMPin* OutputPin, URigVMPin* InputPin, bool b
 	}
 	Notify(ERigVMGraphNotifType::LinkAdded, Link);
 
-	// Remove any temporary preferred type that was added
-	for (int32 i=0; i<2; ++i)
-	{
-		URigVMPin* EdgePin = (i==0) ? OutputPin : InputPin;
-		bool& bAddedTemporaryPreferredType = (i==0) ? bAddedTemporaryOutputPreferredType : bAddedTemporaryInputPreferredType;
-		if (bAddedTemporaryPreferredType)
-		{
-			if (URigVMTemplateNode* TemplateNode = Cast<URigVMTemplateNode>(EdgePin->GetNode()))
-			{
-				URigVMPin* RootPin = EdgePin->GetRootPin();
-				RemovePreferredType(TemplateNode, RootPin->GetFName(), bSetupUndoRedo);				
-			}
-		}
-	}
-
-	// Find any library pin that doesn't own an argument, but should
-	if (!bSuspendTemplateComputation)
-	{
-		if (URigVMLibraryNode* LibraryNode = Graph->GetTypedOuter<URigVMLibraryNode>())
-		{
-			for (URigVMPin* Pin : LibraryNode->GetPins())
-			{
-				if (const FRigVMTemplate* Template = LibraryNode->GetTemplate())
-				{
-					if (!Template->FindArgument(Pin->GetFName()))
-					{
-						if (ShouldPinOwnArgument(Pin))
-						{
-							URigVMPin* InterfacePin = nullptr;
-							if (URigVMFunctionEntryNode* EntryNode = LibraryNode->GetEntryNode())
-							{
-								InterfacePin = EntryNode->FindPin(Pin->GetName());
-							}
-							if (!InterfacePin)
-							{
-								if (URigVMFunctionReturnNode* ReturnNode = LibraryNode->GetReturnNode())
-								{
-									InterfacePin = ReturnNode->FindPin(Pin->GetName());
-								}
-							}
-							AddArgumentForPin(InterfacePin, nullptr, bSetupUndoRedo);
-						}
-					}
-				}
-			}
-		}
-	}
-
 	if (bSetupUndoRedo)
 	{
 		UpdateRerouteNodeAfterChangingLinks(OutputPin, bSetupUndoRedo);
@@ -10295,11 +10094,6 @@ bool URigVMController::BreakLink(URigVMPin* OutputPin, URigVMPin* InputPin, bool
 			RemoveUnusedOrphanedPins(OutputPin->GetNode());
 			RemoveUnusedOrphanedPins(InputPin->GetNode());
 			
-			if (!bIsTransacting && !bSuspendRecomputingTemplateFilters)
-			{
-				RecomputeAllTemplateFilteredPermutations(bSetupUndoRedo);
-			}
-			
 			if (!bSuspendNotifications)
 			{
 				Graph->MarkPackageDirty();
@@ -10390,7 +10184,6 @@ bool URigVMController::BreakAllLinks(URigVMPin* Pin, bool bAsInput, bool bSetupU
 
 	int32 LinksBroken = 0;
 	{
-		TGuardValue<bool> GuardSuspendTemplateRecompute(bSuspendRecomputingTemplateFilters, true);
 		if (Pin->IsBoundToVariable() && bAsInput && bSetupUndoRedo)
 		{
 			UnbindPinFromVariable(Pin, bSetupUndoRedo);
@@ -10412,11 +10205,6 @@ bool URigVMController::BreakAllLinks(URigVMPin* Pin, bool bAsInput, bool bSetupU
 		}
 	}
 
-	if (LinksBroken > 0 && !bSuspendRecomputingTemplateFilters)
-	{
-		RecomputeAllTemplateFilteredPermutations(bSetupUndoRedo);
-	}
-
 	if (bSetupUndoRedo)
 	{
 		if (LinksBroken > 0)
@@ -10436,7 +10224,6 @@ bool URigVMController::BreakAllLinksRecursive(URigVMPin* Pin, bool bAsInput, boo
 {
 	bool bBrokenLinks = false;
 	{
-		TGuardValue<bool> GuardSuspendRecomputeTemplates(bSuspendRecomputingTemplateFilters, true);
 		if (bTowardsParent)
 		{
 			URigVMPin* ParentPin = Pin->GetParentPin();
@@ -10456,11 +10243,6 @@ bool URigVMController::BreakAllLinksRecursive(URigVMPin* Pin, bool bAsInput, boo
 		}
 	}
 
-	if (bBrokenLinks && !bSuspendRecomputingTemplateFilters)
-	{
-		RecomputeAllTemplateFilteredPermutations(bSetupUndoRedo);
-	}
-	
 	return bBrokenLinks;
 }
 
@@ -10507,19 +10289,13 @@ FName URigVMController::AddExposedPin(const FName& InPinName, ERigVMPinDirection
 		}
 	}
 
-	// For now we only allow wildcards for collapsed nodes
+	// For now we do not allow wildcards for library nodes
 	if (CPPTypeObject)
 	{
 		if(CPPTypeObject == RigVMTypeUtils::GetWildCardCPPTypeObject())
 		{
-			if(const URigVMNode* CollapseNode = Cast<URigVMNode>(GetGraph()->GetOuter()))
-			{
-				if(CollapseNode->GetOuter()->IsA<URigVMFunctionLibrary>())
-				{
-					ReportError(TEXT("Cannot expose pins of wildcard type in functions."));
-					return NAME_None;
-				}
-			}
+			ReportError(TEXT("Cannot expose pins of wildcard type in functions."));
+			return NAME_None;
 		}
 	}
 	
@@ -10613,7 +10389,6 @@ FName URigVMController::AddExposedPin(const FName& InPinName, ERigVMPinDirection
 			TGuardValue<bool> SuspendNotifications(bSuspendNotifications, true);
 			RefreshFunctionPins(EntryNode);
 		}
-		EntryNode->InitializeFilteredPermutations();
 
 		Notify(ERigVMGraphNotifType::NodeAdded, EntryNode);
 	}
@@ -10627,7 +10402,6 @@ FName URigVMController::AddExposedPin(const FName& InPinName, ERigVMPinDirection
 			TGuardValue<bool> SuspendNotifications(bSuspendNotifications, true);
 			RefreshFunctionPins(ReturnNode);
 		}
-		ReturnNode->InitializeFilteredPermutations();
 
 		Notify(ERigVMGraphNotifType::NodeAdded, ReturnNode);
 	}
@@ -10635,39 +10409,6 @@ FName URigVMController::AddExposedPin(const FName& InPinName, ERigVMPinDirection
 	RefreshFunctionPins(EntryNode);
 	RefreshFunctionPins(ReturnNode);
 
-	// If the type is not a wildcard, add pin as template argument
-	if (InCPPType != RigVMTypeUtils::GetWildCardCPPType() && InCPPType != RigVMTypeUtils::GetWildCardArrayCPPType())
-	{
-		URigVMPin* InterfacePin = nullptr;
-		if (EntryNode)
-		{
-			InterfacePin = EntryNode->FindPin(Pin->GetName());
-		}
-		if (!InterfacePin)
-		{
-			if (ReturnNode)
-			{
-				InterfacePin = ReturnNode->FindPin(Pin->GetName());
-			}
-		}
-
-		// Add preferred type
-		{
-			URigVMTemplateNode* InterfaceNode = Cast<URigVMTemplateNode>(InterfacePin->GetNode());
-			TArray<FRigVMTemplatePreferredType> NewPreferredTypes = InterfaceNode->PreferredPermutationPairs;
-
-			NewPreferredTypes.Add(FRigVMTemplatePreferredType(*InterfacePin->GetName(), FRigVMRegistry::Get().GetTypeIndexFromCPPType(InCPPType)));
-
-			if (bSetupUndoRedo)
-			{
-				ActionStack->AddAction(FRigVMSetPreferredTemplatePermutationsAction(InterfaceNode, NewPreferredTypes));
-			}
-			InterfaceNode->PreferredPermutationPairs = NewPreferredTypes;
-		}	
-		
-		AddArgumentForPin(InterfacePin, nullptr, bSetupUndoRedo);
-	}
-	
 	RefreshFunctionReferences(LibraryNode, bSetupUndoRedo);
 
 	if (bSetupUndoRedo)
@@ -10695,7 +10436,7 @@ FName URigVMController::AddExposedPin(const FName& InPinName, ERigVMPinDirection
 	return PinName;
 }
 
-bool URigVMController::AddArgumentForPin(URigVMPin* InPin, URigVMPin* InToLinkPin, bool bSetupUndoRedo, bool bPrintPythonCommand)
+bool URigVMController::RemoveExposedPin(const FName& InPinName, bool bSetupUndoRedo, bool bPrintPythonCommand)
 {
 	if (!IsValidGraph())
 	{
@@ -10709,396 +10450,34 @@ bool URigVMController::AddArgumentForPin(URigVMPin* InPin, URigVMPin* InToLinkPi
 
 	URigVMGraph* Graph = GetGraph();
 	check(Graph);
-	
-	if (Graph != InPin->GetGraph())
-	{
-		return false;
-	}
 
-	URigVMLibraryNode* LibraryNode = Graph->GetTypedOuter<URigVMLibraryNode>();
-	if (!LibraryNode)
+	if (Graph->IsTopLevelGraph())
 	{
+		ReportError(TEXT("Exposed pins can only be edited on nested graphs."));
 		return false;
 	}
 
-	if (!InPin->GetNode()->IsA<URigVMFunctionEntryNode>() && !InPin->GetNode()->IsA<URigVMFunctionReturnNode>())
+	if (Graph->IsA<URigVMFunctionLibrary>())
 	{
+		ReportError(TEXT("Cannot remove exposed pins in function library graphs."));
 		return false;
 	}
 
-	FRigVMBaseAction Action;
-	if (bSetupUndoRedo)
+	URigVMLibraryNode* LibraryNode = Cast<URigVMLibraryNode>(Graph->GetOuter());
+	check(LibraryNode);
+
+	URigVMPin* Pin = LibraryNode->FindPin(InPinName.ToString());
+	if (Pin == nullptr)
 	{
-		Action.Title = TEXT("Add exposed pin from inner pin");
-		ActionStack->BeginAction(Action);
+		return false;
 	}
 
-	const FRigVMTemplate& CurTemplate = LibraryNode->Template;
-	FRigVMTemplate NewTemplate;
+	if(bSetupUndoRedo)
 	{
-		// Find all pins connected to the interface pin
-		TArray<URigVMPin*> ConnectedPins = InPin->GetLinkedSourcePins();
-		ConnectedPins.Append(InPin->GetLinkedTargetPins());
-		TArray<URigVMLink*> ConnectedLinks;
-
-		for (URigVMPin* ConnectedPin : ConnectedPins)
+		if(RequestBulkEditDialogDelegate.IsBound())
 		{
-			ConnectedLinks.Append(ConnectedPin->GetLinks().FilterByPredicate([&](URigVMLink* Link)
-			{
-				return Link->GetSourcePin() == InPin || Link->GetTargetPin() == InPin;
-			}));
-		}
-		for (URigVMNode* Node : Graph->GetNodes())
-		{
-			if (URigVMVariableNode* VariableNode = Cast<URigVMVariableNode>(Node))
-			{
-				if (VariableNode->GetVariableName() == *InPin->GetName())
-				{
-					URigVMPin* ValuePin = VariableNode->GetValuePin();
-					ConnectedPins.Append(ValuePin->GetLinkedSourcePins());
-					ConnectedPins.Append(ValuePin->GetLinkedTargetPins());
-
-					ConnectedLinks.Append(ValuePin->GetLinks());
-				}
-			}
-		}
-	
-		URigVMPin* PinToLink = InToLinkPin;
-		FRigVMRegistry& Registry = FRigVMRegistry::Get();
-
-		// Find the types shared by all the connected pins
-		TArray<TRigVMTypeIndex> TypeIndices;
-		{
-			// If no linked pin provided, find one of the connected pins that has a reduced number of filtered types
-			if (PinToLink == nullptr)
-			{
-				for (URigVMPin* Pin : ConnectedPins)
-				{
-					if (URigVMTemplateNode* TemplateNode = Cast<URigVMTemplateNode>(Pin->GetNode()))
-					{
-						if (TemplateNode->GetFilteredPermutationsIndices().Num() < 10)
-						{
-							PinToLink = Pin;
-							break;
-						}
-					}
-				}
-			}
-		
-			// Initialize types to the pin to link
-			if (PinToLink)
-			{
-				if (URigVMTemplateNode* ToLinkTemplateNode = Cast<URigVMTemplateNode>(PinToLink->GetNode()))
-				{
-					TypeIndices = ToLinkTemplateNode->GetFilteredTypesForPin(PinToLink);
-
-					// Avoid adding arguments with too many filtered types due to connections with unfiltered if, select, reroute or array nodes
-					if (TypeIndices.Num() > 10)
-					{
-						if (bSetupUndoRedo)
-						{
-							ActionStack->CancelAction(Action, this);
-						}
-						return false;
-					}
-				}
-				else
-				{
-					TypeIndices.Add(Registry.GetTypeIndexFromCPPType(PinToLink->GetCPPType()));
-				}
-			}
-			else
-			{
-				// If we couldn't find a pin to link (even through the connections) try to find a preferred type in interface nodes
-				TRigVMTypeIndex TypeIndex = INDEX_NONE;
-				for (int32 EntryReturn=0; EntryReturn<2 && TypeIndex == INDEX_NONE; ++EntryReturn)
-				{
-					URigVMTemplateNode* Node = (EntryReturn == 0) ? (URigVMTemplateNode*)Graph->GetEntryNode() : (URigVMTemplateNode*)Graph->GetReturnNode();
-					if (Node)
-					{
-						for (int32 i=0; i<Node->PreferredPermutationPairs.Num(); ++i)
-						{
-							FString ArgName;
-							if (Node->PreferredPermutationPairs[i].GetArgument() == InPin->GetFName())
-							{
-								TypeIndex = Node->PreferredPermutationPairs[i].GetTypeIndex();
-								break;
-							}
-						}
-					}
-				}
-
-				if (TypeIndex != INDEX_NONE)
-				{
-					TypeIndices.Add(TypeIndex);
-				}
-				else
-				{
-					if (bSetupUndoRedo)
-					{
-						ActionStack->CancelAction(Action, this);
-					}
-					return false;
-				}				
-			}
-
-			// Reduce types to the ones supported by all connections
-			for (int32 i=0; i<ConnectedPins.Num(); ++i)
-			{
-				URigVMPin* ConnectedPin = ConnectedPins[i];
-				if (ConnectedPin != PinToLink)
-				{
-					if (URigVMTemplateNode* ConnectedTemplateNode = Cast<URigVMTemplateNode>(ConnectedPin->GetNode()))
-					{
-						TypeIndices = TypeIndices.FilterByPredicate([&](int32 Type)
-						{
-							return ConnectedTemplateNode->FilteredSupportsType(ConnectedPin, Type);
-						});		
-					}
-					else
-					{
-						// This is a weird situation. If the pin was connected, but was not an argument it usually means
-						// it was connected to an if, select, array or reroute that has not been type-filtered
-						check(false);
-					}
-				}
-			}
-		}
-		if (PinToLink)
-		{
-			ConnectedPins.AddUnique(PinToLink);
-		}
-
-		ensure(!TypeIndices.IsEmpty());
-
-		// Temporarily disconnect connected links
-		if (CurTemplate.NumArguments() > 0)
-		{
-			TGuardValue<bool> SuspendNotifications(bSuspendNotifications, true);
-			DetachLinksFromPinObjects(&ConnectedLinks);
-		}
-	
-		// Resolve ConnectedPins to each possible filtered type and see how that affects other exposed pins	
-		TArray<TArray<TArray<TRigVMTypeIndex>>> TypeToPermutations; // #types_to_resolve x #resolved_permutations x #arguments
-		ensure(TypeIndices.Num() < 10);
-		for (TRigVMTypeIndex& TypeIndex : TypeIndices)
-		{
-			TArray<TArray<TRigVMTypeIndex>> PermutationTypes;
-			if (CurTemplate.Arguments.Num() > 0)
-			{
-				TGuardValue<bool> GuardNotifications(bSuspendNotifications, true);
-				TGuardValue<bool> GuardOuterRecomputing(bSuspendRecomputingOuterTemplateFilters, true);
-				OpenUndoBracket(TEXT("Resolving to type"));
-
-				for (URigVMPin* ConnectedPin : ConnectedPins)
-				{
-					if (ConnectedPin->IsWildCard())
-					{
-						ResolveWildCardPin(ConnectedPin, TypeIndex, true);
-					}					
-				}
-
-				for (int32 i=0; i<CurTemplate.NumPermutations(); ++i)
-				{
-					TArray<TRigVMTypeIndex> ArgTypes;					
-					for (const FRigVMTemplateArgument& Arg : CurTemplate.Arguments)
-					{
-						ArgTypes.Add(Arg.TypeIndices[i]);
-					}
-					PermutationTypes.Emplace(ArgTypes);
-				}
-
-				verify(CancelUndoBracket());
-			}
-			TypeToPermutations.Emplace(PermutationTypes);
-		}
-
-		// Reconnect  links
-		if (CurTemplate.NumArguments() > 0)
-		{
-			TGuardValue<bool> SuspendNotifications(bSuspendNotifications, true);
-			ReattachLinksToPinObjects(false, &ConnectedLinks);
-		}
-		
-		FRigVMTemplateArgument NewArgument;
-		NewArgument.Index = CurTemplate.Arguments.Num();
-		NewArgument.Name = InPin->GetFName();
-		NewArgument.Direction = LibraryNode->FindPin(InPin->GetName())->GetDirection();
-	
-		// Add the arguments
-		for (int32 i=0; i<CurTemplate.Arguments.Num(); ++i)
-		{
-			FRigVMTemplateArgument Arg;
-			Arg.Index = i;
-			Arg.Name = CurTemplate.Arguments[i].Name;
-			Arg.Direction = CurTemplate.Arguments[i].Direction;
-			NewTemplate.Arguments.Add(Arg);
-		}
-		NewTemplate.Arguments.Add(NewArgument);
-
-		// Types
-		for (int32 NewArgTypeIndex=0; NewArgTypeIndex<TypeIndices.Num(); NewArgTypeIndex++)
-		{
-			TRigVMTypeIndex& NewArgType = TypeIndices[NewArgTypeIndex];			
-
-			if (CurTemplate.NumArguments() > 0)
-			{
-				for (int32 PermutationIndex=0; PermutationIndex<TypeToPermutations[NewArgTypeIndex].Num(); ++PermutationIndex)
-				{
-					int32 GlobalPermutationIndex = NewTemplate.Permutations.Num();
-					NewTemplate.Permutations.Add(INDEX_NONE);
-					for (int32 ArgIndex=0; ArgIndex<NewTemplate.NumArguments()-1; ++ArgIndex)
-					{
-						TRigVMTypeIndex ArgType = TypeToPermutations[NewArgTypeIndex][PermutationIndex][ArgIndex];
-						NewTemplate.Arguments[ArgIndex].TypeIndices.Add(ArgType);
-						if (TArray<int32>* PermArray = NewTemplate.Arguments[ArgIndex].TypeToPermutations.Find(ArgType))
-						{
-							PermArray->Add(GlobalPermutationIndex);
-						}
-						else
-						{
-							NewTemplate.Arguments[ArgIndex].TypeToPermutations.Add(ArgType, {GlobalPermutationIndex});
-						}
-					}
-					NewTemplate.Arguments.Last().TypeIndices.Add(NewArgType);
-					if (TArray<int32>* PermArray = NewTemplate.Arguments.Last().TypeToPermutations.Find(NewArgType))
-					{
-						PermArray->Add(GlobalPermutationIndex);
-					}
-					else
-					{
-						NewTemplate.Arguments.Last().TypeToPermutations.Add(NewArgType, {GlobalPermutationIndex});
-					}				
-				}
-			}
-			else
-			{
-				int32 GlobalPermutationIndex = NewTemplate.Permutations.Num();
-				NewTemplate.Permutations.Add(INDEX_NONE);
-				NewTemplate.Arguments.Last().TypeIndices.Add(NewArgType);
-				if (TArray<int32>* PermArray = NewTemplate.Arguments.Last().TypeToPermutations.Find(NewArgType))
-				{
-					PermArray->Add(GlobalPermutationIndex);
-				}
-				else
-				{
-					NewTemplate.Arguments.Last().TypeToPermutations.Add(NewArgType, {GlobalPermutationIndex});
-				}
-			}
-		}
-	}
-	
-	NewTemplate.ComputeNotationFromArguments(LibraryNode->GetName());
-	
-	if (bSetupUndoRedo)
-	{
-		FRigVMControllerGraphGuard GraphGuard(this, LibraryNode->GetTypedOuter<URigVMGraph>(), bSetupUndoRedo);
-		ActionStack->AddAction(FRigVMSetLibraryTemplateAction(LibraryNode, NewTemplate));
-	}
-	
-	LibraryNode->Template = NewTemplate;
-	{
-		FRigVMControllerGraphGuard GraphGuard(this, LibraryNode->GetTypedOuter<URigVMGraph>(), false);
-		Notify(ERigVMGraphNotifType::LibraryTemplateChanged, LibraryNode);
-	}
-	
-	// Update Entry and Return filtered permutations to all permutations
-	URigVMTemplateNode* EntryNode = Graph->GetEntryNode();
-	URigVMTemplateNode* ReturnNode = Graph->GetReturnNode();	
-	for (int32 EntryReturn=0; EntryReturn<2; ++EntryReturn)
-	{
-		URigVMTemplateNode* Node = (EntryReturn == 0) ? EntryNode : ReturnNode;
-		if (Node)
-		{
-			TArray<int32> OldPermutations = Node->FilteredPermutations;
-
-			Node->FilteredPermutations.SetNumUninitialized(NewTemplate.NumPermutations());
-			for (int32 i=0; i<Node->FilteredPermutations.Num(); ++i)
-			{
-				Node->FilteredPermutations[i] = i;
-			}
-
-			if (bSetupUndoRedo)
-			{
-				ActionStack->AddAction(FRigVMSetTemplateFilteredPermutationsAction(Node, OldPermutations));
-			}
-		}
-	}
-	{
-		bool bUpdatedTypes = false;
-		for (int32 EntryReturn=0; EntryReturn<2; ++EntryReturn)
-		{
-			URigVMTemplateNode* Node = (EntryReturn == 0) ? EntryNode : ReturnNode;
-			if (Node)
-			{
-				// Update template node pin types of an interface node will update both entry and return types.
-				if (!bUpdatedTypes)
-				{
-					UpdateTemplateNodePinTypes(Node, bSetupUndoRedo);
-					bUpdatedTypes = true;
-				}
-				Notify(ERigVMGraphNotifType::NodeDescriptionChanged, Node);
-			}
-		}
-	}
-
-	if (!bIsTransacting && !bSuspendRecomputingOuterTemplateFilters)
-	{
-		FRigVMControllerGraphGuard GraphGuard(this, LibraryNode->GetTypedOuter<URigVMGraph>(), bSetupUndoRedo);
-		RecomputeAllTemplateFilteredPermutations(bSetupUndoRedo);
-		Notify(ERigVMGraphNotifType::NodeDescriptionChanged, LibraryNode);		
-	}
-
-	if (bSetupUndoRedo)
-	{
-		ActionStack->EndAction(Action);
-	}
-	
-	return true;
-}
-
-bool URigVMController::RemoveExposedPin(const FName& InPinName, bool bSetupUndoRedo, bool bPrintPythonCommand)
-{
-	if (!IsValidGraph())
-	{
-		return false;
-	}
-
-	if (!bIsTransacting && !IsGraphEditable())
-	{
-		return false;
-	}
-
-	URigVMGraph* Graph = GetGraph();
-	check(Graph);
-
-	if (Graph->IsTopLevelGraph())
-	{
-		ReportError(TEXT("Exposed pins can only be edited on nested graphs."));
-		return false;
-	}
-
-	if (Graph->IsA<URigVMFunctionLibrary>())
-	{
-		ReportError(TEXT("Cannot remove exposed pins in function library graphs."));
-		return false;
-	}
-
-	URigVMLibraryNode* LibraryNode = Cast<URigVMLibraryNode>(Graph->GetOuter());
-	check(LibraryNode);
-
-	URigVMPin* Pin = LibraryNode->FindPin(InPinName.ToString());
-	if (Pin == nullptr)
-	{
-		return false;
-	}
-
-	if(bSetupUndoRedo)
-	{
-		if(RequestBulkEditDialogDelegate.IsBound())
-		{
-			FRigVMController_BulkEditResult Result = RequestBulkEditDialogDelegate.Execute(LibraryNode, ERigVMControllerBulkEditType::RemoveExposedPin);
-			if(Result.bCanceled)
+			FRigVMController_BulkEditResult Result = RequestBulkEditDialogDelegate.Execute(LibraryNode, ERigVMControllerBulkEditType::RemoveExposedPin);
+			if(Result.bCanceled)
 			{
 				return false;
 			}
@@ -11115,8 +10494,6 @@ bool URigVMController::RemoveExposedPin(const FName& InPinName, bool bSetupUndoR
 
 	bool bSuccessfullyRemovedPin = false;
 	{
-		TGuardValue<bool> GuardRecomputeFilters(bSuspendRecomputingTemplateFilters, true);
-		TGuardValue<bool> GuardRecomputeOuterFilters(bSuspendRecomputingOuterTemplateFilters, true);
 		{
 			FRigVMControllerGraphGuard GraphGuard(this, LibraryNode->GetGraph(), bSetupUndoRedo);
 			bSuccessfullyRemovedPin = RemovePin(Pin, bSetupUndoRedo);
@@ -11143,34 +10520,6 @@ bool URigVMController::RemoveExposedPin(const FName& InPinName, bool bSetupUndoR
 		RefreshFunctionReferences(LibraryNode, false);
 	}
 
-	TArray<URigVMTemplateNode*> InterfaceNodes = {Graph->GetEntryNode(), Graph->GetReturnNode()};
-	// Remove preferred types in interface nodes
-	{
-		for (URigVMTemplateNode* InterfaceNode : InterfaceNodes)
-		{
-			RemovePreferredType(InterfaceNode, InPinName, bSetupUndoRedo);						
-		}
-	}
-
-	// Remove preferred type of library node (and references to it)
-	{
-		FRigVMControllerGraphGuard GraphGuard(this, LibraryNode->GetGraph(), bSetupUndoRedo);
-		RemovePreferredType(LibraryNode, InPinName, bSetupUndoRedo);
-
-		// If the library node is a function, remove preferred type from all references
-		if (URigVMFunctionLibrary* FunctionLibrary = Cast<URigVMFunctionLibrary>(LibraryNode->GetGraph()))
-		{
-			FunctionLibrary->ForEachReference(LibraryNode->GetFName(), [this, InPinName, bSetupUndoRedo](URigVMFunctionReferenceNode* ReferenceNode)
-			{
-				FRigVMControllerGraphGuard GraphGuard(this, ReferenceNode->GetGraph(), bSetupUndoRedo);
-				RemovePreferredType(ReferenceNode, InPinName, bSetupUndoRedo);
-			});
-		}
-	}
-
-	// Recompute filtered types and library template
-	RecomputeAllTemplateFilteredPermutations(bSetupUndoRedo);
-
 	if (bSetupUndoRedo)
 	{
 		if (bSuccessfullyRemovedPin)
@@ -11324,52 +10673,11 @@ bool URigVMController::RenameExposedPin(const FName& InOldPinName, const FName&
 		{
 			if (URigVMPin* InterfacePin = InterfaceNode->FindPin(InOldPinName.ToString()))
 			{
-				TArray<FRigVMTemplatePreferredType> NewPreferredPermutationsTypes = InterfaceNode->PreferredPermutationPairs;
-				for (FRigVMTemplatePreferredType& PreferredType : NewPreferredPermutationsTypes)
-				{
-					if (PreferredType.GetArgument() == InOldPinName)
-					{
-						PreferredType.Argument = PinName;
-						break;
-					}
-				}
-
-				if (bSetupUndoRedo)
-				{
-					ActionStack->AddAction(FRigVMSetPreferredTemplatePermutationsAction(InterfaceNode, NewPreferredPermutationsTypes));
-				}
-
-				InterfaceNode->PreferredPermutationPairs = NewPreferredPermutationsTypes;
-				
 				Local::RenamePin(this, InterfacePin, PinName);
 			}
 		}
 	}
 
-	// Find Argument and rename if it exists
-	{
-		FRigVMTemplate NewTemplate = LibraryNode->Template;
-		for (FRigVMTemplateArgument& Argument : NewTemplate.Arguments)
-		{
-			if (Argument.Name == InOldPinName)
-			{
-				Argument.Name = PinName;
-				break;
-			}
-		}
-		NewTemplate.ComputeNotationFromArguments(LibraryNode->GetName());
-		if (bSetupUndoRedo)
-		{
-			FRigVMControllerGraphGuard Guard(this, LibraryNode->GetGraph(), bSetupUndoRedo);
-			ActionStack->AddAction(FRigVMSetLibraryTemplateAction(LibraryNode, NewTemplate));
-		}
-		LibraryNode->Template = NewTemplate;
-		{
-			FRigVMControllerGraphGuard GraphGuard(this, LibraryNode->GetTypedOuter<URigVMGraph>(), false);
-			Notify(ERigVMGraphNotifType::LibraryTemplateChanged, LibraryNode);
-		}
-	}
-
 	if (URigVMFunctionLibrary* FunctionLibrary = Cast<URigVMFunctionLibrary>(LibraryNode->GetGraph()))
 	{
 		FunctionLibrary->ForEachReference(LibraryNode->GetFName(), [this, InOldPinName, PinName](URigVMFunctionReferenceNode* ReferenceNode)
@@ -11526,86 +10834,47 @@ bool URigVMController::ChangeExposedPinType(const FName& InPinName, const FStrin
 
 	// If the pin does not support the type, first break all links in the contained graph to this pin (and subpins)
 	TArray<URigVMTemplateNode*> InterfaceNodes = {Graph->GetEntryNode(), Graph->GetReturnNode()};
-	if (!LibraryNode->SupportsType(Pin, Registry.GetTypeIndexFromCPPType(InCPPType)))
+
+	// Break all links to this pin
 	{
+		TArray<URigVMLink*> InterfacePinLinks;
+		TArray<URigVMPin*> ExtendedInterfacePins;
+		for (URigVMNode* Node : Graph->GetNodes())
 		{
-			TGuardValue<bool> GuardRecompute (bSuspendRecomputingTemplateFilters, true);
-			TGuardValue<bool> GuardRecomputeOuter (bSuspendRecomputingOuterTemplateFilters, true);
-
-			// Remove preferred types for this pin
-			for (URigVMTemplateNode* InterfaceNode : InterfaceNodes)
+			if (Node->IsA<URigVMFunctionEntryNode>() || Node->IsA<URigVMFunctionReturnNode>())
 			{
-				if (InterfaceNode)
+				if (URigVMPin* InterfacePin = Node->FindPin(Pin->GetName()))
 				{
-					if (URigVMPin* InterfacePin = InterfaceNode->FindPin(Pin->GetName()))
-					{
-						RemovePreferredType(InterfaceNode, InPinName, bSetupUndoRedo);
-					}										
+					ExtendedInterfacePins.Add(InterfacePin);
 				}
 			}
-		
-			// Break all links to this pin
+			else if (URigVMVariableNode* VariableNode = Cast<URigVMVariableNode>(Node))
 			{
-				TArray<URigVMLink*> InterfacePinLinks;	
-				TArray<URigVMPin*> ExtendedInterfacePins;
-				for (URigVMNode* Node : Graph->GetNodes())
-				{
-					if (Node->IsA<URigVMFunctionEntryNode>() || Node->IsA<URigVMFunctionReturnNode>())
-					{
-						if (URigVMPin* InterfacePin = Node->FindPin(Pin->GetName()))
-						{
-							ExtendedInterfacePins.Add(InterfacePin);
-						}
-					}
-					else if (URigVMVariableNode* VariableNode = Cast<URigVMVariableNode>(Node))
-					{
-						if (VariableNode->GetVariableName() == InPinName)
-						{
-							ExtendedInterfacePins.Add(VariableNode->GetValuePin());
-						}
-					}
-				}
-		
-				for (URigVMPin* InterfacePin : ExtendedInterfacePins)
-				{
-					TArray<URigVMPin*> PinsToProcess;
-					PinsToProcess.Add(InterfacePin);
-					for (int32 i=0; i<PinsToProcess.Num(); ++i)
-					{
-						InterfacePinLinks.Append(PinsToProcess[i]->GetLinks());
-						PinsToProcess.Append(PinsToProcess[i]->GetSubPins());
-					}
-				}		
-				for (int32 i=0; i<InterfacePinLinks.Num(); ++i)
+				if (VariableNode->GetVariableName() == InPinName)
 				{
-					BreakLink(InterfacePinLinks[i]->GetSourcePin(), InterfacePinLinks[i]->GetTargetPin(), bSetupUndoRedo);
+					ExtendedInterfacePins.Add(VariableNode->GetValuePin());
 				}
 			}
 		}
 
-		RecomputeAllTemplateFilteredPermutations(bSetupUndoRedo);
-	}
-
-	// Resolve interface pins to the input type
-	for (URigVMTemplateNode* InterfaceNode : InterfaceNodes)
-	{
-		if (InterfaceNode)
+		for (URigVMPin* InterfacePin : ExtendedInterfacePins)
 		{
-			if (URigVMPin* InterfacePin = InterfaceNode->FindPin(Pin->GetName()))
+			TArray<URigVMPin*> PinsToProcess;
+			PinsToProcess.Add(InterfacePin);
+			for (int32 i=0; i<PinsToProcess.Num(); ++i)
 			{
-				if (!ResolveWildCardPin(InterfacePin, FRigVMTemplateArgumentType(*InCPPType, CPPTypeObject), bSetupUndoRedo))
-				{
-					if (bSetupUndoRedo)
-					{
-						ActionStack->CancelAction(Action, this);
-					}
-					return false;
-				}
+				InterfacePinLinks.Append(PinsToProcess[i]->GetLinks());
+				PinsToProcess.Append(PinsToProcess[i]->GetSubPins());
 			}
 		}
-	}
-	
 		
+		for (int32 i=0; i<InterfacePinLinks.Num(); ++i)
+		{
+			BreakLink(InterfacePinLinks[i]->GetSourcePin(), InterfacePinLinks[i]->GetTargetPin(), bSetupUndoRedo);
+		}
+	}
+
+	// Change pin type of the library node in the function library
 	{
 		bool bSuccessChangingType;
 		{
@@ -11627,6 +10896,7 @@ bool URigVMController::ChangeExposedPinType(const FName& InPinName, const FStrin
 		}
 	}
 
+	// Repopulate pin on interface nodes
 	for (URigVMTemplateNode* InterfaceNode : InterfaceNodes)
 	{
 		if (InterfaceNode)
@@ -11636,6 +10906,7 @@ bool URigVMController::ChangeExposedPinType(const FName& InPinName, const FStrin
 		}
 	}
 
+	// Change pin type on function references
 	if (URigVMFunctionLibrary* FunctionLibrary = Cast<URigVMFunctionLibrary>(LibraryNode->GetGraph()))
 	{
 		FunctionLibrary->ForEachReference(LibraryNode->GetFName(), [this, &Pin, InCPPType, InCPPTypeObjectPath, bSetupUndoRedo, bSetupOrphanPins](URigVMFunctionReferenceNode* ReferenceNode)
@@ -11649,6 +10920,7 @@ bool URigVMController::ChangeExposedPinType(const FName& InPinName, const FStrin
         });
 	}
 
+	// Change pin types on input variable nodes
 	for (URigVMNode* Node : Graph->GetNodes())
 	{
 		if (URigVMVariableNode* VariableNode = Cast<URigVMVariableNode>(Node))
@@ -11818,12 +11090,10 @@ URigVMFunctionReferenceNode* URigVMController::AddFunctionReferenceNodeFromDescr
 	FunctionRefNode->Position = InNodePosition;
 	FunctionRefNode->ReferencedFunctionHeader = InFunctionDefinition;
 	Graph->Nodes.Add(FunctionRefNode);
-	FunctionRefNode->InitializeFilteredPermutations();
 
 	FRigVMControllerCompileBracketScope CompileScope(this);
 	
 	RepopulatePinsOnNode(FunctionRefNode, false, false, false);
-	UpdateTemplateNodePinTypes(FunctionRefNode, false);
 
 	Notify(ERigVMGraphNotifType::NodeAdded, FunctionRefNode);
 
@@ -13234,22 +12504,6 @@ URigVMRerouteNode* URigVMController::AddFreeRerouteNode(bool bShowAsFullNode, co
 
 	Notify(ERigVMGraphNotifType::NodeAdded, Node);
 
-	Node->InitializeFilteredPermutations();
-	if (InCPPType != RigVMTypeUtils::GetWildCardCPPType() && InCPPType != RigVMTypeUtils::GetWildCardArrayCPPType())
-	{
-		PrepareTemplatePinForType(ValuePin, {ValuePin->GetTypeIndex()}, bSetupUndoRedo);
-		TArray<int32> FilterPermutations = Node->GetFilteredPermutationsIndices();
-		if (FilterPermutations.Num() == 1)
-		{
-			const TArray<FRigVMTemplatePreferredType> NewPreferredPermutationTypes = Node->GetPreferredTypesForPermutation(FilterPermutations[0]);
-			if (bSetupUndoRedo)
-			{
-				ActionStack->AddAction(FRigVMSetPreferredTemplatePermutationsAction(Node, NewPreferredPermutationTypes));
-			}
-			Node->PreferredPermutationPairs = NewPreferredPermutationTypes;
-		}
-	}
-
 	if (bSetupUndoRedo)
 	{
 		ActionStack->EndAction(Action);
@@ -13464,7 +12718,6 @@ URigVMTemplateNode* URigVMController::AddTemplateNode(const FName& InNotation, c
 	int32 PermutationIndex = INDEX_NONE;
 	FRigVMTemplate::FTypeMap Types;
 	Template->FullyResolve(Types, PermutationIndex);
-	Node->InitializeFilteredPermutations();
 
 	FRigVMRegistry& Registry = FRigVMRegistry::Get();
 	AddPinsForTemplate(Template, Types, Node);
@@ -13735,7 +12988,7 @@ URigVMNode* URigVMController::AddArrayNode(ERigVMOpCode InOpCode, const FString&
 				TypeMap.Add(ArgumentNameToResolve, TypeIndex);
 
 				TArray<int32> Permutations;
-				Template->Resolve(TypeMap, Permutations, true);
+				Template->Resolve(TypeMap, Permutations, false);
 				check(Permutations.Num() == 1);
 
 				for(const FRigVMTemplate::FTypePair& Pair : TypeMap)
@@ -13745,7 +12998,6 @@ URigVMNode* URigVMController::AddArrayNode(ERigVMOpCode InOpCode, const FString&
 						if(URigVMPin* Pin = Node->FindPin(Pair.Key.ToString()))
 						{
 							ChangePinType(Pin, Pair.Value, false, false);
-							Node->UpdateFilteredPermutations(Pin, {Pair.Value});
 						}
 					}
 				}
@@ -16708,8 +15960,7 @@ void URigVMController::ResolveTemplateNodeMetaData(URigVMTemplateNode* InNode, b
 #if WITH_EDITOR
 	check(InNode);
 	
-	const TArray<int32> FilteredPermutationIndices = InNode->GetFilteredPermutationsIndices();
-
+	TArray<int32> FilteredPermutationIndices = InNode->GetResolvedPermutationIndices(false);
 	if(InNode->IsA<URigVMUnitNode>())
 	{
 		const FLinearColor PreviousColor = InNode->NodeColor;
@@ -16769,28 +16020,21 @@ bool URigVMController::FullyResolveTemplateNode(URigVMTemplateNode* InNode, int3
 	const FRigVMTemplate* Template = InNode->GetTemplate();
 	const FRigVMDispatchFactory* Factory = Template->GetDispatchFactory();
 	
-	int32 InputPermutation = InPermutationIndex;
+	InNode->ResolvedPermutation = InPermutationIndex;
 	
 	// Figure out the permutation index from the pin types
 	if (InPermutationIndex == INDEX_NONE)
 	{
-		FRigVMTemplate::FTypeMap TypeMap;
-		for (URigVMPin* Pin : InNode->GetPins())
-		{
-			check(!Pin->IsWildCard());
-			TypeMap.Add(Pin->GetFName(), Pin->GetTypeIndex());
-		}
-
-		TArray<int32> Permutations;
-		Template->Resolve(TypeMap, Permutations, true);
+		TArray<int32> Permutations = InNode->GetResolvedPermutationIndices(false);
 		check(!Permutations.IsEmpty());
-		InputPermutation = Permutations[0];
+		InNode->ResolvedPermutation = Permutations[0];
 
 		// make sure the permutation exists
 		if(Factory)
 		{
+			FRigVMTemplateTypeMap TypeMap = InNode->GetTemplate()->GetTypesForPermutation(InNode->ResolvedPermutation);
 			const FRigVMFunctionPtr DispatchFunction = Factory->GetDispatchFunction(TypeMap);
-			const FRigVMFunction* ResolvedFunction = Template->GetPermutation(InputPermutation);
+			const FRigVMFunction* ResolvedFunction = Template->GetPermutation(InNode->ResolvedPermutation);
 			check(DispatchFunction);
 			check(ResolvedFunction);
 			check(ResolvedFunction->FunctionPtr == DispatchFunction);
@@ -16798,8 +16042,8 @@ bool URigVMController::FullyResolveTemplateNode(URigVMTemplateNode* InNode, int3
 	}
 	
 
-	const FRigVMFunction* ResolvedFunction = Template->GetPermutation(InputPermutation);
-	const TArray<int32> PermutationIndices = {InputPermutation};
+	const FRigVMFunction* ResolvedFunction = Template->GetPermutation(InNode->ResolvedPermutation);
+	const TArray<int32> PermutationIndices = {InNode->ResolvedPermutation};
 
 	// find all existing pins that we may need to change
 	TArray<FRigVMTemplateArgument> MissingPins;
@@ -16834,7 +16078,7 @@ bool URigVMController::FullyResolveTemplateNode(URigVMTemplateNode* InNode, int3
 			const TArray<FRigVMTemplateArgument>& Arguments = Template->Arguments;
 			for(const FRigVMTemplateArgument& Argument : Arguments)
 			{
-				const TRigVMTypeIndex ExpectedTypeIndex = Argument.TypeIndices[InputPermutation];
+				const TRigVMTypeIndex ExpectedTypeIndex = Argument.TypeIndices[InNode->ResolvedPermutation];
 				if(URigVMPin* Pin = InNode->FindPin(Argument.GetName().ToString()))
 				{
 					if(Pin->GetTypeIndex() != ExpectedTypeIndex && ExpectedTypeIndex != RigVMTypeUtils::TypeIndex::Execute)
@@ -16923,25 +16167,6 @@ bool URigVMController::FullyResolveTemplateNode(URigVMTemplateNode* InNode, int3
 	// exit out early if there's nothing to do
 	if(PinTypesToChange.IsEmpty() && MissingPins.IsEmpty() && PinsToRemove.IsEmpty())
 	{
-		if (URigVMCollapseNode* CollapseNode = Cast<URigVMCollapseNode>(InNode))
-		{
-			{
-				FRigVMControllerGraphGuard GraphGuard(this, CollapseNode->GetContainedGraph());
-				TGuardValue<bool> GuardEditable(CollapseNode->GetContainedGraph()->bEditable, true);
-				TGuardValue<bool> GuardRecomputeOuter(bSuspendRecomputingOuterTemplateFilters, true);
-				TArray<URigVMTemplateNode*> InterfaceNodes = {CollapseNode->GetEntryNode(), CollapseNode->GetReturnNode()};
-				for (URigVMTemplateNode* InterfaceNode : InterfaceNodes)
-				{
-					if (InterfaceNode)
-					{
-						UpdateFilteredPermutations(InterfaceNode, {CollapseNode->FilteredPermutations}, bSetupUndoRedo);
-						break;
-					}
-				}
-				RecomputeAllTemplateFilteredPermutations(bSetupUndoRedo);
-			}
-			UpdateFilteredPermutations(CollapseNode, {0}, bSetupUndoRedo);
-		}
 		ResolveTemplateNodeMetaData(InNode, bSetupUndoRedo);
 		return true;
 	}
@@ -16952,7 +16177,6 @@ bool URigVMController::FullyResolveTemplateNode(URigVMTemplateNode* InNode, int3
 	}
 
 	// update the incorrectly typed pins
-	bool bNeedsTemplateUpdate = false;
 	for(const TPair<URigVMPin*, TRigVMTypeIndex>& Pair : PinTypesToChange)
 	{
 		URigVMPin* Pin = Pair.Key;
@@ -16962,9 +16186,19 @@ bool URigVMController::FullyResolveTemplateNode(URigVMTemplateNode* InNode, int3
 		{
 			if(Pin->GetTypeIndex() != ExpectedTypeIndex)
 			{
-				bNeedsTemplateUpdate = true;
-				const int32 WildCardIndex = Pin->IsArray() ? RigVMTypeUtils::TypeIndex::WildCardArray : RigVMTypeUtils::TypeIndex::WildCard;
-				if(!ChangePinType(Pin, WildCardIndex, bSetupUndoRedo, false, true, true))
+				if (Pin->GetDirection() != ERigVMPinDirection::Hidden)
+				{
+					const int32 WildCardIndex = Pin->IsArray() ? RigVMTypeUtils::TypeIndex::WildCardArray : RigVMTypeUtils::TypeIndex::WildCard;
+					if(!ChangePinType(Pin, WildCardIndex, bSetupUndoRedo, false, true, true))
+					{
+						if(bSetupUndoRedo)
+						{
+							CancelUndoBracket();
+						}
+						return false;
+					}
+				}
+				if(!ChangePinType(Pin, ExpectedTypeIndex, bSetupUndoRedo, false, true, true))
 				{
 					if(bSetupUndoRedo)
 					{
@@ -16974,20 +16208,6 @@ bool URigVMController::FullyResolveTemplateNode(URigVMTemplateNode* InNode, int3
 				}
 			}
 		}
-		
-		if(Pin->IsWildCard())
-		{
-			bNeedsTemplateUpdate = true;
-			
-			if (!UpdateFilteredPermutations(Pin, {ExpectedTypeIndex}, bSetupUndoRedo))
-			{
-				if(bSetupUndoRedo)
-				{
-					CancelUndoBracket();
-				}
-				return false;
-			}
-		}
 	}
 
 	// remove obsolete pins
@@ -17007,7 +16227,7 @@ bool URigVMController::FullyResolveTemplateNode(URigVMTemplateNode* InNode, int3
 				
 				URigVMPin* Pin = NewObject<URigVMPin>(Cast<UObject>(InNode), MissingPin.GetName());
 
-				const TRigVMTypeIndex TypeIndex = MissingPin.TypeIndices[InputPermutation];
+				const TRigVMTypeIndex TypeIndex = MissingPin.TypeIndices[InNode->ResolvedPermutation];
 				const FRigVMTemplateArgumentType& Type = FRigVMRegistry::Get().GetType(TypeIndex); 
 				
 				Pin->Direction = MissingPin.GetDirection();
@@ -17051,11 +16271,6 @@ bool URigVMController::FullyResolveTemplateNode(URigVMTemplateNode* InNode, int3
 		}
 	}
 
-	if (bNeedsTemplateUpdate)
-	{
-		UpdateTemplateNodePinTypes(InNode, bSetupUndoRedo);
-	}
-	
 	if(bSetupUndoRedo)
 	{
 #if WITH_EDITOR
@@ -17067,7 +16282,7 @@ bool URigVMController::FullyResolveTemplateNode(URigVMTemplateNode* InNode, int3
 	return true;
 }
 
-bool URigVMController::PrepareTemplatePinForType(URigVMPin* InPin, const TArray<TRigVMTypeIndex>& InTypeIndices, bool bSetupUndoRedo)
+bool URigVMController::ResolveWildCardPin(const FString& InPinPath, const FString& InCPPType, const FName& InCPPTypeObjectPath, bool bSetupUndoRedo, bool bPrintPythonCommand)
 {
 	if (!IsValidGraph())
 	{
@@ -17079,2245 +16294,521 @@ bool URigVMController::PrepareTemplatePinForType(URigVMPin* InPin, const TArray<
 		return false;
 	}
 
-	URigVMTemplateNode* TemplateNode = Cast<URigVMTemplateNode>(InPin->GetNode());
-	if (!TemplateNode)
+	URigVMGraph* Graph = GetGraph();
+	check(Graph);
+
+	UObject* CPPTypeObject = nullptr;
+	if (!InCPPTypeObjectPath.IsNone())
 	{
-		return false;
+		CPPTypeObject = RigVMTypeUtils::FindObjectFromCPPTypeObjectPath<UObject>(InCPPTypeObjectPath.ToString());
+		if (CPPTypeObject == nullptr)
+		{
+			ReportErrorf(TEXT("Cannot find cpp type object for path '%s'."), *InCPPTypeObjectPath.ToString());
+			return false;
+		}
 	}
 
-	if (TemplateNode->IsSingleton())
+	const FString CPPType = RigVMTypeUtils::PostProcessCPPType(InCPPType, CPPTypeObject);
+
+	if (URigVMPin* Pin = Graph->FindPin(InPinPath))
 	{
-		return false;
+		if (ResolveWildCardPin(Pin, FRigVMTemplateArgumentType(*CPPType, CPPTypeObject), bSetupUndoRedo, bPrintPythonCommand))
+		{
+			if (bPrintPythonCommand)
+			{
+				const FString GraphName = GetSanitizedGraphName(GetGraph()->GetGraphName());
+
+				// bool ResolveWildCardPin(const FString& InPinPath, const FString& InCPPType, const FName& InCPPTypeObjectPath, bool bSetupUndoRedo = true, bool bPrintPythonCommand = false);
+				RigVMPythonUtils::Print(GetGraphOuterName(), 
+									FString::Printf(TEXT("blueprint.get_controller_by_name('%s').resolve_wild_card_pin('%s', '%s', '%s')"),
+													*GraphName,
+													*InPinPath,
+													*InCPPType,
+													*InCPPTypeObjectPath.ToString()));
+			}
+			
+			return true;
+		}
 	}
 
-	if (InTypeIndices.IsEmpty())
+	return false;
+}
+
+bool URigVMController::ResolveWildCardPin(URigVMPin* InPin, const FRigVMTemplateArgumentType& InType, bool bSetupUndoRedo, bool bPrintPythonCommand)
+{
+	if (!IsValidGraph())
 	{
 		return false;
 	}
 
-	FRigVMRegistry& Registry = FRigVMRegistry::Get();
-
-	/*
-	if (InTypeIndices.Num() == 1)
+	if (!bIsTransacting && !IsGraphEditable())
 	{
-		if(Registry.CanMatchTypes(InPin->GetTypeIndex(), InTypeIndices[0], true))
-		{
-			return true;
-		}
+		return false;
 	}
-	*/
 
-	// It might be a subpin of a struct, in that case we just want to make sure
-	// the type required matches the one in the pin
-	if (InPin->IsStructMember() || InPin->IsOrphanPin())
+	return ResolveWildCardPin(InPin, FRigVMRegistry::Get().GetTypeIndex(InType), bSetupUndoRedo, bPrintPythonCommand);
+}
+
+bool URigVMController::ResolveWildCardPin(const FString& InPinPath, TRigVMTypeIndex InTypeIndex, bool bSetupUndoRedo,
+	bool bPrintPythonCommand)
+{
+	if (!IsValidGraph())
 	{
-		TRigVMTypeIndex PinTypeIndex = InPin->GetTypeIndex();
-		
-		// Check if any of the types allows the connection
-		for (const TRigVMTypeIndex& Type : InTypeIndices)
-		{
-			if (Registry.CanMatchTypes(Type, PinTypeIndex, true))
-			{
-				return true;
-			}
-		}
-		
-		return false;		
+		return false;
 	}
 
-	FRigVMBaseAction Action;
-	Action.Title = TEXT("Prepare template pin for type");
-	if (bSetupUndoRedo)
+	if (!bIsTransacting && !IsGraphEditable())
 	{
-		ActionStack->BeginAction(Action);		
+		return false;
 	}
 
-	// If the pin is an element of an array, get the array pin and convert the types
-	URigVMPin* RootPin = InPin;
-	TArray<TRigVMTypeIndex> Types = InTypeIndices;
-	if (InPin->IsArrayElement())
+	URigVMGraph* Graph = GetGraph();
+	check(Graph);
+
+	if (URigVMPin* Pin = Graph->FindPin(InPinPath))
 	{
-		RootPin = InPin->GetParentPin();
-		for (TRigVMTypeIndex& Type : Types)
-		{
-			Type = FRigVMRegistry::Get().GetArrayTypeFromBaseTypeIndex(Type);
-		}
+		return ResolveWildCardPin(Pin, InTypeIndex, bSetupUndoRedo, bPrintPythonCommand);
 	}
+	return false;
+}
 
-	bool bFilteredSupportsType = false;
-	for (const TRigVMTypeIndex& Type : Types)
+bool URigVMController::ResolveWildCardPin(URigVMPin* InPin, TRigVMTypeIndex InTypeIndex, bool bSetupUndoRedo, bool bPrintPythonCommand)
+{
+	if (InPin->IsStructMember())
 	{
-		if (TemplateNode->FilteredSupportsType(RootPin, Type))
-		{
-			bFilteredSupportsType = true;
-			break;
-		}
+		return false;
 	}
 
-	bool bSupportsType = false;
-	TArray<FRigVMTemplatePreferredType> OldPreferredTypes;
-	if (!bFilteredSupportsType)
+	if(FRigVMRegistry::Get().IsWildCardType(InTypeIndex))
 	{
-		if (!TemplateNode->PreferredPermutationPairs.IsEmpty())
-		{
-			if (bSetupUndoRedo)
-			{
-				ActionStack->AddAction(FRigVMSetPreferredTemplatePermutationsAction(TemplateNode, {}));
-			}
-			OldPreferredTypes = TemplateNode->PreferredPermutationPairs;
-			TemplateNode->PreferredPermutationPairs = {};
-		}
-		
-		for (const TRigVMTypeIndex& Type : Types)
-		{
-			if (TemplateNode->SupportsType(RootPin, Type))
-			{
-				bSupportsType = true;
-				break;
-			}
-		}	
+		return false;
 	}
 
-	if (!bFilteredSupportsType && !bSupportsType)
+	if (InPin->GetTypeIndex() == InTypeIndex)
 	{
-		if (bSetupUndoRedo)
-		{
-			ActionStack->CancelAction(Action, this);
-		}
 		return false;
 	}
 
-	if (bFilteredSupportsType)
+	URigVMPin* RootPin = InPin;
+	TRigVMTypeIndex Type = InTypeIndex;	
+	while (RootPin->IsArrayElement())
 	{
-		{
-			TGuardValue<FRigVMController_RequestBreakLinksDialogDelegate> GuardDelegate(RequestBreakLinksDialogDelegate, nullptr);
-			OpenUndoBracket(FString::Printf(TEXT("Resolve wildcard pin %s"), *InPin->GetPinPath()));
-			if (TemplateNode->PinNeedsFilteredTypesUpdate(RootPin, Types))
-			{
-				if (UpdateFilteredPermutations(RootPin, Types, bSetupUndoRedo))
-				{
-					UpdateTemplateNodePinTypes(TemplateNode, bSetupUndoRedo);
-					PropagateTemplateFilteredTypes(TemplateNode, bSetupUndoRedo);
-				}
-			}
-		}
+		RootPin = InPin->GetParentPin();
+		Type = FRigVMRegistry::Get().GetArrayTypeFromBaseTypeIndex(Type);
+	}
+	
+	ensure(RootPin->GetNode()->IsA<URigVMTemplateNode>());
+	URigVMTemplateNode* TemplateNode = CastChecked<URigVMTemplateNode>(RootPin->GetNode());
 
-		// Even if the filter supports type, we might find we break links in parent graphs
-		// Find break link actions
-		TArray<URigVMLink*> LinksToBreak;
-		TArray<TPair<FString, FString>> InconsistentLinks;
-		TArray<FSoftObjectPath> InconsistentLinksGraph;
+	TRigVMTypeIndex NewType = INDEX_NONE;
+	TemplateNode->SupportsType(RootPin, Type, &NewType);
+	if (NewType != INDEX_NONE)
+	{
+		// We support the new type, and its different than the pin type
+		if (NewType == Type && InPin->GetTypeIndex() != NewType)
 		{
-			TArray<FRigVMActionKey> Actions = ActionStack->BracketActions.Last()->SubActions;
-			for (int32 i=0; i<Actions.Num(); ++i)
-			{
-				const FRigVMActionKey& ActionKey = Actions[i];
-				FRigVMActionWrapper Wrapper(ActionKey);
-				if (Wrapper.GetAction()->GetScriptStruct() == FRigVMBreakLinkAction::StaticStruct())
-				{
-					const FRigVMBreakLinkAction* BreakLinkAction = (const FRigVMBreakLinkAction*)Wrapper.GetAction();
-					InconsistentLinks.Add(TPair<FString,FString>(BreakLinkAction->OutputPinPath, BreakLinkAction->InputPinPath));
-					InconsistentLinksGraph.Add(BreakLinkAction->GraphPath);
-				}
-
-				Actions.Append(Wrapper.GetAction()->SubActions);
-			}
+			 
 		}
-
-		if (InconsistentLinks.IsEmpty())
+		// We support the new type, but its different than the one provided
+		else if (NewType != Type)
 		{
-			CloseUndoBracket();
-			if (bSetupUndoRedo)
+			// if its the same as the pin, we are done
+			if (InPin->GetTypeIndex() == NewType)
 			{
-				ActionStack->EndAction(Action);
+				return false;
 			}
-			return true;
-		}
-
-		verify(CancelUndoBracket());
-		for (int32 i=0; i<InconsistentLinks.Num(); ++i)
-		{
-			TSoftObjectPtr<URigVMGraph> GraphPtr(InconsistentLinksGraph[i]);
-			LinksToBreak.AddUnique(GraphPtr.Get()->FindLink(URigVMLink::GetPinPathRepresentation(InconsistentLinks[i].Key, InconsistentLinks[i].Value)));
-		}
-		
-		// Warn the user the links that would be broken
-		bool bBreakLinks = true;
-		bool bConsultedBreak = false;
-		if(!bIsTransacting && RequestBreakLinksDialogDelegate.IsBound())
-		{
-			bBreakLinks = RequestBreakLinksDialogDelegate.Execute(LinksToBreak);
-			bConsultedBreak = true;
+			Type = NewType;
 		}
+	}
+	else
+	{
+		return false;
+	}
+	
+	FRigVMBaseAction Action;
+	if (bSetupUndoRedo)
+	{
+		Action.Title = TEXT("Resolve Wildcard Pin");
+		ActionStack->BeginAction(Action);
+	}
 
-		if (bBreakLinks)
+	if (!InPin->IsWildCard())
+	{
+		if (!UnresolveTemplateNodes({TemplateNode}, bSetupUndoRedo))
 		{
-			// Break links
-			for (URigVMLink* Link : LinksToBreak)
-			{
-				FRigVMControllerGraphGuard GraphGuard(this, Link->GetGraph(), bSetupUndoRedo);
-				if (!bConsultedBreak && !bIsTransacting && !bSuspendNotifications)
-				{
-					ReportWarningf(TEXT("The link between %s was broken due to inconsistent types"), *Link->GetPinPathRepresentation());
-				}
-				BreakLink(Link->GetSourcePin(), Link->GetTargetPin(), bSetupUndoRedo);
-			}
-
-			ensure(LinksToBreak.Num() > 0);
-			PrepareTemplatePinForType(InPin, InTypeIndices, bSetupUndoRedo);
-
-			if (bSetupUndoRedo)
-			{
-				ActionStack->EndAction(Action);
-			}
-			return true;
+			return false;
 		}
-
+	}
+	
+	if (!ChangePinType(RootPin, Type, bSetupUndoRedo, true, true, false))
+	{
 		if (bSetupUndoRedo)
 		{
 			ActionStack->CancelAction(Action, this);
 		}
 		return false;
 	}
-	else
-	{
-		// Figure out what links we would need to break to allow this
-		TArray<URigVMLink*> LinksToBreak;
-		{
-			TGuardValue<FRigVMController_RequestBreakLinksDialogDelegate> GuardDelegate(RequestBreakLinksDialogDelegate, nullptr);
-			TGuardValue<bool> GuardNotifications(bSuspendNotifications, true);
-			TGuardValue<bool> SuspendRecomputeTemplates(bSuspendRecomputingTemplateFilters, true);
-			TGuardValue<bool> DisableCasting(bEnableTypeCasting, false);
-			OpenUndoBracket(FString::Printf(TEXT("Resolve wildcard pin %s"), *InPin->GetPinPath()));
 
-			bool bBrokenLinks = false;
-			do
-			{
-				bBrokenLinks = false;
-				
-				// Initialize all template nodes in the graph
-				InitializeAllTemplateFiltersInGraph(true, false);
+	UpdateTemplateNodePinTypes(TemplateNode, bSetupUndoRedo);
 
-				// Update our node with the requested type
-				UpdateFilteredPermutations(InPin, InTypeIndices, true);
-				
-				// Update types in case some pins on this node disappear, and should no longer propagate
-				UpdateTemplateNodePinTypes(TemplateNode, true);
+	if (bSetupUndoRedo)
+	{
+		ActionStack->EndAction(Action);
+	}
+	
+	return true;
 
-				// Propagate the filtered types and collect break link actions 
-				bBrokenLinks = !PropagateTemplateFilteredTypes(TemplateNode, true);
-			}
-			while (bBrokenLinks);
-
-			// Find break link actions
-			TArray<TPair<FString, FString>> InconsistentLinks;
-			TArray<FRigVMActionKey> Actions = ActionStack->BracketActions.Last()->SubActions;
-			for (int32 i=0; i<Actions.Num(); ++i)
-			{
-				const FRigVMActionKey& ActionKey = Actions[i];
-				FRigVMActionWrapper Wrapper(ActionKey);
-				if (Wrapper.GetAction()->GetScriptStruct() == FRigVMBreakLinkAction::StaticStruct())
-				{
-					const FRigVMBreakLinkAction* BreakLinkAction = (const FRigVMBreakLinkAction*)Wrapper.GetAction();
-					InconsistentLinks.AddUnique(TPair<FString,FString>(BreakLinkAction->OutputPinPath, BreakLinkAction->InputPinPath));
-				}
-
-				Actions.Append(Wrapper.GetAction()->SubActions);
-			}
-			verify(CancelUndoBracket());
-
-			for (TPair<FString, FString>& Pair : InconsistentLinks)
-			{
-				LinksToBreak.AddUnique(GetGraph()->FindLink(URigVMLink::GetPinPathRepresentation(Pair.Key, Pair.Value)));
-			}
-		}
-
-		// We might be in this situation just because we are resolving to a different type than previosly resolved
-		// so no links to break
-		if (LinksToBreak.IsEmpty())
-		{
-			if (!TemplateNode->PreferredPermutationPairs.IsEmpty())
-			{
-				ensure(false);
-				if (bSetupUndoRedo)
-				{
-					ActionStack->CancelAction(Action, this);
-				}
-				return false;
-			}
-			
-			UnresolveTemplateNodes({TemplateNode}, bSetupUndoRedo, false);
-			PrepareTemplatePinForType(InPin, InTypeIndices, bSetupUndoRedo);
-
-			// Reapply the valid preferred types
-			for (int32 i=0; i<OldPreferredTypes.Num(); ++i)
-			{
-				if (URigVMPin* Pin = TemplateNode->FindPin(OldPreferredTypes[i].Argument.ToString()))
-				{
-					if (Pin->IsWildCard())
-					{
-						if (!TemplateNode->FilteredSupportsType(Pin, OldPreferredTypes[i].GetTypeIndex()))
-						{
-							continue;
-						}
-					}
-				
-					if (Pin->IsWildCard() || Pin->GetTypeIndex() == OldPreferredTypes[i].GetTypeIndex())
-					{
-						ResolveWildCardPin(Pin, OldPreferredTypes[i].GetTypeIndex(), bSetupUndoRedo);
-					}
-				}
-			}
-			if (bSetupUndoRedo)
-			{
-				ActionStack->EndAction(Action);
-			}
-			return true;
-		}
-
-		// Warn the user the links that would be broken
-		bool bBreakLinks = true;
-		bool bConsultedBreak = false;
-		if(!bIsTransacting && RequestBreakLinksDialogDelegate.IsBound())
-		{
-			bBreakLinks = RequestBreakLinksDialogDelegate.Execute(LinksToBreak);
-			bConsultedBreak = true;
-		}
-
-		if (bBreakLinks)
-		{
-			// Break links
-			for (URigVMLink* Link : LinksToBreak)
-			{
-				if (!bConsultedBreak && !bIsTransacting && !bSuspendNotifications)
-				{
-					ReportWarningf(TEXT("The link between %s was broken due to inconsistent types"), *Link->GetPinPathRepresentation());
-				}
-				BreakLink(Link->GetSourcePin(), Link->GetTargetPin(), bSetupUndoRedo);
-			}
-
-			ensure(LinksToBreak.Num() > 0);
-			PrepareTemplatePinForType(InPin, InTypeIndices, bSetupUndoRedo);
-
-			// Reapply the valid preferred types
-			for (int32 i=0; i<OldPreferredTypes.Num(); ++i)
-			{
-				if (URigVMPin* Pin = TemplateNode->FindPin(OldPreferredTypes[i].Argument.ToString()))
-				{
-					if (Pin->IsWildCard())
-					{
-						if (!TemplateNode->FilteredSupportsType(Pin, OldPreferredTypes[i].GetTypeIndex()))
-						{
-							continue;
-						}
-					}
-				
-					if (Pin->IsWildCard() || Pin->GetTypeIndex() == OldPreferredTypes[i].GetTypeIndex())
-					{
-						ResolveWildCardPin(Pin, OldPreferredTypes[i].GetTypeIndex(), bSetupUndoRedo);
-					}
-				}
-			}
-			if (bSetupUndoRedo)
-			{
-				ActionStack->EndAction(Action);
-			}
-			return true;
-		}
-	}
-
-	if (bSetupUndoRedo)
-	{
-		ActionStack->CancelAction(Action, this);
-	}
-	return false;
-}
-
-TArray<TRigVMTypeIndex> URigVMController::GetFilteredTypes(URigVMPin* InPin)
-{
-	if (!InPin->IsStructMember() && !InPin->IsOrphanPin())
-	{
-		if (URigVMTemplateNode* TemplateNode = Cast<URigVMTemplateNode>(InPin->GetNode()))
-		{
-			if (!TemplateNode->IsSingleton())
-			{
-				return TemplateNode->GetFilteredTypesForPin(InPin);
-			}
-		}
-	}
-	if (InPin->IsWildCard())
-	{
-		return TArray<TRigVMTypeIndex>();
-	}
-	return {InPin->GetTypeIndex()};
-}
-
-bool URigVMController::ResolveWildCardPin(const FString& InPinPath, const FString& InCPPType, const FName& InCPPTypeObjectPath, bool bSetupUndoRedo, bool bPrintPythonCommand)
-{
-	if (!IsValidGraph())
-	{
-		return false;
-	}
-
-	if (!bIsTransacting && !IsGraphEditable())
-	{
-		return false;
-	}
-
-	URigVMGraph* Graph = GetGraph();
-	check(Graph);
-
-	UObject* CPPTypeObject = nullptr;
-	if (!InCPPTypeObjectPath.IsNone())
-	{
-		CPPTypeObject = RigVMTypeUtils::FindObjectFromCPPTypeObjectPath<UObject>(InCPPTypeObjectPath.ToString());
-		if (CPPTypeObject == nullptr)
-		{
-			ReportErrorf(TEXT("Cannot find cpp type object for path '%s'."), *InCPPTypeObjectPath.ToString());
-			return false;
-		}
-	}
-
-	const FString CPPType = RigVMTypeUtils::PostProcessCPPType(InCPPType, CPPTypeObject);
-
-	if (URigVMPin* Pin = Graph->FindPin(InPinPath))
-	{
-		if (ResolveWildCardPin(Pin, FRigVMTemplateArgumentType(*CPPType, CPPTypeObject), bSetupUndoRedo, bPrintPythonCommand))
-		{
-			if (bPrintPythonCommand)
-			{
-				const FString GraphName = GetSanitizedGraphName(GetGraph()->GetGraphName());
-
-				// bool ResolveWildCardPin(const FString& InPinPath, const FString& InCPPType, const FName& InCPPTypeObjectPath, bool bSetupUndoRedo = true, bool bPrintPythonCommand = false);
-				RigVMPythonUtils::Print(GetGraphOuterName(), 
-									FString::Printf(TEXT("blueprint.get_controller_by_name('%s').resolve_wild_card_pin('%s', '%s', '%s')"),
-													*GraphName,
-													*InPinPath,
-													*InCPPType,
-													*InCPPTypeObjectPath.ToString()));
-			}
-			
-			return true;
-		}
-	}
-
-	return false;
-}
-
-bool URigVMController::ResolveWildCardPin(URigVMPin* InPin, const FRigVMTemplateArgumentType& InType, bool bSetupUndoRedo, bool bPrintPythonCommand)
-{
-	if (!IsValidGraph())
-	{
-		return false;
-	}
-
-	if (!bIsTransacting && !IsGraphEditable())
-	{
-		return false;
-	}
-
-	return ResolveWildCardPin(InPin, FRigVMRegistry::Get().GetTypeIndex(InType), bSetupUndoRedo, bPrintPythonCommand);
-}
-
-bool URigVMController::ResolveWildCardPin(const FString& InPinPath, TRigVMTypeIndex InTypeIndex, bool bSetupUndoRedo,
-	bool bPrintPythonCommand)
-{
-	if (!IsValidGraph())
-	{
-		return false;
-	}
-
-	if (!bIsTransacting && !IsGraphEditable())
-	{
-		return false;
-	}
-
-	URigVMGraph* Graph = GetGraph();
-	check(Graph);
-
-	if (URigVMPin* Pin = Graph->FindPin(InPinPath))
-	{
-		return ResolveWildCardPin(Pin, InTypeIndex, bSetupUndoRedo, bPrintPythonCommand);
-	}
-	return false;
-}
-
-bool URigVMController::ResolveWildCardPin(URigVMPin* InPin, TRigVMTypeIndex InTypeIndex, bool bSetupUndoRedo,
-	bool bPrintPythonCommand)
-{
-	if (InPin->IsStructMember())
-	{
-		return false;
-	}
-
-	if(FRigVMRegistry::Get().IsWildCardType(InTypeIndex))
-	{
-		return false;
-	}
-
-	URigVMPin* RootPin = InPin;
-	TRigVMTypeIndex Type = InTypeIndex;	
-	while (RootPin->IsArrayElement())
-	{
-		RootPin = InPin->GetParentPin();
-		Type = FRigVMRegistry::Get().GetArrayTypeFromBaseTypeIndex(Type);
-	}
-	
-	if(RootPin->IsWildCard())
-	{
-		ensure(RootPin->GetNode()->IsA<URigVMTemplateNode>());
-		URigVMTemplateNode* TemplateNode = CastChecked<URigVMTemplateNode>(RootPin->GetNode());
-		if (!TemplateNode->SupportsType(RootPin, Type, &Type))
-		{
-			return false;
-		}
-		
-		FRigVMBaseAction Action;
-		if (bSetupUndoRedo)
-		{
-			Action.Title = TEXT("Resolve Wildcard Pin");
-			ActionStack->BeginAction(Action);
-		}
-		
-		// Add preferred type
-		AddPreferredType(TemplateNode, RootPin->GetFName(), Type, bSetupUndoRedo);
-
-		if (TemplateNode->IsA<URigVMFunctionEntryNode>() || TemplateNode->IsA<URigVMFunctionReturnNode>())
-		{
-			if (TemplateNode->GetTemplate()->FindArgument(RootPin->GetFName()) == nullptr)
-			{
-				if (!AddArgumentForPin(RootPin, nullptr, bSetupUndoRedo))
-				{
-					if (bSetupUndoRedo)
-					{
-						ActionStack->CancelAction(Action, this);
-					}
-					return false;
-				}
-			}
-		}
-		
-		if (!PrepareTemplatePinForType(InPin, {InTypeIndex}, bSetupUndoRedo))
-		{
-			if (bSetupUndoRedo)
-			{
-				ActionStack->CancelAction(Action, this);
-			}
-			return false;
-		}
-				
-		
-
-		URigVMGraph* Graph = GetGraph();
-		if (URigVMGraph* ParentGraph = Graph->GetParentGraph())
-		{
-			FRigVMControllerGraphGuard GraphGuard(this, ParentGraph, bSetupUndoRedo);
-			UpdateLibraryTemplate(Graph->GetTypedOuter<URigVMLibraryNode>(), bSetupUndoRedo);
-		}
-		
-		if (bSetupUndoRedo)
-		{
-			ActionStack->EndAction(Action);
-		}
-		
-		return true;		
-	}
-	else if (FRigVMRegistry::Get().CanMatchTypes(RootPin->GetTypeIndex(), Type, true))
-	{
-		if (URigVMTemplateNode* TemplateNode = Cast<URigVMTemplateNode>(RootPin->GetNode()))
-		{
-			TRigVMTypeIndex TypeIndex = TemplateNode->GetPreferredType(RootPin->GetFName());
-			if (TypeIndex == Type && RootPin->GetTypeIndex() == Type)
-			{
-				return true;
-			}
-			
-			if (TypeIndex != Type && TypeIndex != INDEX_NONE)
-			{
-				RemovePreferredType(TemplateNode, RootPin->GetFName(), bSetupUndoRedo);
-				TypeIndex = INDEX_NONE;
-			}
-			if (TypeIndex == INDEX_NONE)
-			{
-				AddPreferredType(TemplateNode, RootPin->GetFName(), Type, bSetupUndoRedo);
-			}
-	
-			if (RootPin->GetTypeIndex() != Type)
-			{
-				UpdateTemplateNodePinTypes(TemplateNode, bSetupUndoRedo, false);
-			}
-		}
-	}
-	else
-	{
-		OpenUndoBracket(TEXT("Resolving wildcard pin"));
-		URigVMTemplateNode* Node = Cast<URigVMTemplateNode>(InPin->GetNode());
-		TArray<URigVMLink*> Links = Node->GetLinks();
-		UnresolveTemplateNodes({*Node->GetName()}, bSetupUndoRedo);
-		if (ResolveWildCardPin(InPin, InTypeIndex, bSetupUndoRedo, bPrintPythonCommand))
-		{
-			// Try to reattach old links
-			for (URigVMLink* Link : Links)
-			{
-				URigVMPin* SourcePin = Link->GetSourcePin();
-				URigVMPin* TargetPin = Link->GetTargetPin();
-				URigVMPin* LocalPin = (SourcePin->GetNode() == Node) ? SourcePin : TargetPin;
-				URigVMPin* RemotePin = (LocalPin == TargetPin) ? SourcePin : TargetPin;
-				{
-					TArray<TRigVMTypeIndex> FilteredTypes = GetFilteredTypes(RemotePin);
-					bool bSupportsLink = false;
-					for (const TRigVMTypeIndex& FilteredType : FilteredTypes)
-					{
-						if (Node->FilteredSupportsType(LocalPin, FilteredType))
-						{
-							bSupportsLink = true;
-							break;
-						}
-					}
-					if (bSupportsLink)
-					{
-						AddLink(SourcePin, TargetPin, bSetupUndoRedo);
-					}
-				}
-			}
-			
-			CloseUndoBracket();
-			return true;
-		}
-		CancelUndoBracket();
-	}
-
-	return false;
-}
-
-bool URigVMController::UpdateFilteredPermutations(URigVMPin* InPin, URigVMPin* InLinkedPin, bool bSetupUndoRedo)
-{
-	URigVMTemplateNode* Node = Cast<URigVMTemplateNode>(InPin->GetNode());
-	if (!Node)
-	{
-		return false;
-	}
-	
-	TArray<int32> OldPermutations;
-	if (bSetupUndoRedo)
-	{
-		OldPermutations = Node->GetFilteredPermutationsIndices();
-	}
-
-	if (Node->UpdateFilteredPermutations(InPin, InLinkedPin))
-	{
-		if (bSetupUndoRedo)
-		{
-			FRigVMSetTemplateFilteredPermutationsAction Action(Node, OldPermutations);
-			ActionStack->BeginAction(Action);
-			ActionStack->AddAction(Action);
-			ActionStack->EndAction(Action);
-		}
-
-		// If updating the filtered permutations of entry or return nodes, update the library template
-		if (Node->IsA<URigVMFunctionEntryNode>() || Node->IsA<URigVMFunctionReturnNode>())
-		{
-			if (URigVMLibraryNode* LibraryNode = Node->GetTypedOuter<URigVMLibraryNode>())
-			{
-				UpdateLibraryTemplate(LibraryNode, bSetupUndoRedo);
-			}
-		}
-
-		return true;
-	}
-
-	return false;
-}
-
-bool URigVMController::UpdateFilteredPermutations(URigVMPin* InPin, const TArray<TRigVMTypeIndex>& InTypeIndices, bool bSetupUndoRedo)
-{
-	URigVMTemplateNode* Node = Cast<URigVMTemplateNode>(InPin->GetNode());
-	if (!Node)
-	{
-		return false;
-	}
-	
-	TArray<int32> OldPermutations;
-	if (bSetupUndoRedo)
-	{
-		OldPermutations = Node->GetFilteredPermutationsIndices();
-	}
-
-	if (Node->UpdateFilteredPermutations(InPin, InTypeIndices))
-	{
-		if (bSetupUndoRedo)
-		{
-			FRigVMSetTemplateFilteredPermutationsAction Action(Node, OldPermutations);
-			ActionStack->BeginAction(Action);
-			ActionStack->AddAction(Action);
-			ActionStack->EndAction(Action);
-		}
-		
-		// If updating the filtered permutations of entry or return nodes, update the library template
-        if (Node->IsA<URigVMFunctionEntryNode>() || Node->IsA<URigVMFunctionReturnNode>())
-        {
-        	if (URigVMLibraryNode* LibraryNode = Node->GetTypedOuter<URigVMLibraryNode>())
-        	{
-        		UpdateLibraryTemplate(LibraryNode, bSetupUndoRedo);
-        	}
-        }
-
-		return true;
-	}
-
-	return false;
-}
-
-bool URigVMController::UpdateFilteredPermutations(URigVMTemplateNode* InNode, const TArray<int32>& InPermutations, bool bSetupUndoRedo)
-{
-	TArray<int32> OldPermutations;
-	if (bSetupUndoRedo)
-	{
-		OldPermutations = InNode->GetFilteredPermutationsIndices();
-	}
-
-	InNode->FilteredPermutations = InPermutations;
-
-	if (bSetupUndoRedo)
-	{
-		FRigVMSetTemplateFilteredPermutationsAction Action(InNode, OldPermutations);
-		ActionStack->BeginAction(Action);
-		ActionStack->AddAction(Action);
-		ActionStack->EndAction(Action);
-	}
-
-	// If updating the filtered permutations of entry or return nodes, update the library template
-	if (InNode->IsA<URigVMFunctionEntryNode>() || InNode->IsA<URigVMFunctionReturnNode>())
-	{
-		if (URigVMLibraryNode* LibraryNode = InNode->GetTypedOuter<URigVMLibraryNode>())
-		{
-			UpdateLibraryTemplate(LibraryNode, bSetupUndoRedo);
-		}
-	}
-
-	return true;
-}
-
-bool URigVMController::UpdateTemplateNodePinTypes(URigVMTemplateNode* InNode, bool bSetupUndoRedo, bool bInitializeDefaultValue)
-{
-	URigVMGraph* Graph = GetGraph();
-	check(InNode->GetGraph() == Graph);
-	bool bAnyTypeChanged = false;
-
-	const FRigVMTemplate* Template = InNode->GetTemplate();
-	if(Template == nullptr)
-	{
-		return false;
-	}
-
-	// Because interface nodes reference the same template, they need to be resolved at the same time (otherwise we
-	// might end up in a situation where they resolve to different permutations).
-	const bool bIsInterfaceNode = InNode->IsA<URigVMFunctionEntryNode>() || InNode->IsA<URigVMFunctionReturnNode>();
-	URigVMTemplateNode* EntryNode = nullptr;
-	URigVMTemplateNode* ReturnNode = nullptr;
-	if (bIsInterfaceNode)
-	{
-		EntryNode = Graph->GetEntryNode();
-		ReturnNode = Graph->GetReturnNode();
-	}
-	
-	TArray<int32> ResolvedPermutations;
-	TArray<int32> PreferredPermutations;
-	TArray<FRigVMTemplatePreferredType> PreferredPermutationPairs = InNode->PreferredPermutationPairs;
-	if (!bIsInterfaceNode)
-	{
-		ResolvedPermutations = InNode->GetFilteredPermutationsIndices();
-		
-		// Try to pick permutations that respects the preferred types
-		PreferredPermutations = InNode->FindPermutationsForTypes(InNode->PreferredPermutationPairs, true);
-		ResolvedPermutations = ResolvedPermutations.FilterByPredicate([&](const int32& OtherPermutation) { return PreferredPermutations.Contains(OtherPermutation); });
-		InNode->ResolvedPermutation = ResolvedPermutations.Num() == 1 ? ResolvedPermutations[0] : INDEX_NONE;		
-	}
-	else
-	{
-		if (EntryNode)
-		{
-			ResolvedPermutations = EntryNode->FilteredPermutations;
-			PreferredPermutationPairs = EntryNode->PreferredPermutationPairs;
-		}
-		if (ReturnNode)
-		{
-			ResolvedPermutations = ResolvedPermutations.FilterByPredicate([ReturnNode](const int32& Index)
-			{
-				return ReturnNode->FilteredPermutations.Contains(Index);
-			});
-
-			for (const FRigVMTemplatePreferredType& Preferred : ReturnNode->PreferredPermutationPairs)
-			{
-				PreferredPermutationPairs.AddUnique(Preferred);				
-			}
-		}
-	
-		// Try to pick permutations that respects the preferred types
-		PreferredPermutations = InNode->FindPermutationsForTypes(PreferredPermutationPairs, true);
-		ResolvedPermutations = ResolvedPermutations.FilterByPredicate([&](const int32& OtherPermutation) { return PreferredPermutations.Contains(OtherPermutation); });
-
-		if (EntryNode)
-		{
-			EntryNode->ResolvedPermutation = ResolvedPermutations.Num() == 1 ? ResolvedPermutations[0] : INDEX_NONE;
-		}
-		if (ReturnNode)
-		{
-			ReturnNode->ResolvedPermutation = ResolvedPermutations.Num() == 1 ? ResolvedPermutations[0] : INDEX_NONE;
-		}
-	}
-
-	TArray<URigVMPin*> Pins = InNode->GetPins();
-	if (bIsInterfaceNode)
-	{
-		if (EntryNode)
-		{
-			Pins = EntryNode->GetPins();
-		}
-		if (ReturnNode)
-		{
-			Pins.Append(ReturnNode->GetPins());
-		}
-	}
-
-	TArray<const FRigVMTemplateArgument*> Arguments;
-	Arguments.SetNumUninitialized(Pins.Num());
-	for(int32 PinIndex=0; PinIndex < Pins.Num(); ++PinIndex)
-	{
-		URigVMPin* Pin = Pins[PinIndex];
-		Arguments[PinIndex] = Template->FindArgument(Pin->GetFName());
-	}
-
-	// Remove invalid permutations
-	ResolvedPermutations.RemoveAll([&](const int32& Permutation)
-	{
-		for (const FRigVMTemplateArgument* Argument : Arguments)
-		{
-			if (Argument && Argument->TypeIndices[Permutation] == INDEX_NONE)
-			{
-				return true;
-			}
-		}
-		return false;
-	});	
-
-	// Find the types for each possible permutation
-	TMap<int32, TArray<TRigVMTypeIndex>> PinTypes; // permutation to pin types
-	for (int32 ResolvedPermutation : ResolvedPermutations)
-	{
-		for(int32 PinIndex=0; PinIndex < Pins.Num(); ++PinIndex)
-		{
-			URigVMPin* Pin = Pins[PinIndex];
-			if (Pin->GetDirection() == ERigVMPinDirection::Hidden)
-			{
-				PinTypes.FindOrAdd(ResolvedPermutation).Add(INDEX_NONE);
-				continue;
-			}
-
-			bool bAddedType = false;
-			TArray<TRigVMTypeIndex>& Types = PinTypes.FindOrAdd(ResolvedPermutation);
-			if (const FRigVMTemplateArgument* Argument = Arguments[PinIndex])
-			{
-				Types.Add(Argument->TypeIndices[ResolvedPermutation]);
-				bAddedType = true;
-			}
-			else
-			{
-				// We might be in the case where the pin doesn't own an argument yet, but the node has a preferred type
-				for (int32 i=0; i<PreferredPermutationPairs.Num(); ++i)
-				{
-					if (Pin->GetFName() == PreferredPermutationPairs[i].Argument)
-					{
-						Types.Add(PreferredPermutationPairs[i].GetTypeIndex());
-						bAddedType = true;
-						break;
-					}
-				}
-			}
-			
-			if (!bAddedType)
-			{
-				// This is a pin with no argument or preferred type
-				// Its marked as invalid, and will show wildcard type
-				Types.Add(INDEX_NONE);
-			}
-		}
-	}
-
-	// Some pins can benefit from reducing types to a single option
-	// Even if some pins reduce to a single type (maybe represented by a single permutation), the rest should still display a wildcard pin
-	// If reduction happens for multiple pins, we need to make sure that they reduce to the same permutation
-	// If possible, we want to respect the preferred type if it exists, or the current pin type if its not wildcard
-	TArray<bool> WasReduced;
-	WasReduced.SetNumZeroed(Pins.Num());
-
-	TMap<int32, TArray<TRigVMTypeIndex>> ReducedTypes = PinTypes;
-	TArray<TRigVMTypeIndex> FinalPinTypes;
-	FinalPinTypes.SetNumUninitialized(Pins.Num());
-	
-	// Reduce to preferred types (no casting)
-	if (ReducedTypes.Num() > 1)
-	{
-		for (FRigVMTemplatePreferredType& Preferred : PreferredPermutationPairs)
-		{
-			// Finding all pins related to this argument (in the case of interface nodes, there might be more than one)
-			TArray<URigVMPin*> ArgumentPins = Pins.FilterByPredicate([Preferred](URigVMPin* Pin)
-			{
-				return Pin->GetFName() == Preferred.Argument;
-			});
-
-			for (URigVMPin* ArgumentPin : ArgumentPins)
-			{
-				int32 PinIndex = Pins.Find(ArgumentPin);
-				WasReduced[PinIndex] = true;
-				FinalPinTypes[PinIndex] = Preferred.TypeIndex;
-
-				// If the argument is hidden, the available PinTypes will all be INDEX_NONE
-				if (ArgumentPin->Direction != ERigVMPinDirection::Hidden)
-				{					
-					ReducedTypes = ReducedTypes.FilterByPredicate([Preferred, PinIndex](const TPair<int32, TArray<TRigVMTypeIndex>>& Permutation)
-					{
-						return Permutation.Value[PinIndex] == Preferred.TypeIndex;
-					});
-				}
-			}
-		}
-	}
-
-	if (PinTypes.Num() > ReducedTypes.Num())
-	{
-		ensureMsgf(!ReducedTypes.IsEmpty(), TEXT("Found incompatible preferred types"));
-	}
-	
-	// Reduce compatible types, try to find the permutation in the reduced types
-	for(int32 PinIndex=0; PinIndex < Pins.Num(); ++PinIndex)
-	{
-		if (WasReduced[PinIndex])
-		{
-			continue;
-		}
-		
-		URigVMPin* Pin = Pins[PinIndex];
-		if (Pin->GetDirection() == ERigVMPinDirection::Hidden)
-		{
-			continue;
-		}
-		TArray<TRigVMTypeIndex> Types;
-		Types.Reserve(ResolvedPermutations.Num());
-		TRigVMTypeIndex PreferredType = INDEX_NONE;
-		int32 TypesFoundInReduced = 0;
-		for (int32 ResolvedPermutation : ResolvedPermutations)
-		{
-			Types.Add(PinTypes.FindChecked(ResolvedPermutation)[PinIndex]);
-			if (ReducedTypes.Contains(ResolvedPermutation))
-			{
-				TypesFoundInReduced++;
-				if (PreferredType == INDEX_NONE)
-				{
-					PreferredType = PinTypes.FindChecked(ResolvedPermutation)[PinIndex];
-				}				
-			}
-		}
-
-		if (TypesFoundInReduced > 1)
-		{
-			PreferredType = Pin->GetTypeIndex();
-		}
-		FinalPinTypes[PinIndex] = InNode->TryReduceTypesToSingle(Types, PreferredType);
-		if (FinalPinTypes[PinIndex] != INDEX_NONE)
-		{
-			WasReduced[PinIndex] =  true;
-		}
-	}
-	
-	// First unresolve any pins that need unresolving, then resolve everything else
-	for (int32 UnresolveResolve=0; UnresolveResolve<2; ++UnresolveResolve)
-	{
-		for(int32 PinIndex=0; PinIndex < Pins.Num(); ++PinIndex)
-		{
-			URigVMPin* Pin = Pins[PinIndex];
-			if (UnresolveResolve == 0 && Pin->IsWildCard())
-			{
-				continue;
-			}
-			
-			if (Pin->GetDirection() == ERigVMPinDirection::Hidden)
-			{
-				continue;
-			}
-
-			bool bShouldUnresolve = FinalPinTypes[PinIndex] == INDEX_NONE;
-			// If we are about to resolve the pin to a different type, first unresolve it
-			if (!bShouldUnresolve && UnresolveResolve == 0 && Pin->GetTypeIndex() != FinalPinTypes[PinIndex])
-			{
-				bShouldUnresolve = true;
-			}
-
-			if (bShouldUnresolve)
-			{
-				// Unresolve
-				if (Pin->HasInjectedNodes())
-				{
-					EjectNodeFromPin(Pin, bSetupUndoRedo);
-				}
-				
-				const FRigVMTemplateArgument* Argument = Template->FindArgument(*Pin->GetName());
-				FRigVMTemplateArgument::EArrayType ArrayType;
-				if (Argument)
-				{
-					ArrayType = Argument->GetArrayType();
-				}
-				else
-				{
-					ArrayType = Pin->IsArray() ? FRigVMTemplateArgument::EArrayType::EArrayType_ArrayValue : FRigVMTemplateArgument::EArrayType::EArrayType_SingleValue;
-				}
-				
-				FString CPPType = RigVMTypeUtils::GetWildCardCPPType();
-				UObject* CPPObjectType = RigVMTypeUtils::GetWildCardCPPTypeObject();
-
-				if (ArrayType == FRigVMTemplateArgument::EArrayType_ArrayValue)
-				{
-					CPPType = RigVMTypeUtils::GetWildCardArrayCPPType();
-				}
-				else if(ArrayType == FRigVMTemplateArgument::EArrayType_Mixed)
-				{				
-					CPPType = Pin->IsArray() ? RigVMTypeUtils::GetWildCardArrayCPPType() : RigVMTypeUtils::GetWildCardCPPType();
-				}
-
-				// execute pins are no longer part of the template, avoid changing the type
-				if(Argument == nullptr && Pin->IsExecuteContext())
-				{
-					CPPType = Pin->GetCPPType();
-					CPPObjectType = Pin->GetCPPTypeObject();
-				}
-			
-				if (Pin->GetCPPType() != CPPType || Pin->GetCPPTypeObject() != CPPObjectType)
-				{
-					bAnyTypeChanged = !Pin->IsExecuteContext();
-					if(bAnyTypeChanged)
-					{
-						ReportPinTypeChange(Pin, CPPType);
-					}
-					ChangePinType(Pin, CPPType, CPPObjectType, bSetupUndoRedo, false, false, false, bInitializeDefaultValue);
-				}
-			}
-			else 
-			{
-				// Resolve
-				if (Pin->GetTypeIndex() != FinalPinTypes[PinIndex])
-				{
-					bAnyTypeChanged = !Pin->IsExecuteContext();
-					if(bAnyTypeChanged)
-					{
-						const FString CPPType = FRigVMRegistry::Get().GetType(FinalPinTypes[PinIndex]).CPPType.ToString();
-						ReportPinTypeChange(Pin, CPPType);
-					}
-					ChangePinType(Pin, FinalPinTypes[PinIndex], bSetupUndoRedo, false, false, false, bInitializeDefaultValue);
-				}
-
-				// Once an pin is resolved, add it as a preferred type
-				URigVMTemplateNode* TemplateNode = CastChecked<URigVMTemplateNode>(Pin->GetNode());
-				TRigVMTypeIndex TypeIndex = TemplateNode->GetPreferredType(Pin->GetFName());
-				if (TypeIndex != FinalPinTypes[PinIndex] && TypeIndex != INDEX_NONE)
-				{
-					RemovePreferredType(TemplateNode, Pin->GetFName(), bSetupUndoRedo);
-					TypeIndex = INDEX_NONE;
-				}
-				if (TypeIndex == INDEX_NONE)
-				{
-					AddPreferredType(TemplateNode, Pin->GetFName(), FinalPinTypes[PinIndex], bSetupUndoRedo);
-				}
-			}
-		}
-	}
-
-	bool bHasWildcard=false;
-	if (bIsInterfaceNode)
-	{
-		if (EntryNode)
-		{
-			bHasWildcard = EntryNode->HasWildCardPin();
-		}
-		if (ReturnNode)
-		{
-			bHasWildcard |= ReturnNode->HasWildCardPin();
-		}
-	}
-	else
-	{
-		bHasWildcard = InNode->HasWildCardPin();
-	}
-	if (bHasWildcard)
-	{
-		if (bIsInterfaceNode)
-		{
-			if (EntryNode)
-			{
-				EntryNode->ResolvedPermutation = INDEX_NONE;
-			}
-			if (ReturnNode)
-			{
-				ReturnNode->ResolvedPermutation = INDEX_NONE;
-			}
-		}
-		else
-		{
-			InNode->ResolvedPermutation = INDEX_NONE;
-		}
-	}
-	else if (!InNode->PreferredPermutationPairs.IsEmpty())
-	{
-		if (bIsInterfaceNode)
-		{
-			PreferredPermutationPairs.Reset();
-			if (EntryNode)
-			{
-				PreferredPermutationPairs = EntryNode->PreferredPermutationPairs;
-			}
-			if (ReturnNode)
-			{
-				for (const FRigVMTemplatePreferredType& Preferred : ReturnNode->PreferredPermutationPairs)
-				{
-					PreferredPermutationPairs.AddUnique(Preferred);				
-				}
-			}
-			
-			PreferredPermutations = InNode->FindPermutationsForTypes(PreferredPermutationPairs, false);
-			ensureMsgf(PreferredPermutations.Num() == 1, TEXT("Number of preferred permutations is not 1 in %s"), *GetPackage()->GetPathName());
-			if (EntryNode)
-			{
-				EntryNode->ResolvedPermutation = PreferredPermutations[0];
-			}
-			if (ReturnNode)
-			{
-				EntryNode->ResolvedPermutation = PreferredPermutations[0];
-			}
-		}
-		else
-		{
-			PreferredPermutations = InNode->FindPermutationsForTypes(InNode->PreferredPermutationPairs, false);
-			ensureMsgf(PreferredPermutations.Num() == 1, TEXT("Number of preferred permutations is not 1 in %s"), *GetPackage()->GetPathName());
-			InNode->ResolvedPermutation = PreferredPermutations[0];
-		}			
-	}
-	
-	if (URigVMUnitNode* UnitNode = Cast<URigVMUnitNode>(InNode))
-	{
-		if (const FRigVMFunction* Function = UnitNode->GetResolvedFunction())
-		{
-			UnitNode->ResolvedFunctionName = Function->GetName();
-		}
-	}
-	
-	return bAnyTypeChanged;
-}
-
-bool URigVMController::PropagateTemplateFilteredTypes(URigVMTemplateNode* InNode, bool bSetupUndoRedo) 
-{
-	auto UpdateAndPropagate = [&](URigVMPin* Pin)
-	{
-		const TArray<URigVMLink*>& Links = Pin->GetLinks();
-		for (URigVMLink* Link : Links)
-		{
-			URigVMPin* OtherPin = Link->GetSourcePin() == Pin ? Link->GetTargetPin() : Link->GetSourcePin();
-			bool bPropagate = false;
-			bool bIsTemplate = false;
-			if (URigVMTemplateNode* OtherTemplate = Cast<URigVMTemplateNode>(OtherPin->GetNode()))
-			{
-				if (OtherTemplate->IsA<URigVMFunctionEntryNode>() || OtherTemplate->IsA<URigVMFunctionReturnNode>())
-				{
-					const URigVMPin* OtherPinRoot = OtherPin->GetRootPin();
-					if (!OtherTemplate->GetTemplate()->FindArgument(OtherPinRoot->GetFName()))
-					{
-						if (!AddArgumentForPin(OtherPin, Pin, bSetupUndoRedo))
-						{
-							continue;
-						}
-					}
-				}				
-				
-				if (!OtherTemplate->IsSingleton())
-				{
-					bIsTemplate = true;
-					if (OtherTemplate->PinNeedsFilteredTypesUpdate(OtherPin, Pin))
-					{
-						if (UpdateFilteredPermutations(OtherPin, Pin, bSetupUndoRedo))
-						{
-							UpdateTemplateNodePinTypes(OtherTemplate, bSetupUndoRedo);
-							if (!PropagateTemplateFilteredTypes(OtherTemplate, bSetupUndoRedo))
-							{
-								return false;
-							}
-						}
-						else
-						{
-							ensureMsgf(!ActionStack->BracketActions.IsEmpty(), TEXT("Unexpected link broken %s in package %s"), *Link->GetPinPathRepresentation(), *GetPackage()->GetPathName());
-							BreakLink(Link->GetSourcePin(), Link->GetTargetPin(), bSetupUndoRedo);
-							return false;
-						}
-					}
-				}								
-			}
-
-			if (!bIsTemplate)
-			{
-				if (!InNode->FilteredSupportsType(Pin, OtherPin->GetTypeIndex()))
-				{
-					ensureMsgf(!ActionStack->BracketActions.IsEmpty(), TEXT("Unexpected link broken %s in package %s"), *Link->GetPinPathRepresentation(), *GetPackage()->GetPathName());
-					BreakLink(Link->GetSourcePin(), Link->GetTargetPin(), bSetupUndoRedo);
-					return false;
-				}
-			}
-		}
-		return true;
-	};
-	
-	for(URigVMPin* Pin : InNode->GetPins())
-	{
-		if (!InNode->GetTemplate()->FindArgument(Pin->GetFName()))
-		{
-			continue;
-		}
-		
-		if (!UpdateAndPropagate(Pin))
-		{
-			return false;
-		}
-
-		if (Pin->IsArray())
-		{
-			if (Pin->GetSubPins().Num() > 0)
-			{
-				for (URigVMPin* SubPin : Pin->GetSubPins())
-				{
-					if (!UpdateAndPropagate(SubPin))
-					{
-						return false;
-					}
-				}
-			}
-		}
-	}
-	return true;
-}
-
-bool URigVMController::AddPreferredType(URigVMTemplateNode* InNode, const FName& InPinName, const TRigVMTypeIndex& InPreferredTypeIndex, bool bSetupUndoRedo)
-{
-	check(InNode->FindPin(InPinName.ToString()));
-	check(!InNode->IsSingleton());
-
-	FRigVMTemplatePreferredType* PreferredType = InNode->PreferredPermutationPairs.FindByPredicate([&](const FRigVMTemplatePreferredType& Other)
-	{
-		return Other.Argument == InPinName;
-	});
-
-	if (PreferredType)
-	{
-		return PreferredType->TypeIndex == InPreferredTypeIndex;		
-	}
-
-	// We need to remove or modify the existing preferred types that are not compatible with the new preferred type
-	// For example, a template which has 2 pins (A, B), and both can be float, double or FVector (always A type == B type)
-	// If a preferred type exists in which B is double, and we are adding A as float
-	// we need to change the B preferred type to float
-	// If the old preferred type is not compatible, we need to remove the B preferred type
-	const FRigVMTemplate* Template = InNode->GetTemplate();
-	TArray<int32> Permutations;
-	FRigVMTemplate::FTypeMap TypeMap;
-	TypeMap.Add(InPinName, InPreferredTypeIndex);
-	Template->Resolve(TypeMap, Permutations, false);
-
-	const FRigVMRegistry& Registry = FRigVMRegistry::Get();
-	
-	TArray<FRigVMTemplatePreferredType> NewPreferredTypes = InNode->PreferredPermutationPairs;
-	TArray<int32> ToRemove;
-	for (int32 PreferredIndex=0; PreferredIndex<InNode->PreferredPermutationPairs.Num(); ++PreferredIndex)
-	{
-		FRigVMTemplatePreferredType& Preferred = InNode->PreferredPermutationPairs[PreferredIndex];
-		TArray<TRigVMTypeIndex> PinTypes;
-		if (const FRigVMTemplateArgument* Argument = Template->FindArgument(Preferred.Argument))
-		{
-			bool bFoundMatch = Permutations.ContainsByPredicate([Argument, Preferred](const int32& Permutation)
-			{
-				if (Argument->TypeIndices[Permutation] == Preferred.TypeIndex)
-				{
-					return true;
-				}
-				return false;
-			});
-			if (!bFoundMatch)
-			{
-				for (const int32& Permutation : Permutations)
-				{
-					if (Registry.CanMatchTypes(Argument->TypeIndices[Permutation], Preferred.TypeIndex, true))
-					{
-						NewPreferredTypes[PreferredIndex].TypeIndex = Argument->TypeIndices[Permutation];
-#if UE_RIGVM_DEBUG_TYPEINDEX
-						NewPreferredTypes[PreferredIndex].UpdateStringFromIndex();
-#endif
-						bFoundMatch = true;
-						break;
-					}
-				}
-			}
-			if (!bFoundMatch)
-			{
-				ToRemove.Add(PreferredIndex);
-			}
-		}
-	}
-	for (int32 i=ToRemove.Num()-1; i>=0; --i)
-	{
-		NewPreferredTypes.RemoveAt(ToRemove[i]);
-	}
-
-	
-	NewPreferredTypes.Add(FRigVMTemplatePreferredType(InPinName, InPreferredTypeIndex));
-	if (bSetupUndoRedo)
-	{
-		ActionStack->AddAction(FRigVMSetPreferredTemplatePermutationsAction(InNode, NewPreferredTypes));
-	}
-	InNode->PreferredPermutationPairs = NewPreferredTypes;
-	return true;
-}
-
-bool URigVMController::RemovePreferredType(URigVMTemplateNode* InNode, const FName& InPinName, bool bSetupUndoRedo)
-{
-	int32 FoundIndex = INDEX_NONE;
-	for (int32 i=0; i<InNode->PreferredPermutationPairs.Num(); ++i)
-	{
-		if (InNode->PreferredPermutationPairs[i].Argument == InPinName)
-		{
-			FoundIndex = i;
-			break;
-		}
-	}
-	
-	if (FoundIndex == INDEX_NONE)
-	{
-		return false;
-	}
-
-	TArray<FRigVMTemplatePreferredType> NewPreferredTypes = InNode->PreferredPermutationPairs;
-	NewPreferredTypes.RemoveAt(FoundIndex);	
-	
-	if (bSetupUndoRedo)
-	{
-		ActionStack->AddAction(FRigVMSetPreferredTemplatePermutationsAction(InNode, NewPreferredTypes));
-	}
-	InNode->PreferredPermutationPairs = NewPreferredTypes;
-	return true;
-}
-
-bool URigVMController::ShouldPinOwnArgument(URigVMPin* InPin)
-{
-	if (!IsValidGraph())
-	{
-		return false;
-	}
-
-	if (!bIsTransacting && !IsGraphEditable())
-	{
-		return false;
-	}
-
-	URigVMLibraryNode* LibraryNode = Cast<URigVMLibraryNode>(InPin->GetNode());
-	if (LibraryNode == nullptr)
-	{
-		return false;
-	}
-
-	URigVMGraph* ContainedGraph = LibraryNode->GetContainedGraph();
-	check(ContainedGraph);
-
-	if (ContainedGraph->IsRootGraph())
-	{
-		return false;
-	}
-
-	URigVMFunctionEntryNode* EntryNode = ContainedGraph->GetEntryNode();
-	URigVMFunctionReturnNode* ReturnNode = ContainedGraph->GetReturnNode();
-
-	auto IsPinUsingArgument = [](URigVMPin* TestPin)
-	{
-		// If any subpins that are struct members are connected, then it should own an argument
-		TArray<URigVMPin*> ToProcessPins = TestPin->GetSubPins();
-		for (int32 i=0; i<ToProcessPins.Num(); ++i)
-		{
-			URigVMPin* CurPin = ToProcessPins[i];
-			if (CurPin->IsStructMember() && CurPin->IsLinked())
-			{
-				return true;
-			}
-			ToProcessPins.Append(CurPin->GetSubPins());
-		}
-		
-		for (int32 i=0; i<2; ++i)
-		{
-			TArray<URigVMPin*> ConnectedPins = (i==0) ? TestPin->GetLinkedSourcePins() : TestPin->GetLinkedTargetPins();
-			for (URigVMPin* ConnectedPin : ConnectedPins)
-			{
-				if (URigVMTemplateNode* TemplateNode = Cast<URigVMTemplateNode>(ConnectedPin->GetNode()))
-				{
-					if (TemplateNode->GetFilteredPermutationsIndices().Num() < 10)
-					{
-						return true;
-					}
-				}
-				else
-				{					
-					return true;
-				}
-			}								
-		}
-		return false;
-	};
-
-	bool bIsInputPin = false;
-	for (int32 i=0; i<2; ++i)
-	{
-		if (URigVMTemplateNode* InterfaceNode = (i == 0) ? (URigVMTemplateNode*)EntryNode : (URigVMTemplateNode*)ReturnNode)
-		{
-			if (URigVMPin* InterfacePin = InterfaceNode->FindPin(InPin->GetName()))
-			{
-				for (const FRigVMTemplatePreferredType& PreferredType : InterfaceNode->PreferredPermutationPairs)
-				{
-					if (PreferredType.Argument == InPin->GetFName())
-					{
-						return true;
-					}
-				}
-				if (IsPinUsingArgument(InterfacePin))
-				{
-					return true;
-				}
-				if (i == 0)
-				{
-					bIsInputPin = true;
-				}
-			}
-		}
-	}
-
-	// Find input variable nodes
-	for (URigVMNode* Node : ContainedGraph->GetNodes())
-	{
-		if (URigVMVariableNode* VariableNode = Cast<URigVMVariableNode>(Node))
-		{
-			if (VariableNode->GetVariableName() == InPin->GetFName())
-			{
-				if (IsPinUsingArgument(VariableNode->GetValuePin()))
-				{
-					return true;
-				}
-			}
-		}
-	}
-
-	return false;
-}
-
-bool URigVMController::UpdateLibraryTemplate(URigVMLibraryNode* LibraryNode, bool bSetupUndoRedo)
-{
-	if (!IsValidGraph())
-	{
-		return false;
-	}
-
-	if (!bIsTransacting && !IsGraphEditable())
-	{
-		return false;
-	}
-
-	if (!LibraryNode)
-	{
-		return false;
-	}
-
-	URigVMGraph* ContainedGraph = LibraryNode->GetContainedGraph();
-	check(ContainedGraph);
-
-	if (ContainedGraph->IsRootGraph())
-	{
-		return false;
-	}
-
-	URigVMTemplateNode* EntryNode = ContainedGraph->GetEntryNode();
-	URigVMTemplateNode* ReturnNode = ContainedGraph->GetReturnNode();
-	const TArray<URigVMTemplateNode*> InterfaceNodes = {EntryNode, ReturnNode};
-	const FRigVMTemplate* CurTemplate = LibraryNode->GetTemplate();
-
-	// Remove unused arguments from template
-	TArray<FName> ArgumentsToRemove;
-	{
-		for (const FRigVMTemplateArgument& CurArgument : CurTemplate->Arguments)
-		{
-			URigVMPin* LibraryPin = LibraryNode->FindPin(CurArgument.Name.ToString());
-			if (!LibraryPin || !ShouldPinOwnArgument(LibraryPin))
-			{
-				ArgumentsToRemove.Add(CurArgument.Name);
-			}
-		}
-	}
-
-	// If interface node contains a preferred type that differs from the type in the template, keep the interface type
-	TMap<FName, TRigVMTypeIndex> ArgPreferredType;
-	for (URigVMTemplateNode* InterfaceNode : InterfaceNodes)
-	{
-		if (InterfaceNode)
-		{
-			for (int32 i=0; i<InterfaceNode->PreferredPermutationPairs.Num(); ++i)
-			{
-				const FRigVMTemplatePreferredType& PreferredType = InterfaceNode->PreferredPermutationPairs[i];
-				if (!ArgumentsToRemove.Contains(PreferredType.Argument) && !ArgPreferredType.Contains(PreferredType.Argument))
-				{
-					ArgPreferredType.Add(PreferredType.Argument, PreferredType.GetTypeIndex());
-				}			
-			}
-		}
-	}
-
-	FRigVMTemplate NewTemplate;
-	NewTemplate.Index = CurTemplate->Index;
-	
-	NewTemplate.Arguments.Reserve(CurTemplate->Arguments.Num());
-	int32 ArgIndex = 0;
-	for (const FRigVMTemplateArgument& CurArgument : CurTemplate->Arguments)
-	{
-		if (ArgumentsToRemove.Contains(CurArgument.Name))
-		{
-			continue;
-		}
-		
-		FRigVMTemplateArgument NewArgument;
-		NewArgument.Index = ArgIndex++;
-		NewArgument.Name = CurArgument.Name;
-		NewArgument.Direction = CurArgument.Direction;
-		NewTemplate.Arguments.Add(NewArgument);
-	}
-
-	// The filtered permutations of the entry and return nodes reference the library template permutations
-	// so they can be combined
-	TArray<int32> CombinedFilteredPermutations;
-	{
-		if (EntryNode)
-		{
-			CombinedFilteredPermutations = EntryNode->FilteredPermutations;
-		}
-		if (ReturnNode)
-		{
-			CombinedFilteredPermutations = CombinedFilteredPermutations.FilterByPredicate([&](int32 Index)
-			{
-				return ReturnNode->FilteredPermutations.Contains(Index);
-			});			
-		}
-	}
-
-	TArray<FString> AllPermutationTypes;
-	for (int32 PermutationIndex=0; PermutationIndex<CombinedFilteredPermutations.Num(); ++PermutationIndex)
-	{
-		int32 OldIndex = CombinedFilteredPermutations[PermutationIndex];
-
-		// For each permutation, create a string that identifies all arguments with their types
-		FString PermutationTypesStr;
-		TMap<const FName, TRigVMTypeIndex> TypesToAdd;
-		for (const FRigVMTemplateArgument& CurArgument : CurTemplate->Arguments)
-		{
-			if (ArgumentsToRemove.Contains(CurArgument.Name))
-			{
-				continue;
-			}
-
-			if (TRigVMTypeIndex* TypeIndex = ArgPreferredType.Find(CurArgument.Name))
-			{
-				TypesToAdd.Add(CurArgument.Name, *TypeIndex);
-				const int32 TypeInt = *TypeIndex;				
-				PermutationTypesStr += FString::Printf(TEXT("%s:%d;"), *CurArgument.Name.ToString(), TypeInt);	
-			}
-			else
-			{
-				TypesToAdd.Add(CurArgument.Name, CurArgument.TypeIndices[OldIndex]);
-				const int32 TypeInt = CurArgument.TypeIndices[OldIndex];
-				PermutationTypesStr += FString::Printf(TEXT("%s:%d;"), *CurArgument.Name.ToString(), TypeInt);
-			}
-		}
-		if (AllPermutationTypes.Contains(PermutationTypesStr))
-		{
-			continue;
-		}
+}
 
-		AllPermutationTypes.Add(PermutationTypesStr);
+bool URigVMController::UpdateTemplateNodePinTypes(URigVMTemplateNode* InNode, bool bSetupUndoRedo, bool bInitializeDefaultValue)
+{
+	URigVMGraph* Graph = GetGraph();
+	check(InNode->GetGraph() == Graph);
+	bool bAnyTypeChanged = false;
 
-		// If this is a new permutation, add it to the template 
-		int32 NewArgumentIndex=0;
-		for (TPair<const FName, TRigVMTypeIndex> TypePair : TypesToAdd)
-		{
-			FRigVMTemplateArgument& NewArgument = NewTemplate.Arguments[NewArgumentIndex++];
-			const TRigVMTypeIndex& TypeIndex = TypePair.Value; 
-			NewArgument.TypeIndices.Add(TypeIndex);
-			if (TArray<int32>* Permutations = NewArgument.TypeToPermutations.Find(TypeIndex))
-			{
-				Permutations->Add(PermutationIndex);
-			}
-			else
-			{
-				NewArgument.TypeToPermutations.Add(TypeIndex, {PermutationIndex});
-			}
-		}
-		NewTemplate.Permutations.Add(INDEX_NONE);
+	const FRigVMTemplate* Template = InNode->GetTemplate();
+	if(Template == nullptr)
+	{
+		return false;
+	}
+
+	if (InNode->IsA<URigVMFunctionEntryNode>() || InNode->IsA<URigVMFunctionReturnNode>())
+	{
+		return false;
 	}
 
-	// Add a permutation if non was created but we are about to add arguments
-	if (NewTemplate.Permutations.IsEmpty() && !ArgPreferredType.IsEmpty())
+	TArray<int32> ResolvedPermutations = InNode->GetResolvedPermutationIndices(true);
+	InNode->ResolvedPermutation = ResolvedPermutations.Num() == 1 ? ResolvedPermutations[0] : INDEX_NONE;
+
+	TArray<URigVMPin*> Pins = InNode->GetPins();
+	TArray<const FRigVMTemplateArgument*> Arguments;
+	Arguments.SetNumUninitialized(Pins.Num());
+	for(int32 PinIndex=0; PinIndex < Pins.Num(); ++PinIndex)
 	{
-		NewTemplate.Permutations.Add(INDEX_NONE);
+		URigVMPin* Pin = Pins[PinIndex];
+		Arguments[PinIndex] = Template->FindArgument(Pin->GetFName());
 	}
 
-	// Add any arguments with preferred types that are missing
-	for (TPair<FName, TRigVMTypeIndex>& ArgPair : ArgPreferredType)
+	// Remove invalid permutations
+	ResolvedPermutations.RemoveAll([&](const int32& Permutation)
 	{
-		bool bMissing = true;
-		for (FRigVMTemplateArgument& Argument : NewTemplate.Arguments)
+		for (const FRigVMTemplateArgument* Argument : Arguments)
 		{
-			if (Argument.Name == ArgPair.Key)
+			if (Argument && Argument->TypeIndices[Permutation] == INDEX_NONE)
 			{
-				bMissing = false;
-				break;
+				return true;
 			}
 		}
+		return false;
+	});	
 
-		if (bMissing)
+	// Find the types for each possible permutation
+	TMap<int32, TArray<TRigVMTypeIndex>> PinTypes; // permutation to pin types
+	for (int32 ResolvedPermutation : ResolvedPermutations)
+	{
+		for(int32 PinIndex=0; PinIndex < Pins.Num(); ++PinIndex)
 		{
-			FRigVMTemplateArgument NewArgument;
-			NewArgument.Index = NewTemplate.Arguments.Num();
-			NewArgument.Name = ArgPair.Key;
-			NewArgument.Direction = LibraryNode->FindPin(ArgPair.Key.ToString())->GetDirection();
+			URigVMPin* Pin = Pins[PinIndex];
+			if (Pin->GetDirection() == ERigVMPinDirection::Hidden)
+			{
+				PinTypes.FindOrAdd(ResolvedPermutation).Add(INDEX_NONE);
+				continue;
+			}
 
-			TArray<int32> Permutations;
-			NewArgument.TypeIndices.Reserve(NewTemplate.NumPermutations());
-			Permutations.Reserve(NewTemplate.NumPermutations());
-			for (int32 i=0; i<NewTemplate.NumPermutations(); ++i)
+			bool bAddedType = false;
+			TArray<TRigVMTypeIndex>& Types = PinTypes.FindOrAdd(ResolvedPermutation);
+			if (const FRigVMTemplateArgument* Argument = Arguments[PinIndex])
+			{
+				Types.Add(Argument->TypeIndices[ResolvedPermutation]);
+				bAddedType = true;
+			}
+			
+			if (!bAddedType)
 			{
-				NewArgument.TypeIndices.Add(ArgPair.Value);
-				Permutations.Add(i);
+				// This is a pin with no argument
+				// Its marked as invalid, and will show wildcard type
+				Types.Add(INDEX_NONE);
 			}
-			NewArgument.TypeToPermutations.Add(ArgPair.Value, Permutations);
-			NewTemplate.Arguments.Add(NewArgument);			
 		}
 	}
 
-	NewTemplate.ComputeNotationFromArguments(LibraryNode->GetName());
+	// Some pins can benefit from reducing types to a single option
+	// Even if some pins reduce to a single type (maybe represented by a single permutation), the rest should still display a wildcard pin
+	// If reduction happens for multiple pins, we need to make sure that they reduce to the same permutation
+	// If possible, we want to respect the current pin type if its not wildcard
+	TArray<bool> WasReduced;
+	WasReduced.SetNumZeroed(Pins.Num());
 
-	// Set the new template
-	// The filtered permutations of the collapsed node are not valid anymore, lets reset them
+	TMap<int32, TArray<TRigVMTypeIndex>> ReducedTypes = PinTypes;
+	TArray<TRigVMTypeIndex> FinalPinTypes;
+	FinalPinTypes.SetNumUninitialized(Pins.Num());
+	
+	// Reduce compatible types, try to find the permutation in the reduced types
+	for(int32 PinIndex=0; PinIndex < Pins.Num(); ++PinIndex)
 	{
-		TArray<int32> OldPermutations = LibraryNode->FilteredPermutations;
-		LibraryNode->FilteredPermutations.Reset();
-		if (bSetupUndoRedo)
-		{
-			FRigVMControllerGraphGuard GraphGuard(this, LibraryNode->GetTypedOuter<URigVMGraph>(), bSetupUndoRedo);
-			ActionStack->AddAction(FRigVMSetTemplateFilteredPermutationsAction(LibraryNode, OldPermutations));
-			ActionStack->AddAction(FRigVMSetLibraryTemplateAction(LibraryNode, NewTemplate));		
-		}
-
-		bool bTemplateChanged = false;
-		if(!bTemplateChanged)
+		URigVMPin* Pin = Pins[PinIndex];
+		if (Pin->GetDirection() == ERigVMPinDirection::Hidden)
 		{
-			bTemplateChanged = CurTemplate->GetNotation() != NewTemplate.GetNotation();
+			continue;
 		}
-		if(!bTemplateChanged)
+		TArray<TRigVMTypeIndex> Types;
+		Types.Reserve(ResolvedPermutations.Num());
+		TRigVMTypeIndex PreferredType = INDEX_NONE;
+		int32 TypesFoundInReduced = 0;
+		for (int32 ResolvedPermutation : ResolvedPermutations)
 		{
-			check(CurTemplate->NumArguments() == NewTemplate.NumArguments());
-			for(int32 ArgumentIndex = 0; (ArgumentIndex < NewTemplate.NumArguments()) && (!bTemplateChanged); ArgumentIndex++)
+			Types.Add(PinTypes.FindChecked(ResolvedPermutation)[PinIndex]);
+			if (ReducedTypes.Contains(ResolvedPermutation))
 			{
-				const FRigVMTemplateArgument* ArgA = CurTemplate->GetArgument(ArgumentIndex);
-				const FRigVMTemplateArgument* ArgB = NewTemplate.GetArgument(ArgumentIndex);
-
-				const TArray<TRigVMTypeIndex>& TypesA = ArgA->GetTypeIndices();
-				const TArray<TRigVMTypeIndex>& TypesB = ArgB->GetTypeIndices();
-				bTemplateChanged = TypesA.Num() != TypesB.Num();
-				if(!bTemplateChanged)
+				TypesFoundInReduced++;
+				if (PreferredType == INDEX_NONE)
 				{
-					for(int32 TypeIndex=0; (TypeIndex<TypesA.Num()) && (!bTemplateChanged);TypeIndex++)
-					{
-						bTemplateChanged = TypesA[TypeIndex] != TypesB[TypeIndex];
-					}
+					PreferredType = PinTypes.FindChecked(ResolvedPermutation)[PinIndex];
 				}
 			}
 		}
-		
-		LibraryNode->Template = NewTemplate;
-		Notify(ERigVMGraphNotifType::NodeDescriptionChanged, LibraryNode);
 
-		if(bTemplateChanged)
+		if (TypesFoundInReduced > 1)
 		{
-			FRigVMControllerGraphGuard GraphGuard(this, LibraryNode->GetTypedOuter<URigVMGraph>(), false);
-			Notify(ERigVMGraphNotifType::LibraryTemplateChanged, LibraryNode);
+			PreferredType = Pin->GetTypeIndex();
 		}
-	}
-
-	// Update Entry and Return filtered permutations to all permutations
-	{
-		FRigVMControllerGraphGuard GraphGuard(this, ContainedGraph, bSetupUndoRedo);
-		for (int32 EntryReturn=0; EntryReturn<2; ++EntryReturn)
+		FinalPinTypes[PinIndex] = InNode->TryReduceTypesToSingle(Types, PreferredType);
+		if (FinalPinTypes[PinIndex] != INDEX_NONE)
 		{
-			URigVMTemplateNode* Node = (EntryReturn == 0) ? EntryNode : ReturnNode;
-			if (Node)
-			{
-				TArray<int32> OldPermutations = Node->FilteredPermutations;
-				Node->FilteredPermutations.SetNumUninitialized(NewTemplate.NumPermutations());
-				for (int32 i=0; i<Node->FilteredPermutations.Num(); ++i)
-				{
-					Node->FilteredPermutations[i] = i;
-				}
-
-				if (bSetupUndoRedo)
-				{
-					ActionStack->AddAction(FRigVMSetTemplateFilteredPermutationsAction(Node, OldPermutations));
-				}
-			}
+			WasReduced[PinIndex] =  true;
 		}
+
+		// Remove reduced types which do not match this type
+		TArray<int32> PermutationToRemove;
+		for (TPair<int32, TArray<TRigVMTypeIndex>>& Pair : ReducedTypes)
 		{
-			bool bUpdatedTypes = false;
-			for (int32 EntryReturn=0; EntryReturn<2; ++EntryReturn)
+			if (Pair.Value[PinIndex] != FinalPinTypes[PinIndex])
 			{
-				URigVMTemplateNode* Node = (EntryReturn == 0) ? EntryNode : ReturnNode;
-				if (Node)
-				{
-					// Update template node pin types of an interface node will update both entry and return types.
-					if (!bUpdatedTypes)
-					{
-						UpdateTemplateNodePinTypes(Node, bSetupUndoRedo);
-						bUpdatedTypes = true;
-					}
-					Notify(ERigVMGraphNotifType::NodeDescriptionChanged, Node);
-				}
+				PermutationToRemove.Add(Pair.Key);
 			}
 		}
-	}
-
-	// Update outer graphs
-	if (!bIsTransacting && !bSuspendRecomputingOuterTemplateFilters)
-	{
-		FRigVMControllerGraphGuard GraphGuard(this, LibraryNode->GetGraph(), bSetupUndoRedo);
-		RecomputeAllTemplateFilteredPermutations(bSetupUndoRedo);
-	}
-
-	return true;
-}
-
-bool URigVMController::RecomputeAllTemplateFilteredPermutations(bool bSetupUndoRedo)
-{
-	if (!IsValidGraph())
-	{
-		return false;
-	}
-
-	if (!bIsTransacting && !IsGraphEditable())
-	{
-		return false;
-	}
-
-	URigVMGraph* Graph = GetGraph();
-	check(Graph);
-
-	// Save all template pin types, in case we need them after initializing
-	TMap<URigVMPin*, TRigVMTypeIndex> TypesBeforeRecomputing;
-	for (URigVMNode* Node : Graph->GetNodes())
-	{
-		if (URigVMTemplateNode* TemplateNode = Cast<URigVMTemplateNode>(Node))
+		for (int32& ToRemove : PermutationToRemove)
 		{
-			if (TemplateNode->IsSingleton())
-			{
-				continue;
-			}
-
-			if (const FRigVMTemplate* Template = TemplateNode->GetTemplate())
-			{
-				for (int32 i=0; i<Template->NumArguments(); ++i)
-				{
-					const FRigVMTemplateArgument* Argument = Template->GetArgument(i);
-					URigVMPin* Pin = TemplateNode->FindPin(Argument->GetName().ToString());
-					if (!Pin || Pin->IsWildCard())
-					{
-						continue;
-					}
-
-					TypesBeforeRecomputing.Add(Pin, Pin->GetTypeIndex());
-				}
-			}
+			ReducedTypes.Remove(ToRemove);
 		}
 	}
-
-	bool bAnyTypeChanged = false;
-	bool bAnyLinkBroken = false;
+	
+	// First unresolve any pins that need unresolving, then resolve everything else
+	for (int32 UnresolveResolve=0; UnresolveResolve<2; ++UnresolveResolve)
 	{
-		TGuardValue<bool> GuardSuspendRecomputeOuterFilters(bSuspendRecomputingOuterTemplateFilters, true);
-		
-		// Initialize all filtered permutations to unresolved (except for preferred permutation)
-		// This will also remove all arguments from library templates
-		InitializeAllTemplateFiltersInGraph(bSetupUndoRedo, false);	
-
-		// Apply all links to update filtered permutations
-		TArray<URigVMLink*> SortedLinks = Graph->GetLinks().FilterByPredicate([this](URigVMLink* Link)
-		{
-			// Filter out detached links
-			if(!Link->GetSourcePin()->IsLinkedTo(Link->GetTargetPin()))
-			{
-				return false;
-			}
-
-			// Filter out links on execute pins
-			if(Link->GetSourcePin()->IsExecuteContext() || Link->GetTargetPin()->IsExecuteContext())
-			{
-				return false;
-			}
-			
-			return true;
-		});
-
-		// Solve for unit nodes first, other templates are more expensive to solve. This way we are avoiding solving
-		// links like (reroute-reroute) before their filtered permutations are reduced
-		Algo::Sort(SortedLinks, [](const URigVMLink* A, const URigVMLink* B)
+		for(int32 PinIndex=0; PinIndex < Pins.Num(); ++PinIndex)
 		{
-			const URigVMPin* APinSource = A->GetSourcePin();
-			const URigVMPin* APinTarget = A->GetTargetPin();
-			const URigVMPin* BPinSource = B->GetSourcePin();
-			const URigVMPin* BPinTarget = B->GetTargetPin();
-			if (!IsValid(APinSource) || !IsValid(APinTarget) || !IsValid(BPinSource) || !IsValid(BPinTarget))
-			{
-				return false;
-			}
-	
-			const bool bASourceIsUnitNode = APinSource->GetNode()->IsA<URigVMUnitNode>();
-			const bool bATargetIsUnitNode = APinTarget->GetNode()->IsA<URigVMUnitNode>();
-			const bool bBSourceIsUnitNode = BPinSource->GetNode()->IsA<URigVMUnitNode>();
-			const bool bBTargetIsUnitNode = BPinTarget->GetNode()->IsA<URigVMUnitNode>();
-		
-			if (bASourceIsUnitNode && bATargetIsUnitNode && (!bBSourceIsUnitNode || !bBTargetIsUnitNode))
-			{
-				return true;
-			}
-
-			if ((bASourceIsUnitNode || bATargetIsUnitNode) && (!bBSourceIsUnitNode && !bBTargetIsUnitNode))
-			{
-				return true;
-			}
-
-			return false;		
-		});
-
-		// Keep a list of removed links so that we can only notify adding links when necessary
-		TArray<TPair<FString, FString>> LinksRemoved;
-		FDelegateHandle Delegate = OnModified().AddLambda([&LinksRemoved](ERigVMGraphNotifType NotifType, URigVMGraph* Graph, UObject* Subject) {
-			if (NotifType == ERigVMGraphNotifType::LinkAdded)
-			{
-				if (URigVMLink* Link = Cast<URigVMLink>(Subject))
-				{
-					LinksRemoved.Remove(TPair<FString, FString>(Link->GetSourcePin()->GetPinPath(), Link->GetTargetPin()->GetPinPath()));
-				}
-			}
-			else if (NotifType == ERigVMGraphNotifType::LinkRemoved)
+			URigVMPin* Pin = Pins[PinIndex];
+			if (UnresolveResolve == 0 && Pin->IsWildCard())
 			{
-				if (URigVMLink* Link = Cast<URigVMLink>(Subject))
-				{
-					LinksRemoved.Add(TPair<FString, FString>(Link->GetSourcePin()->GetPinPath(), Link->GetTargetPin()->GetPinPath()));
-				}
+				continue;
 			}
-		});
 			
-		{
-			TGuardValue<bool> SuspendNotifications(bSuspendNotifications, true);
-			DetachLinksFromPinObjects(&SortedLinks);
-		}
-		TArray<URigVMLink*> LinksToRemove;
-		for (int32 i=0; i<SortedLinks.Num(); ++i)
-		{
-			URigVMLink* Link = SortedLinks[i];
-		
-			URigVMPin* OutputPin = Link->GetSourcePin();
-			URigVMPin* InputPin = Link->GetTargetPin();
-
-			if (!IsValid(OutputPin) || !IsValid(InputPin))
+			if (Pin->GetDirection() == ERigVMPinDirection::Hidden)
 			{
 				continue;
 			}
 
+			bool bShouldUnresolve = FinalPinTypes[PinIndex] == INDEX_NONE;
+			// If we are about to resolve the pin to a different type, first unresolve it
+			if (!bShouldUnresolve && UnresolveResolve == 0 && Pin->GetTypeIndex() != FinalPinTypes[PinIndex])
 			{
-				ActionStack->OpenUndoBracket(TEXT("Prepare template for type"));
-				bool bLinkRemoved = false;
-				for (int32 EdgeI=0; EdgeI<2 && !bLinkRemoved; ++EdgeI)
-				{
-					URigVMPin* EdgePin = (EdgeI==0) ? OutputPin : InputPin;
-					if (EdgePin->IsStructMember())
-					{
-						if (URigVMTemplateNode* TemplateNode = Cast<URigVMTemplateNode>(EdgePin->GetNode()))
-						{
-							URigVMPin* RootPin = EdgePin->GetRootPin();
-							if (TRigVMTypeIndex* TypeIndex = TypesBeforeRecomputing.Find(RootPin))
-							{
-								if (TemplateNode->FilteredSupportsType(RootPin, *TypeIndex))
-								{
-									PrepareTemplatePinForType(RootPin, {*TypeIndex}, bSetupUndoRedo);
-								}
-								else
-								{
-									LinksToRemove.Add(Link);
-									bLinkRemoved = true;
-									break;
-								}
-							}
-						}			
-					}
-				}
-				if (bLinkRemoved)
-				{
-					ActionStack->CancelUndoBracket(this);
-					break;
-				}
-				else
-				{
-					ActionStack->CloseUndoBracket(this);
-				}
+				bShouldUnresolve = true;
 			}
 
-			// Propagate filtered permutations
-			if (PrepareToLink(OutputPin, InputPin, bSetupUndoRedo))
+			if (bShouldUnresolve)
 			{
-				OutputPin->Links.Add(Link);
-				InputPin->Links.Add(Link);
-
-				if (LinksRemoved.Contains(TPair<FString, FString>(Link->GetSourcePin()->GetPinPath(), Link->GetTargetPin()->GetPinPath())))
+				// Unresolve
+				if (Pin->HasInjectedNodes())
 				{
-					Notify(ERigVMGraphNotifType::LinkAdded, Link);
+					EjectNodeFromPin(Pin, bSetupUndoRedo);
 				}
-			}
-			else
-			{
-				LinksToRemove.Add(Link);
-			}
-		}
-
-		OnModified().Remove(Delegate);
-
-		bAnyLinkBroken = !LinksToRemove.IsEmpty();
-		{
-			TGuardValue<bool> RecomputeGuard(bSuspendRecomputingTemplateFilters, true);
-			for (URigVMLink* Link : LinksToRemove)
-			{
-				Link->SourcePin->Links.Add(Link);
-				Link->TargetPin->Links.Add(Link);
-
-				BreakLink(Link->SourcePin, Link->TargetPin, bSetupUndoRedo);
-			}
-		}
-	
-		// Now update all template nodes pin types
-		for (URigVMNode* Node : Graph->GetNodes())
-		{
-			if (URigVMTemplateNode* TemplateNode = Cast<URigVMTemplateNode>(Node))
-			{
-				if (TemplateNode->IsSingleton())
+				
+				const FRigVMTemplateArgument* Argument = Template->FindArgument(*Pin->GetName());
+				FRigVMTemplateArgument::EArrayType ArrayType;
+				if (Argument)
+				{
+					ArrayType = Argument->GetArrayType();
+				}
+				else
 				{
-					continue;
+					ArrayType = Pin->IsArray() ? FRigVMTemplateArgument::EArrayType::EArrayType_ArrayValue : FRigVMTemplateArgument::EArrayType::EArrayType_SingleValue;
 				}
+				
+				FString CPPType = RigVMTypeUtils::GetWildCardCPPType();
+				UObject* CPPObjectType = RigVMTypeUtils::GetWildCardCPPTypeObject();
 
-				if(UpdateTemplateNodePinTypes(TemplateNode, bSetupUndoRedo))
+				if (ArrayType == FRigVMTemplateArgument::EArrayType_ArrayValue)
 				{
-					bAnyTypeChanged = true;
+					CPPType = RigVMTypeUtils::GetWildCardArrayCPPType();
 				}
-			}
-		}
-	}
-
-	if (URigVMGraph* ParentGraph = Graph->GetParentGraph())
-	{
-		FRigVMControllerGraphGuard GraphGuard(this, ParentGraph, bSetupUndoRedo);
-		UpdateLibraryTemplate(Graph->GetTypedOuter<URigVMLibraryNode>(), bSetupUndoRedo);
-	}
-
-	
-	return bAnyTypeChanged || bAnyLinkBroken;
-}
-
-bool URigVMController::PrepareToLink(URigVMPin* FirstToResolve, URigVMPin* SecondToResolve, bool bSetupUndoRedo)
-{
-	
-	if (URigVMTemplateNode* FirstTemplateNode = Cast<URigVMTemplateNode>(FirstToResolve->GetNode()))
-	{
-		bool bShouldPrepare = true;
-		if (FirstTemplateNode->IsA<URigVMFunctionEntryNode>() || FirstTemplateNode->IsA<URigVMFunctionReturnNode>())
-		{
-			URigVMPin* RootPin = FirstToResolve->GetRootPin();
-			if (!RootPin->IsOrphanPin())
-			{
-				if (!FirstTemplateNode->GetTemplate()->FindArgument(RootPin->GetFName()))
+				else if(ArrayType == FRigVMTemplateArgument::EArrayType_Mixed)
 				{
-					ensure(RootPin == FirstToResolve);
-					bShouldPrepare = AddArgumentForPin(FirstToResolve, SecondToResolve, bSetupUndoRedo);
+					CPPType = Pin->IsArray() ? RigVMTypeUtils::GetWildCardArrayCPPType() : RigVMTypeUtils::GetWildCardCPPType();
 				}
-			}
-		}
-		if (!FirstTemplateNode->IsSingleton() && bShouldPrepare)
-		{
-			TArray<TRigVMTypeIndex> InputTypes = GetFilteredTypes(SecondToResolve);
-			if (InputTypes.Num() > 0)
-			{
-				if (!PrepareTemplatePinForType(FirstToResolve, InputTypes, bSetupUndoRedo))
+
+				// execute pins are no longer part of the template, avoid changing the type
+				if(Argument == nullptr && Pin->IsExecuteContext())
 				{
-					return false;
+					CPPType = Pin->GetCPPType();
+					CPPObjectType = Pin->GetCPPTypeObject();
 				}
-			}
-		}
-	}
-	if (URigVMTemplateNode* SecondTemplateNode = Cast<URigVMTemplateNode>(SecondToResolve->GetNode()))
-	{
-		bool bShouldPrepare = true;
-		if (SecondTemplateNode->IsA<URigVMFunctionEntryNode>() || SecondTemplateNode->IsA<URigVMFunctionReturnNode>())
-		{
-			URigVMPin* RootPin = SecondToResolve->GetRootPin();
-			if (!RootPin->IsOrphanPin())
-			{
-				if (!SecondTemplateNode->GetTemplate()->FindArgument(RootPin->GetFName()))
+
+				if (Pin->GetCPPType() != CPPType || Pin->GetCPPTypeObject() != CPPObjectType)
 				{
-					ensure(RootPin == SecondToResolve);
-					bShouldPrepare = AddArgumentForPin(SecondToResolve, FirstToResolve, bSetupUndoRedo);
+					bAnyTypeChanged = !Pin->IsExecuteContext();
+					if(bAnyTypeChanged)
+					{
+						ReportPinTypeChange(Pin, CPPType);
+					}
+					ChangePinType(Pin, CPPType, CPPObjectType, bSetupUndoRedo, false, false, false, bInitializeDefaultValue);
 				}
 			}
-		}
-		if (!SecondTemplateNode->IsSingleton() && bShouldPrepare)
-		{
-			TArray<TRigVMTypeIndex> OutTypes = GetFilteredTypes(FirstToResolve);
-			if (OutTypes.Num() > 0)
+			else 
 			{
-				if (!PrepareTemplatePinForType(SecondToResolve, OutTypes, bSetupUndoRedo))
+				// Resolve
+				if (Pin->GetTypeIndex() != FinalPinTypes[PinIndex])
 				{
-					return false;
+					bAnyTypeChanged = !Pin->IsExecuteContext();
+					if(bAnyTypeChanged)
+					{
+						const FString CPPType = FRigVMRegistry::Get().GetType(FinalPinTypes[PinIndex]).CPPType.ToString();
+						ReportPinTypeChange(Pin, CPPType);
+					}
+					ChangePinType(Pin, FinalPinTypes[PinIndex], bSetupUndoRedo, false, false, false, bInitializeDefaultValue);
 				}
 			}
 		}
 	}
 
-	return true;
-}
-
-void URigVMController::InitializeFilteredPermutationsFromTemplateTypes()
-{
-	if (!IsValidGraph())
+	bool bHasWildcard = InNode->HasWildCardPin();
+	
+	if (bHasWildcard)
 	{
-		return;
+		InNode->ResolvedPermutation = INDEX_NONE;
 	}
-
-	if (!bIsTransacting && !IsGraphEditable())
+	else
 	{
-		return;
+		// If we got into a resolved permutation through reducing types, we need to set the result in the node
+		ResolvedPermutations = InNode->GetResolvedPermutationIndices(false);
+		check(ResolvedPermutations.Num() == 1);
+		InNode->ResolvedPermutation = ResolvedPermutations[0];
 	}
-
-	URigVMGraph* Graph = GetGraph();
-	check(Graph);
-
-	for (URigVMNode* Node : Graph->GetNodes())
+	
+	if (URigVMUnitNode* UnitNode = Cast<URigVMUnitNode>(InNode))
 	{
-		if (URigVMLibraryNode* LibraryNode = Cast<URigVMLibraryNode>(Node))
-		{
-			bool bTemplateChanged = false;
-			FRigVMTemplate NewTemplate;
-			NewTemplate.Permutations.Add(INDEX_NONE);
-			
-			TArray<URigVMTemplateNode*> InterfaceNodes = {LibraryNode->GetEntryNode(), LibraryNode->GetReturnNode()};
-			for (URigVMTemplateNode* InterfaceNode : InterfaceNodes)
-			{
-				if (InterfaceNode)
-				{
-					InterfaceNode->PreferredPermutationPairs.Reset();
-					for (URigVMPin* Pin : InterfaceNode->GetPins())
-					{
-						if (!Pin->IsWildCard())
-						{
-							InterfaceNode->PreferredPermutationPairs.Add(FRigVMTemplatePreferredType(Pin->GetFName(), Pin->GetTypeIndex()));
-
-							if (!NewTemplate.FindArgument(Pin->GetFName()))
-							{
-								FRigVMTemplateArgument NewArgument;
-								NewArgument.Name = Pin->GetFName();
-								NewArgument.Index = NewTemplate.NumArguments();
-								NewArgument.Direction = LibraryNode->FindPin(Pin->GetName())->GetDirection();
-								NewArgument.TypeIndices.Add(Pin->GetTypeIndex());
-								NewArgument.TypeToPermutations.Add(Pin->GetTypeIndex(), {0});
-								NewTemplate.Arguments.Add(NewArgument);
-							}
-							
-							bTemplateChanged = true;
-						}
-					}
-					InterfaceNode->FilteredPermutations = {0};
-				}
-			}
-
-			if (bTemplateChanged)
-			{
-				NewTemplate.ComputeNotationFromArguments(LibraryNode->GetName());				
-				LibraryNode->Template = NewTemplate;
-				{
-					FRigVMControllerGraphGuard GraphGuard(this, LibraryNode->GetTypedOuter<URigVMGraph>(), false);
-					Notify(ERigVMGraphNotifType::LibraryTemplateChanged, LibraryNode);
-				}
-			}
-		}
-		
-		if (URigVMTemplateNode* TemplateNode = Cast<URigVMTemplateNode>(Node))
+		if (const FRigVMFunction* Function = UnitNode->GetResolvedFunction())
 		{
-			TemplateNode->InitializeFilteredPermutationsFromTypes(false);
+			UnitNode->ResolvedFunctionName = Function->GetName();
 		}
 	}
+	
+	return bAnyTypeChanged;
 }
 
-void URigVMController::InitializeAllTemplateFiltersInGraph(bool bSetupUndoRedo, bool bChangePinTypes)
+bool URigVMController::PrepareToLink(URigVMPin* FirstToResolve, URigVMPin* SecondToResolve, bool bSetupUndoRedo)
 {
-	if (!IsValidGraph())
-	{
-		return;
-	}
-
-	if (!bIsTransacting && !IsGraphEditable())
+	FRigVMRegistry& Registry = FRigVMRegistry::Get();
+	
+	// Check if there is anything to do
+	if (!FirstToResolve->IsWildCard() &&
+		!SecondToResolve->IsWildCard() &&
+		Registry.CanMatchTypes(FirstToResolve->GetTypeIndex(), SecondToResolve->GetTypeIndex(), true))
 	{
-		return;
+		return true;
 	}
-
-	URigVMGraph* Graph = GetGraph();
-	check(Graph);
-
-	// Remove all arguments (that do not have a preferred type) from the template
-	if (URigVMLibraryNode* LibraryNode = Graph->GetTypedOuter<URigVMLibraryNode>())
+	
+	// Find out the matching supported types
+	TArray<TRigVMTypeIndex> MatchingTypes;
 	{
-		if (const FRigVMTemplate* Template = LibraryNode->GetTemplate())
+		auto GetPinSupportedTypes = [&Registry](URigVMPin* Pin) -> TArray<TRigVMTypeIndex>
 		{
-			FRigVMTemplate NewTemplate;
-			NewTemplate.Permutations.Add(INDEX_NONE);
-
-			TArray<URigVMTemplateNode*> InterfaceNodes = {LibraryNode->GetEntryNode(), LibraryNode->GetReturnNode()};
-			for (const FRigVMTemplateArgument& Argument : Template->Arguments)
+			URigVMTemplateNode* TemplateNode = Cast<URigVMTemplateNode>(Pin->GetNode());
+			if (!TemplateNode || TemplateNode->IsSingleton() || !Pin->IsWildCard())
+			{
+				return {Pin->GetTypeIndex()};
+			}
+			else
 			{
-				// Find out if this argument has a preferred type
-				bool bHasPreferredType = false;
-				TRigVMTypeIndex PreferredTypeIndex;
-				for (URigVMTemplateNode* InterfaceNode : InterfaceNodes)
+				if (const FRigVMTemplate* Template = TemplateNode->GetTemplate())
 				{
-					if (InterfaceNode)
+					URigVMPin* RootPin = Pin;
+					uint8 ArrayLevels = 0;
+					while (RootPin->GetParentPin() != nullptr)
+					{
+						ArrayLevels++;
+						RootPin = RootPin->GetParentPin();
+					}
+					if (const FRigVMTemplateArgument* Argument = Template->FindArgument(RootPin->GetFName()))
 					{
-						for (int32 i=0; i<InterfaceNode->PreferredPermutationPairs.Num(); ++i)
+						TArray<int32> ResolvedPermutations = TemplateNode->GetResolvedPermutationIndices(true);
+						TArray<TRigVMTypeIndex> Types = Argument->GetSupportedTypeIndices(ResolvedPermutations);
+						for (int32 i=0; i<ArrayLevels; ++i)
 						{
-							FRigVMTemplatePreferredType& PreferredType = InterfaceNode->PreferredPermutationPairs[i];
-							if (PreferredType.Argument == Argument.Name)
+							for (TRigVMTypeIndex& Type : Types)
 							{
-								PreferredTypeIndex = PreferredType.GetTypeIndex();
-								bHasPreferredType = true;
-								break;
+								Type = Registry.GetBaseTypeFromArrayTypeIndex(Type);
 							}
 						}
-						if (bHasPreferredType)
-						{
-							break;
-						}
+						return Types;
 					}
 				}
-
-				if (bHasPreferredType)
-				{
-					FRigVMTemplateArgument NewArgument;
-					NewArgument.Name = Argument.Name;
-					NewArgument.Index = NewTemplate.NumArguments();
-					NewArgument.Direction = Argument.Direction;
-					NewArgument.TypeIndices.Add(PreferredTypeIndex);
-					NewArgument.TypeToPermutations.Add(PreferredTypeIndex, {0});
-					NewTemplate.Arguments.Add(NewArgument);
-				}
-			}
-
-			NewTemplate.ComputeNotationFromArguments(LibraryNode->GetName());
-			
-			if (bSetupUndoRedo)
-			{
-				FRigVMControllerGraphGuard GraphGuard(this, LibraryNode->GetTypedOuter<URigVMGraph>(), bSetupUndoRedo);
-				ActionStack->AddAction(FRigVMSetLibraryTemplateAction(LibraryNode, NewTemplate));
 			}
-			LibraryNode->Template = NewTemplate;
-			Notify(ERigVMGraphNotifType::NodeDescriptionChanged, LibraryNode);
-
-			// Update Entry and Return filtered permutations to all permutations
-			for (URigVMTemplateNode* InterfaceNode : InterfaceNodes)
+			return {};
+		};
+		MatchingTypes = GetPinSupportedTypes(FirstToResolve);
+		TArray<TRigVMTypeIndex> SecondTypes = GetPinSupportedTypes(SecondToResolve);
+		MatchingTypes.RemoveAll([&Registry, SecondTypes](const TRigVMTypeIndex& FirstType)
+		{
+			return !SecondTypes.ContainsByPredicate([&Registry, FirstType](const TRigVMTypeIndex& SecondType)
 			{
-				if (InterfaceNode)
-				{
-					TArray<int32> OldPermutations = InterfaceNode->FilteredPermutations;
-					InterfaceNode->FilteredPermutations.Init(INDEX_NONE, NewTemplate.NumPermutations());
-					
-					if (bSetupUndoRedo)
-					{
-						ActionStack->AddAction(FRigVMSetTemplateFilteredPermutationsAction(InterfaceNode, OldPermutations));
-					}
-
-					if (bChangePinTypes)
-					{
-						UpdateTemplateNodePinTypes(InterfaceNode, bSetupUndoRedo);
-					}
+				return Registry.CanMatchTypes(FirstType, SecondType, true);
+			});
+		});
+	}
 
-					Notify(ERigVMGraphNotifType::NodeDescriptionChanged, InterfaceNode);
-				}
-			}
+	if (MatchingTypes.IsEmpty())
+	{
+		return false;
+	}
 
-			// Update outer graphs
-			if (!bIsTransacting && !bSuspendRecomputingOuterTemplateFilters)
-			{
-				FRigVMControllerGraphGuard GraphGuard(this, LibraryNode->GetGraph(), bSetupUndoRedo);
-				RecomputeAllTemplateFilteredPermutations(bSetupUndoRedo);
-			}
+	TRigVMTypeIndex FinalType = INDEX_NONE;
+	if (MatchingTypes.Num() > 1)
+	{
+		// Query the user for the type to resolve the pin (from the list of MatchingTypes)
+		if (RequestPinTypeSelectionDelegate.IsBound())
+		{
+			FinalType = RequestPinTypeSelectionDelegate.Execute(MatchingTypes);
+		}
+		else
+		{
+			return false;
 		}
 	}
-
-	for (URigVMNode* Node : Graph->GetNodes())
+	else
+	{
+		FinalType = MatchingTypes[0];
+	}
+	
+	// If only one match, resolve both pins to that
+	bool bSuccess = true;
+	if (FinalType != INDEX_NONE)
 	{
-		if (URigVMTemplateNode* TemplateNode = Cast<URigVMTemplateNode>(Node))
+		if (FirstToResolve->IsWildCard())
 		{
-			if (!TemplateNode->IsSingleton())
-			{
-				// If the preferred permutation does not fit to the available types in the template, remove them
-				TArray<FRigVMTemplatePreferredType> NewPreferredTypes;
-				for (int32 i=0; i<TemplateNode->PreferredPermutationPairs.Num(); ++i)
-				{
-					const FRigVMTemplatePreferredType& PreferredType = TemplateNode->PreferredPermutationPairs[i];
-					
-					if (TemplateNode->GetTemplate()->FindArgument(PreferredType.Argument) && TemplateNode->GetTemplate()->ArgumentSupportsTypeIndex(PreferredType.Argument, PreferredType.GetTypeIndex()))
-					{
-						NewPreferredTypes.Add(PreferredType);
-					}
-				}
-
-				if (TemplateNode->PreferredPermutationPairs.Num() != NewPreferredTypes.Num())
-				{
-					if (bSetupUndoRedo)
-					{
-						FRigVMSetPreferredTemplatePermutationsAction PreferredTypesAction(TemplateNode, NewPreferredTypes);
-						ActionStack->AddAction(PreferredTypesAction);
-					}
-
-					TemplateNode->PreferredPermutationPairs = NewPreferredTypes;
-				}
-				
-				TArray<int32> OldPermutations = TemplateNode->FilteredPermutations;
-				TemplateNode->InitializeFilteredPermutations();
-				if (bChangePinTypes)
-				{
-					UpdateTemplateNodePinTypes(TemplateNode, bSetupUndoRedo);
-				}
-				if (bSetupUndoRedo)
-				{
-					FRigVMSetTemplateFilteredPermutationsAction FilteringAction(TemplateNode, OldPermutations);
-					ActionStack->AddAction(FilteringAction);
-				}
-			}
+			bSuccess &= ResolveWildCardPin(FirstToResolve, FinalType, bSetupUndoRedo);
+		}
+		if (SecondToResolve->IsWildCard())
+		{
+			bSuccess &= ResolveWildCardPin(SecondToResolve, FinalType, bSetupUndoRedo);
 		}
 	}
+	else
+	{
+		return false;
+	}
+
+	return bSuccess;
 }
 
 bool URigVMController::ChangePinType(const FString& InPinPath, const FString& InCPPType, const FName& InCPPTypeObjectPath, bool bSetupUndoRedo, bool bSetupOrphanPins, bool bBreakLinks, bool bRemoveSubPins, bool bInitializeDefaultValue)
@@ -19424,12 +16915,15 @@ bool URigVMController::ChangePinType(URigVMPin* InPin, TRigVMTypeIndex InTypeInd
 	{
 		if (!bIsTransacting)
 		{
-			if (TemplateNode->IsA<URigVMUnitNode>() || TemplateNode->IsA<URigVMDispatchNode>())
+			if (InPin->GetDirection() != ERigVMPinDirection::Hidden)
 			{
-				if(!TemplateNode->SupportsType(InPin, InTypeIndex))
+				if (TemplateNode->IsA<URigVMUnitNode>() || TemplateNode->IsA<URigVMDispatchNode>())
 				{
-					ReportErrorf(TEXT("ChangePinType: %s doesn't support type '%s'."), *InPin->GetPinPath(), *FRigVMRegistry::Get().GetType(InTypeIndex).CPPType.ToString());
-					return false;
+					if(!TemplateNode->SupportsType(InPin, InTypeIndex))
+					{
+						ReportErrorf(TEXT("ChangePinType: %s doesn't support type '%s'."), *InPin->GetPinPath(), *FRigVMRegistry::Get().GetType(InTypeIndex).CPPType.ToString());
+						return false;
+					}
 				}
 			}
 		}
@@ -19476,21 +16970,24 @@ bool URigVMController::ChangePinType(URigVMPin* InPin, TRigVMTypeIndex InTypeInd
 	{
 		Links.Append(InPin->GetSourceLinks(true));
 		Links.Append(InPin->GetTargetLinks(true));
-		DetachLinksFromPinObjects(&Links);
-
-		const FString OrphanedName = FString::Printf(TEXT("%s%s"), *URigVMPin::OrphanPinPrefix, *InPin->GetName());
-		if(InPin->GetNode()->FindPin(OrphanedName) == nullptr)
+		if (Links.Num() > 0)
 		{
-			URigVMPin* OrphanedPin = NewObject<URigVMPin>(InPin->GetNode(), *OrphanedName);
-			ConfigurePinFromPin(OrphanedPin, InPin);
-			OrphanedPin->DisplayName = InPin->GetFName();
+			DetachLinksFromPinObjects(&Links);
 
-			if(OrphanedPin->IsStruct())
+			const FString OrphanedName = FString::Printf(TEXT("%s%s"), *URigVMPin::OrphanPinPrefix, *InPin->GetName());
+			if(InPin->GetNode()->FindPin(OrphanedName) == nullptr)
 			{
-				AddPinsForStruct(OrphanedPin->GetScriptStruct(), OrphanedPin->GetNode(), OrphanedPin, OrphanedPin->Direction, OrphanedPin->GetDefaultValue(), false);
-			}
+				URigVMPin* OrphanedPin = NewObject<URigVMPin>(InPin->GetNode(), *OrphanedName);
+				ConfigurePinFromPin(OrphanedPin, InPin);
+				OrphanedPin->DisplayName = InPin->GetFName();
+
+				if(OrphanedPin->IsStruct())
+				{
+					AddPinsForStruct(OrphanedPin->GetScriptStruct(), OrphanedPin->GetNode(), OrphanedPin, OrphanedPin->Direction, OrphanedPin->GetDefaultValue(), false);
+				}
 				
-			InPin->GetNode()->OrphanedPins.Add(OrphanedPin);
+				InPin->GetNode()->OrphanedPins.Add(OrphanedPin);
+			}
 		}
 	}
 
@@ -20296,7 +17793,6 @@ FRigVMClientPatchResult URigVMController::PatchUnitNodesOnLoad()
 						if(!FRigVMRegistry::Get().IsWildCardType(*ResolvedTypeIndex))
 						{
 							ChangePinType(Pin, *ResolvedTypeIndex, false, false);
-							NewNode->UpdateFilteredPermutations(Pin, {*ResolvedTypeIndex});
 						}
 					}
 				}
@@ -20455,12 +17951,11 @@ FRigVMClientPatchResult URigVMController::PatchIfSelectNodesOnLoad()
 						{
 							if(!FRigVMRegistry::Get().IsWildCardType(*ResolvedTypeIndex))
 							{
-								NewNode->UpdateFilteredPermutations(Pin, {*ResolvedTypeIndex});
+								ResolveWildCardPin(Pin, *ResolvedTypeIndex, false);
 							}
 						}
 					}
 				}
-				UpdateTemplateNodePinTypes(NewNode, false, false);
 			}
 
 			ApplyPinStates(NewNode, PinStates, {}, false);
@@ -20510,6 +18005,71 @@ FRigVMClientPatchResult URigVMController::PatchArrayNodesOnLoad()
 	return Result;
 }
 
+FRigVMClientPatchResult URigVMController::PatchInvalidLinksOnWildcards()
+{
+	FRigVMClientPatchResult Result;
+	
+	if (const URigVMGraph* Graph = GetGraph())
+	{
+		// Remove links between wildcard pins
+		TArray<URigVMLink*> LinksToRemove;
+		for (URigVMLink* Link : Graph->GetLinks())
+		{
+			bool bRemove = false;
+			if (URigVMPin* SourcePin = Link->GetSourcePin())
+			{
+				if (SourcePin->IsWildCard())
+				{
+					bRemove = true;
+				}
+			}
+			if (URigVMPin* TargetPin = Link->GetTargetPin())
+			{
+				if (TargetPin->IsWildCard())
+				{
+					bRemove = true;
+				}
+			}
+			if (bRemove)
+			{
+				LinksToRemove.Add(Link);
+			}
+		}
+		if (!LinksToRemove.IsEmpty())
+		{
+			Result.bChangedContent = true;
+		}
+		for (URigVMLink* Link : LinksToRemove)
+		{
+			if (!BreakLink(Link->GetSourcePin(), Link->GetTargetPin()))
+			{
+				Result.ErrorMessages.Add(FString::Printf(TEXT("Error breaking link %s in PatchInvalidLinksOnWildcards"), *Link->GetPinPathRepresentation()));
+				Result.bSucceeded = false;
+			}
+		}
+
+		// Remove exposed pins of type wildcard
+		if (URigVMCollapseNode* CollapseNode = Cast<URigVMCollapseNode>(Graph->GetOuter()))
+		{
+			TArray<URigVMPin*> ExposedPins = CollapseNode->GetPins();
+			for (URigVMPin* ExposedPin : ExposedPins)
+			{
+				if (ExposedPin->IsWildCard())
+				{
+					FRigVMControllerGraphGuard GraphGuard(this, CollapseNode->GetContainedGraph());
+					if (!RemoveExposedPin(ExposedPin->GetFName(), false))
+					{
+						Result.ErrorMessages.Add(FString::Printf(TEXT("Error removing exposed pin %s PatchInvalidLinksOnWildcards"), *ExposedPin->GetPinPath(true)));
+						Result.bSucceeded = false;
+					}
+				}
+			}
+		}
+	}
+
+	return Result;
+}
+
 void URigVMController::PostDuplicateHost(const FString& InOldPathName, const FString& InNewPathName)
 {
 	if (const URigVMGraph* Graph = GetGraph())
diff --git a/Engine/Source/Developer/RigVMDeveloper/Private/RigVMModel/RigVMControllerActions.cpp b/Engine/Source/Developer/RigVMDeveloper/Private/RigVMModel/RigVMControllerActions.cpp
--- a/Engine/Source/Developer/RigVMDeveloper/Private/RigVMModel/RigVMControllerActions.cpp
+++ b/Engine/Source/Developer/RigVMDeveloper/Private/RigVMModel/RigVMControllerActions.cpp
@@ -648,9 +648,7 @@ FRigVMRemoveNodeAction::FRigVMRemoveNodeAction(URigVMNode* InNode, URigVMControl
 		else
 		{
 			InverseAction.AddAction(FRigVMAddTemplateNodeAction(TemplateNode), InController);
-			InverseAction.AddAction(FRigVMSetPreferredTemplatePermutationsAction(TemplateNode, TemplateNode->PreferredPermutationPairs), InController);
-			InverseAction.AddAction(FRigVMSetTemplateFilteredPermutationsAction(TemplateNode, {}), InController);
-
+			
 			for (URigVMPin* Pin : TemplateNode->GetPins())
 			{
 				if (!Pin->IsWildCard())
@@ -1234,208 +1232,6 @@ bool FRigVMSetPinDefaultValueAction::Redo(URigVMController* InController)
 	return FRigVMBaseAction::Redo(InController);
 }
 
-FRigVMSetTemplateFilteredPermutationsAction::FRigVMSetTemplateFilteredPermutationsAction(URigVMTemplateNode* InNode, const TArray<int32>& InOldFilteredPermutations)
-: NodePath(InNode->GetNodePath())
-, OldFilteredPermutations(InOldFilteredPermutations)
-, NewFilteredPermutations(InNode->GetFilteredPermutationsIndices())
-{
-	
-}
-
-bool FRigVMSetTemplateFilteredPermutationsAction::Merge(const FRigVMBaseAction* Other)
-{
-	if (!FRigVMBaseAction::Merge(Other))
-	{
-		return false;
-	}
-
-	const FRigVMSetTemplateFilteredPermutationsAction* Action = (const FRigVMSetTemplateFilteredPermutationsAction*)Other;
-	if (NodePath != Action->NodePath)
-	{
-		return false;
-	}
-
-	NewFilteredPermutations = Action->NewFilteredPermutations;
-	return true;
-}
-
-bool FRigVMSetTemplateFilteredPermutationsAction::Undo(URigVMController* InController)
-{
-	if(!FRigVMBaseAction::Undo(InController))
-	{
-		return false;
-	}
-
-	if (URigVMTemplateNode* TemplateNode = Cast<URigVMTemplateNode>(InController->GetGraph()->FindNode(NodePath)))
-	{
-		TemplateNode->FilteredPermutations = OldFilteredPermutations;
-		return true;
-	}
-	return false;
-}
-
-bool FRigVMSetTemplateFilteredPermutationsAction::Redo(URigVMController* InController)
-{
-	URigVMTemplateNode* TemplateNode = Cast<URigVMTemplateNode>(InController->GetGraph()->FindNode(NodePath));
-	if (!TemplateNode)
-	{
-		return false;
-	}
-	TemplateNode->FilteredPermutations = NewFilteredPermutations;
-	return FRigVMBaseAction::Redo(InController);
-}
-
-FRigVMSetPreferredTemplatePermutationsAction::FRigVMSetPreferredTemplatePermutationsAction(URigVMTemplateNode* InNode, const TArray<FRigVMTemplatePreferredType>& InPreferredTypes)
-: NodePath(InNode->GetNodePath())
-, OldPreferredPermutationTypes(InNode->PreferredPermutationPairs)
-, NewPreferredPermutationTypes(InPreferredTypes)
-{
-	
-}
-
-bool FRigVMSetPreferredTemplatePermutationsAction::Merge(const FRigVMBaseAction* Other)
-{
-	if (!FRigVMBaseAction::Merge(Other))
-	{
-		return false;
-	}
-
-	const FRigVMSetPreferredTemplatePermutationsAction* Action = (const FRigVMSetPreferredTemplatePermutationsAction*)Other;
-	if (NodePath != Action->NodePath)
-	{
-		return false;
-	}
-
-	NewPreferredPermutationTypes = Action->NewPreferredPermutationTypes;
-	return true;
-}
-
-bool FRigVMSetPreferredTemplatePermutationsAction::Undo(URigVMController* InController)
-{
-	if(!FRigVMBaseAction::Undo(InController))
-	{
-		return false;
-	}
-
-	if (URigVMTemplateNode* TemplateNode = Cast<URigVMTemplateNode>(InController->GetGraph()->FindNode(NodePath)))
-	{
-		TemplateNode->PreferredPermutationPairs = OldPreferredPermutationTypes;
-		return true;
-	}
-	return false;
-}
-
-bool FRigVMSetPreferredTemplatePermutationsAction::Redo(URigVMController* InController)
-{
-	URigVMTemplateNode* TemplateNode = Cast<URigVMTemplateNode>(InController->GetGraph()->FindNode(NodePath));
-	if (!TemplateNode)
-	{
-		return false;
-	}
-	TemplateNode->PreferredPermutationPairs = NewPreferredPermutationTypes;
-	return FRigVMBaseAction::Redo(InController);
-}
-
-FRigVMSetLibraryTemplateAction::FRigVMSetLibraryTemplateAction(URigVMLibraryNode* InNode, FRigVMTemplate& InNewTemplate)
-: NodePath(InNode->GetNodePath())
-{
-	FMemoryWriter ArchiveWriterOld(OldTemplateBytes);
-	InNode->Template.Serialize(ArchiveWriterOld);
-	FMemoryWriter ArchiveWriterNew(NewTemplateBytes);
-	InNewTemplate.Serialize(ArchiveWriterNew);	
-}
-
-bool FRigVMSetLibraryTemplateAction::Undo(URigVMController* InController)
-{
-	if(!FRigVMBaseAction::Undo(InController))
-	{
-		return false;
-	}
-
-	if (URigVMLibraryNode* LibraryNode = Cast<URigVMLibraryNode>(InController->GetGraph()->FindNode(NodePath)))
-	{
-		FMemoryReader ArchiveReader(OldTemplateBytes);
-		FRigVMTemplate NewTemplate;
-		NewTemplate.Serialize(ArchiveReader);
-		LibraryNode->Template = NewTemplate;
-
-		// Update type to permutations map for each argument
-		for (FRigVMTemplateArgument& Argument : LibraryNode->Template.Arguments)
-		{
-			Argument.TypeToPermutations.Reset();
-			for (int32 PermutationIndex=0; PermutationIndex<Argument.TypeIndices.Num(); ++PermutationIndex)
-			{
-				const int32& TypeIndex = Argument.TypeIndices[PermutationIndex];
-				if (TArray<int32>* Permutations = Argument.TypeToPermutations.Find(TypeIndex))
-				{
-					Permutations->Add(PermutationIndex);
-				}
-				else
-				{
-					Argument.TypeToPermutations.Add(TypeIndex, {PermutationIndex});
-				}
-			}
-		}
-		
-		InController->Notify(ERigVMGraphNotifType::NodeDescriptionChanged, LibraryNode);
-		if (URigVMTemplateNode* EntryNode = Cast<URigVMTemplateNode>(LibraryNode->GetEntryNode()))
-		{
-			InController->Notify(ERigVMGraphNotifType::NodeDescriptionChanged, EntryNode);
-		}
-		if (URigVMTemplateNode* ReturnNode = Cast<URigVMTemplateNode>(LibraryNode->GetReturnNode()))
-		{
-			InController->Notify(ERigVMGraphNotifType::NodeDescriptionChanged, ReturnNode);
-		}
-		
-		return true;
-	}
-	return false;
-}
-
-bool FRigVMSetLibraryTemplateAction::Redo(URigVMController* InController)
-{
-	URigVMLibraryNode* LibraryNode = Cast<URigVMLibraryNode>(InController->GetGraph()->FindNode(NodePath));
-	if (!LibraryNode)
-	{
-		return false;
-	}
-
-	FMemoryReader ArchiveReader(NewTemplateBytes);
-	FRigVMTemplate NewTemplate;
-	NewTemplate.Serialize(ArchiveReader);
-	LibraryNode->Template = NewTemplate;
-	
-	// Update type to permutations map for each argument
-	for (FRigVMTemplateArgument& Argument : LibraryNode->Template.Arguments)
-	{
-		Argument.TypeToPermutations.Reset();
-		for (int32 PermutationIndex=0; PermutationIndex<Argument.TypeIndices.Num(); ++PermutationIndex)
-		{
-			const int32& TypeIndex = Argument.TypeIndices[PermutationIndex];
-			if (TArray<int32>* Permutations = Argument.TypeToPermutations.Find(TypeIndex))
-			{
-				Permutations->Add(PermutationIndex);
-			}
-			else
-			{
-				Argument.TypeToPermutations.Add(TypeIndex, {PermutationIndex});
-			}
-		}
-	}
-
-	InController->Notify(ERigVMGraphNotifType::NodeDescriptionChanged, LibraryNode);
-	if (URigVMTemplateNode* EntryNode = Cast<URigVMTemplateNode>(LibraryNode->GetEntryNode()))
-	{
-		InController->Notify(ERigVMGraphNotifType::NodeDescriptionChanged, EntryNode);
-	}
-	if (URigVMTemplateNode* ReturnNode = Cast<URigVMTemplateNode>(LibraryNode->GetReturnNode()))
-	{
-		InController->Notify(ERigVMGraphNotifType::NodeDescriptionChanged, ReturnNode);
-	}
-	
-	return FRigVMBaseAction::Redo(InController);
-}
-
 FRigVMInsertArrayPinAction::FRigVMInsertArrayPinAction(URigVMPin* InArrayPin, int32 InIndex, const FString& InNewDefaultValue)
 : ArrayPinPath(InArrayPin->GetPinPath())
 , Index(InIndex)
diff --git a/Engine/Source/Developer/RigVMDeveloper/Private/RigVMModel/RigVMPin.cpp b/Engine/Source/Developer/RigVMDeveloper/Private/RigVMModel/RigVMPin.cpp
--- a/Engine/Source/Developer/RigVMDeveloper/Private/RigVMModel/RigVMPin.cpp
+++ b/Engine/Source/Developer/RigVMDeveloper/Private/RigVMModel/RigVMPin.cpp
@@ -632,6 +632,11 @@ bool URigVMPin::IsWildCard() const
 			return true;
 		}
 	}
+	if (CPPType.IsEmpty())
+	{
+		// Unknown type
+		return true;
+	}
 	return false;
 }
 
@@ -1736,41 +1741,69 @@ bool URigVMPin::CanLink(URigVMPin* InSourcePin, URigVMPin* InTargetPin, FString*
 						}
 					}
 					return true;
-
 				};
 
-				auto IsPinValidForTypeChange = [](URigVMPin* InPin, bool bIsInput, ERigVMPinDirection InDirection) -> bool
+				if(InSourcePin->IsWildCard() && !InTargetPin->IsWildCard())
+				{
+					bCPPTypesDiffer = !TemplateNodeSupportsType(InSourcePin, InTargetPin->GetTypeIndex(), OutFailureReason);
+				}
+				else if(InTargetPin->IsWildCard() && !InSourcePin->IsWildCard())
 				{
-					if(InPin->IsWildCard())
+					bCPPTypesDiffer = !TemplateNodeSupportsType(InTargetPin, InSourcePin->GetTypeIndex(), OutFailureReason);
+				}
+				else if(InSourcePin->IsWildCard() && InTargetPin->IsWildCard())
+				{
+					// Find out if these pins have any type in common
+					uint8 SourceLevels = 0;
+					uint8 TargetLevels = 0;
+					URigVMPin* RootSourcePin = InSourcePin;
+					URigVMPin* RootTargetPin = InTargetPin;
+					while (RootSourcePin->IsArrayElement())
 					{
-						return true;
+						SourceLevels++;
+						RootSourcePin = RootSourcePin->GetParentPin();
 					}
-
-					if(!InPin->GetNode()->IsA<URigVMTemplateNode>() || InDirection == ERigVMPinDirection::Invalid)
+					while (RootTargetPin->IsArrayElement())
 					{
-						return false;
+						TargetLevels++;
+						RootTargetPin = RootTargetPin->GetParentPin();
 					}
 
-					if((bIsInput && (InDirection != ERigVMPinDirection::Input)) ||
-						(!bIsInput && (InDirection != ERigVMPinDirection::Output)))
+					URigVMTemplateNode* SourceTemplateNode = Cast<URigVMTemplateNode>(RootSourcePin->GetNode());
+					URigVMTemplateNode* TargetTemplateNode = Cast<URigVMTemplateNode>(RootTargetPin->GetNode());
+					TArray<int32> SourcePermutations = SourceTemplateNode->GetResolvedPermutationIndices(true);
+					TArray<int32> TargetPermutations = TargetTemplateNode->GetResolvedPermutationIndices(true);
+					const FRigVMTemplate* SourceTemplate = SourceTemplateNode->GetTemplate();
+					const FRigVMTemplate* TargetTemplate = TargetTemplateNode->GetTemplate();
+					const FRigVMTemplateArgument* SourceRootArgument = SourceTemplate->FindArgument(RootSourcePin->GetFName());
+					const FRigVMTemplateArgument* TargetRootArgument = TargetTemplate->FindArgument(RootTargetPin->GetFName());
+
+					TArray<TRigVMTypeIndex> SourceTypes;
+					FRigVMRegistry& Registry = FRigVMRegistry::Get();
+					for (int32 Permutation : SourcePermutations)
 					{
-						if(InPin->IsRootPin() ||
-							(InPin->GetParentPin()->IsRootPin() && InPin->IsArrayElement()))
+						TRigVMTypeIndex Type = SourceRootArgument->GetTypeIndices()[Permutation];
+						for (int32 i=0; i<SourceLevels; ++i)
 						{
-							return true;
+							check(Registry.IsArrayType(Type));
+							Type = Registry.GetBaseTypeFromArrayTypeIndex(Type);
+						}
+						SourceTypes.Add(Type);
+					}
+					for (int32 Permutation : TargetPermutations)
+					{
+						TRigVMTypeIndex Type = TargetRootArgument->GetTypeIndices()[Permutation];
+						for (int32 i=0; i<TargetLevels; ++i)
+						{
+							check(Registry.IsArrayType(Type));
+							Type = Registry.GetBaseTypeFromArrayTypeIndex(Type);
+						}
+						if (SourceTypes.Contains(Type))
+						{
+							bCPPTypesDiffer = false;
+							break;
 						}
 					}
-					
-					return false;
-				};
-				
-				if(IsPinValidForTypeChange(InSourcePin, false, InUserLinkDirection))
-				{
-					bCPPTypesDiffer = !TemplateNodeSupportsType(InSourcePin, InTargetPin->GetTypeIndex(), OutFailureReason);
-				}
-				else if(IsPinValidForTypeChange(InTargetPin, true, InUserLinkDirection))
-				{
-					bCPPTypesDiffer = !TemplateNodeSupportsType(InTargetPin, InSourcePin->GetTypeIndex(), OutFailureReason);
 				}
 			}
 		}
diff --git a/Engine/Source/Developer/RigVMDeveloper/Public/RigVMModel/Nodes/RigVMDispatchNode.h b/Engine/Source/Developer/RigVMDeveloper/Public/RigVMModel/Nodes/RigVMDispatchNode.h
--- a/Engine/Source/Developer/RigVMDeveloper/Public/RigVMModel/Nodes/RigVMDispatchNode.h
+++ b/Engine/Source/Developer/RigVMDeveloper/Public/RigVMModel/Nodes/RigVMDispatchNode.h
@@ -77,7 +77,6 @@ protected:
 
 	virtual FText GetToolTipTextForPin(const URigVMPin* InPin) const override;
 	const UScriptStruct* GetFactoryStruct() const;
-	const FRigVMTemplateTypeMap& GetFilteredTypes() const;
 	virtual void InvalidateCache() override;
 	virtual bool ShouldInputPinComputeLazily(const URigVMPin* InPin) const override;
 	FString GetFactoryDefaultValue() const;
diff --git a/Engine/Source/Developer/RigVMDeveloper/Public/RigVMModel/Nodes/RigVMFunctionEntryNode.h b/Engine/Source/Developer/RigVMDeveloper/Public/RigVMModel/Nodes/RigVMFunctionEntryNode.h
--- a/Engine/Source/Developer/RigVMDeveloper/Public/RigVMModel/Nodes/RigVMFunctionEntryNode.h
+++ b/Engine/Source/Developer/RigVMDeveloper/Public/RigVMModel/Nodes/RigVMFunctionEntryNode.h
@@ -18,8 +18,8 @@ public:
 
 	// Override template node functions
 	virtual UScriptStruct* GetScriptStruct() const override { return nullptr; }
-	virtual const FRigVMTemplate* GetTemplate() const override;
-	virtual FName GetNotation() const override;
+	virtual const FRigVMTemplate* GetTemplate() const override { return nullptr; }
+	virtual FName GetNotation() const override { return NAME_None; }
 
 	// Override node functions
 	virtual FLinearColor GetNodeColor() const override;
diff --git a/Engine/Source/Developer/RigVMDeveloper/Public/RigVMModel/Nodes/RigVMFunctionReturnNode.h b/Engine/Source/Developer/RigVMDeveloper/Public/RigVMModel/Nodes/RigVMFunctionReturnNode.h
--- a/Engine/Source/Developer/RigVMDeveloper/Public/RigVMModel/Nodes/RigVMFunctionReturnNode.h
+++ b/Engine/Source/Developer/RigVMDeveloper/Public/RigVMModel/Nodes/RigVMFunctionReturnNode.h
@@ -18,8 +18,8 @@ public:
 
 	// Override template node functions
 	virtual UScriptStruct* GetScriptStruct() const override { return nullptr; }
-	virtual const FRigVMTemplate* GetTemplate() const override;
-	virtual FName GetNotation() const override;
+	virtual const FRigVMTemplate* GetTemplate() const override { return nullptr; }
+	virtual FName GetNotation() const override { return NAME_None; }
 
 	// Override node functions
 	virtual FLinearColor GetNodeColor() const override;
diff --git a/Engine/Source/Developer/RigVMDeveloper/Public/RigVMModel/Nodes/RigVMLibraryNode.h b/Engine/Source/Developer/RigVMDeveloper/Public/RigVMModel/Nodes/RigVMLibraryNode.h
--- a/Engine/Source/Developer/RigVMDeveloper/Public/RigVMModel/Nodes/RigVMLibraryNode.h
+++ b/Engine/Source/Developer/RigVMDeveloper/Public/RigVMModel/Nodes/RigVMLibraryNode.h
@@ -24,7 +24,6 @@ class RIGVMDEVELOPER_API URigVMLibraryNode : public URigVMTemplateNode
 	GENERATED_BODY()
 
 public:
-	URigVMLibraryNode();
 
 	// Override node functions
 	virtual bool IsDefinedAsConstant() const override;
@@ -33,8 +32,8 @@ public:
 	
 	// Override template node functions
 	virtual UScriptStruct* GetScriptStruct() const override { return nullptr; }
-	virtual const FRigVMTemplate* GetTemplate() const override;
-	virtual FName GetNotation() const override;
+	virtual const FRigVMTemplate* GetTemplate() const override { return nullptr; }
+	virtual FName GetNotation() const override { return NAME_None; }
 
 	// URigVMNode interface
 	virtual  FText GetToolTipText() const override;
@@ -63,8 +62,6 @@ public:
 
 protected:
 
-	FRigVMTemplate Template;
-	
 	virtual TArray<int32> GetInstructionsForVMImpl(URigVM* InVM, const FRigVMASTProxy& InProxy = FRigVMASTProxy()) const override; 
 	const static TArray<URigVMNode*> EmptyNodes;
 	const static TArray<URigVMLink*> EmptyLinks;
diff --git a/Engine/Source/Developer/RigVMDeveloper/Public/RigVMModel/Nodes/RigVMTemplateNode.h b/Engine/Source/Developer/RigVMDeveloper/Public/RigVMModel/Nodes/RigVMTemplateNode.h
--- a/Engine/Source/Developer/RigVMDeveloper/Public/RigVMModel/Nodes/RigVMTemplateNode.h
+++ b/Engine/Source/Developer/RigVMDeveloper/Public/RigVMModel/Nodes/RigVMTemplateNode.h
@@ -106,12 +106,15 @@ public:
 	// returns true if a pin supports a given type
 	bool SupportsType(const URigVMPin* InPin, TRigVMTypeIndex InTypeIndex, TRigVMTypeIndex* OutTypeIndex = nullptr);
 
-	// returns true if a pin supports a given type after filtering
-	bool FilteredSupportsType(const URigVMPin* InPin, TRigVMTypeIndex InTypeIndex, TRigVMTypeIndex* OutTypeIndex = nullptr, bool bAllowFloatingPointCasts = true);
-
 	// returns the resolved functions for the template
 	TArray<const FRigVMFunction*> GetResolvedPermutations() const;
 
+	// returns the resolved permutation indices for the template
+	TArray<int32> GetResolvedPermutationIndices(bool bAllowFloatingPointCasts) const;
+
+	// Returns a map of the resolved pin types
+	FRigVMTemplateTypeMap GetTemplatePinTypeMap(bool bIncludeHiddenPins = false, bool bIncludeExecutePins = false) const;
+
 	// returns the template used for this node
 	virtual const FRigVMTemplate* GetTemplate() const;
 
@@ -132,48 +135,15 @@ public:
 	// returns the display name for a pin
 	FName GetDisplayNameForPin(const FName& InRootPinName, const TArray<int32>& InPermutationIndices = TArray<int32>()) const;
 
-	// returns the indeces of the filtered permutations
-	const TArray<int32>& GetFilteredPermutationsIndices() const;
-
-	// returns the filtered types of this pin
-	TArray<TRigVMTypeIndex> GetFilteredTypesForPin(URigVMPin* InPin) const;
-
 	// Tries to reduce the input types to a single type, if all are compatible
 	// Will prioritize the InPreferredType if available
 	TRigVMTypeIndex TryReduceTypesToSingle(const TArray<TRigVMTypeIndex>& InTypes, const TRigVMTypeIndex PreferredType = TRigVMTypeIndex()) const;
 
-	// returns true if updating pin filters with InTypes would result in different filters 
-	bool PinNeedsFilteredTypesUpdate(URigVMPin* InPin, const TArray<TRigVMTypeIndex>& InTypeIndices);
-	bool PinNeedsFilteredTypesUpdate(URigVMPin* InPin, URigVMPin* LinkedPin);
-
-	// updates the filtered permutations given a link or the types for a pin
-	bool UpdateFilteredPermutations(URigVMPin* InPin, const TArray<TRigVMTypeIndex>& InTypeIndices);
-	bool UpdateFilteredPermutations(URigVMPin* InPin, URigVMPin* LinkedPin);
-
-	// initializes the filtered permutations to all possible permutations
-	void InitializeFilteredPermutations();
-
-	// Initializes the filtered permutations and preferred permutation from the types of the pins
-	void InitializeFilteredPermutationsFromTypes(bool bAllowCasting);
-
-	// Converts the preferred types per pin from index to string
-	void ConvertPreferredTypesToString();
-
-	// Converts the preferred types per pin from string to indices
-	void ConvertPreferredTypesToTypeIndex();
-
-	// Returns the preferred type, or an invalid type if non was found
-	TRigVMTypeIndex GetPreferredType(const FName& ArgumentName) const;
-
 protected:
 
 	virtual void InvalidateCache() override;
 	
-	TArray<int32> GetNewFilteredPermutations(URigVMPin* InPin, URigVMPin* LinkedPin);
-	TArray<int32> GetNewFilteredPermutations(URigVMPin* InPin, const TArray<TRigVMTypeIndex>& InTypeIndices);
-
 	TArray<int32> FindPermutationsForTypes(const TArray<FRigVMTemplatePreferredType>& ArgumentTypes, bool bAllowCasting = false) const;
-	TArray<FRigVMTemplatePreferredType> GetPreferredTypesForPermutation(const int32 InPermutationIndex) const;
 	FRigVMTemplateTypeMap GetTypesForPermutation(const int32 InPermutationIndex) const;
 
 	UPROPERTY()
@@ -190,9 +160,7 @@ protected:
 #endif
 
 	UPROPERTY()
-	TArray<FRigVMTemplatePreferredType> PreferredPermutationPairs;
-
-	TArray<int32> FilteredPermutations;
+	TArray<FRigVMTemplatePreferredType> PreferredPermutationPairs_DEPRECATED;
 
 	mutable const FRigVMTemplate* CachedTemplate;
 	mutable const FRigVMFunction* CachedFunction;
diff --git a/Engine/Source/Developer/RigVMDeveloper/Public/RigVMModel/RigVMController.h b/Engine/Source/Developer/RigVMDeveloper/Public/RigVMModel/RigVMController.h
--- a/Engine/Source/Developer/RigVMDeveloper/Public/RigVMModel/RigVMController.h
+++ b/Engine/Source/Developer/RigVMDeveloper/Public/RigVMModel/RigVMController.h
@@ -83,6 +83,7 @@ DECLARE_DELEGATE_RetVal_ThreeParams(FName, FRigVMController_RequestNewExternalVa
 DECLARE_DELEGATE_RetVal_TwoParams(bool, FRigVMController_IsDependencyCyclicDelegate, const FRigVMGraphFunctionHeader& Dependent, const FRigVMGraphFunctionHeader& Dependency)
 DECLARE_DELEGATE_RetVal_TwoParams(FRigVMController_BulkEditResult, FRigVMController_RequestBulkEditDialogDelegate, URigVMLibraryNode*, ERigVMControllerBulkEditType)
 DECLARE_DELEGATE_RetVal_OneParam(bool, FRigVMController_RequestBreakLinksDialogDelegate, TArray<URigVMLink*>)
+DECLARE_DELEGATE_RetVal_OneParam(TRigVMTypeIndex, FRigVMController_RequestPinTypeSelectionDelegate, const TArray<TRigVMTypeIndex>& Types)
 DECLARE_DELEGATE_FiveParams(FRigVMController_OnBulkEditProgressDelegate, TSoftObjectPtr<URigVMFunctionReferenceNode>, ERigVMControllerBulkEditType, ERigVMControllerBulkEditProgress, int32, int32)
 DECLARE_DELEGATE_RetVal_TwoParams(FString, FRigVMController_PinPathRemapDelegate, const FString& /* InPinPath */, bool /* bIsInput */);
 DECLARE_DELEGATE_OneParam(FRigVMController_RequestJumpToHyperlinkDelegate, const UObject* InSubject);
@@ -321,7 +322,7 @@ public:
 	// Turns a resolved templated node(s) back into its template.
 	UFUNCTION(BlueprintCallable, Category = RigVMController)
 	bool UnresolveTemplateNodes(const TArray<FName>& InNodeNames, bool bSetupUndoRedo = true, bool bPrintPythonCommand = false);
-	bool UnresolveTemplateNodes(const TArray<URigVMNode*>& InNodes, bool bSetupUndoRedo, bool bBreakLinks = true);
+	bool UnresolveTemplateNodes(const TArray<URigVMNode*>& InNodes, bool bSetupUndoRedo);
 
 	// Upgrades a set of nodes with each corresponding next known version
 	UFUNCTION(BlueprintCallable, Category = RigVMController)
@@ -895,6 +896,9 @@ public:
 	// A delegate to ask the host / client for a dialog to confirm a bulk edit
 	FRigVMController_RequestBreakLinksDialogDelegate RequestBreakLinksDialogDelegate;
 
+	// A delegate to ask the host / client for a dialog to select a pin type
+	FRigVMController_RequestPinTypeSelectionDelegate RequestPinTypeSelectionDelegate;
+
 	// A delegate to inform the host / client about the progress during a bulk edit
 	FRigVMController_OnBulkEditProgressDelegate OnBulkEditProgressDelegate;
 
@@ -921,22 +925,9 @@ public:
 	// removes any orphan pins that no longer holds a link
 	bool RemoveUnusedOrphanedPins(URigVMNode* InNode);
 
-	// Initializes and recomputes the filtered permutations of all template nodes in the graph
-	// Returns true if any pin has change it's type or link was broken
-	bool RecomputeAllTemplateFilteredPermutations(bool bSetupUndoRedo);
-
-	// Update the template of a subgraph with the filtered permutations of the interface nodes
-	bool UpdateLibraryTemplate(URigVMLibraryNode* LibraryNode, bool bSetupUndoRedo);
-
 	// Update filtered permutations, and propagate both ways of the link before adding this link
 	bool PrepareToLink(URigVMPin* FirstToResolve, URigVMPin* SecondToResolve, bool bSetupUndoRedo);
 
-	// Try to initialize the filterd permutations from the pin types
-	void InitializeFilteredPermutationsFromTemplateTypes();
-
-	// Initializes filtered permuations to be unresolved in all template nodes in graph
-	void InitializeAllTemplateFiltersInGraph(bool bSetupUndoRedo, bool bChangePinTypes);
-	
 #endif
 
 	bool FullyResolveTemplateNode(URigVMTemplateNode* InNode, int32 InPermutationIndex, bool bSetupUndoRedo);
@@ -951,9 +942,6 @@ public:
 	// A flag that can be used to turn off pin default value validation if necessary
 	bool bValidatePinDefaults;
 
-	// A flag to suspend the recomputation of filtered permutations of outer graphs
-	bool bSuspendRecomputingOuterTemplateFilters;
-
 	const FRigVMByteCode* GetCurrentByteCode() const;
 
 	void ReportInfo(const FString& InMessage) const;
@@ -1157,43 +1145,10 @@ private:
 	static void PostProcessDefaultValue(URigVMPin* Pin, FString& OutDefaultValue);
 
 	void ResolveTemplateNodeMetaData(URigVMTemplateNode* InNode, bool bSetupUndoRedo);
-	
-	// Prepare the graph for the change this template node is about to make
-	// If any of the types is supported (without breaking any links), then the filtered permutations will be updated and the change will
-	// propagate to other nodes in the graph.
-	// If it is not supported, we will attempt to find and break any links that do not support this change
-	bool PrepareTemplatePinForType(URigVMPin* InPin, const TArray<TRigVMTypeIndex>& InTypeIndices, bool bSetupUndoRedo);
-
-	// Get filtered types for a wildcard node. If template node, that means just returning its filtered wildcard types, but if it's another type of node (select, if, rereoute), iterate
-	// its connections to figure out the filtered types
-	TArray<TRigVMTypeIndex> GetFilteredTypes(URigVMPin* InPin);
-	
-	// Updates the permutations allowed without having to break any links
-	bool UpdateFilteredPermutations(URigVMPin* InPin, URigVMPin* InLinkedPin, bool bSetupUndoRedo);
-	bool UpdateFilteredPermutations(URigVMPin* InPin, const TArray<TRigVMTypeIndex>& InTypeIndices, bool bSetupUndoRedo);
-	bool UpdateFilteredPermutations(URigVMTemplateNode* InNode, const TArray<int32>& InPermutations, bool bSetupUndoRedo);
 
 	// Changes Pin types if filtered types of a pin are unique
 	bool UpdateTemplateNodePinTypes(URigVMTemplateNode* InNode, bool bSetupUndoRedo, bool bInitializeDefaultValue = true);
 
-	// Reduces the filtered permutations of all templates in the graph to comply with the types filtered by InNode
-	// Returns false if a link had to be broken
-	bool PropagateTemplateFilteredTypes(URigVMTemplateNode* InNode, bool bSetupUndoRedo);
-
-	// Adds a preferred type for the pin
-	// Returns false if the pin already has a different type
-	bool AddPreferredType(URigVMTemplateNode* InNode, const FName& InPinName, const TRigVMTypeIndex& InPreferredTypeIndex, bool bSetupUndoRedo);
-
-	// Removes preferred type
-	// Returns true if the preferred type was found and removed
-	bool RemovePreferredType(URigVMTemplateNode* InNode, const FName& InPinName, bool bSetupUndoRedo);
-
-	// Returns true if the pin is connected, and the filtered types is reduced (not infinite like reroute, if, select or array nodes)
-	bool ShouldPinOwnArgument(URigVMPin* InPin);
-
-	// Given a Entry or Return pin and a potential pin to link it to, try to add the first pin as an argument of the library's template
-	bool AddArgumentForPin(URigVMPin* InPin, URigVMPin* InToLinkPin, bool bSetupUndoRedo = true, bool bPrintPythonCommand = false);
-
 	bool ChangePinType(const FString& InPinPath, const FString& InCPPType, const FName& InCPPTypeObjectPath, bool bSetupUndoRedo, bool bSetupOrphanPins = true, bool bBreakLinks = true, bool bRemoveSubPins = true, bool bInitializeDefaultValue = true);
 	bool ChangePinType(URigVMPin* InPin, const FString& InCPPType, const FName& InCPPTypeObjectPath, bool bSetupUndoRedo, bool bSetupOrphanPins = true, bool bBreakLinks = true, bool bRemoveSubPins = true, bool bInitializeDefaultValue = true);
 	bool ChangePinType(URigVMPin* InPin, const FString& InCPPType,UObject* InCPPTypeObject, bool bSetupUndoRedo, bool bSetupOrphanPins = true, bool bBreakLinks = true, bool bRemoveSubPins = true, bool bInitializeDefaultValue = true);
@@ -1269,6 +1224,7 @@ protected:
 	FRigVMClientPatchResult PatchBranchNodesOnLoad();
 	FRigVMClientPatchResult PatchIfSelectNodesOnLoad();
 	FRigVMClientPatchResult PatchArrayNodesOnLoad();
+	FRigVMClientPatchResult PatchInvalidLinksOnWildcards();
 
 	// work to do after a duplication of the host asset
 	void PostDuplicateHost(const FString& InOldPathName, const FString& InNewPathName);
@@ -1333,7 +1289,6 @@ private:
 
 public:
 	
-	bool bSuspendRecomputingTemplateFilters;
 	bool bSuspendTemplateComputation;
 
 private:
diff --git a/Engine/Source/Developer/RigVMDeveloper/Public/RigVMModel/RigVMControllerActions.h b/Engine/Source/Developer/RigVMDeveloper/Public/RigVMModel/RigVMControllerActions.h
--- a/Engine/Source/Developer/RigVMDeveloper/Public/RigVMModel/RigVMControllerActions.h
+++ b/Engine/Source/Developer/RigVMDeveloper/Public/RigVMModel/RigVMControllerActions.h
@@ -1032,88 +1032,6 @@ public:
 	FString NewDefaultValue;
 };
 
-/**
- * An action setting the filtered permutations on a template node
- */
-USTRUCT()
-struct FRigVMSetTemplateFilteredPermutationsAction : public FRigVMBaseAction
-{
-	GENERATED_BODY()
-
-public:
-
-	FRigVMSetTemplateFilteredPermutationsAction() {}
-	FRigVMSetTemplateFilteredPermutationsAction(URigVMTemplateNode* InNode, const TArray<int32>& InOldFilteredPermutations);
-	virtual ~FRigVMSetTemplateFilteredPermutationsAction() {};
-	virtual UScriptStruct* GetScriptStruct() const override { return FRigVMSetTemplateFilteredPermutationsAction::StaticStruct(); }
-	virtual bool Merge(const FRigVMBaseAction* Other);
-	virtual bool Undo(URigVMController* InController) override;
-	virtual bool Redo(URigVMController* InController) override;
-
-	UPROPERTY()
-	FString NodePath;
-
-	UPROPERTY()
-	TArray<int32> OldFilteredPermutations;
-
-	UPROPERTY()
-	TArray<int32> NewFilteredPermutations;
-};
-
-/**
- * An action setting the filtered permutations on a template node
- */
-USTRUCT()
-struct FRigVMSetPreferredTemplatePermutationsAction : public FRigVMBaseAction
-{
-	GENERATED_BODY()
-
-public:
-
-	FRigVMSetPreferredTemplatePermutationsAction() {}
-	FRigVMSetPreferredTemplatePermutationsAction(URigVMTemplateNode* InNode, const TArray<FRigVMTemplatePreferredType>& InPreferredTypes);
-	virtual ~FRigVMSetPreferredTemplatePermutationsAction() {};
-	virtual UScriptStruct* GetScriptStruct() const override { return FRigVMSetTemplateFilteredPermutationsAction::StaticStruct(); }
-	virtual bool Merge(const FRigVMBaseAction* Other);
-	virtual bool Undo(URigVMController* InController) override;
-	virtual bool Redo(URigVMController* InController) override;
-
-	UPROPERTY()
-	FString NodePath;
-
-	UPROPERTY()
-	TArray<FRigVMTemplatePreferredType> OldPreferredPermutationTypes;
-
-	UPROPERTY()
-	TArray<FRigVMTemplatePreferredType> NewPreferredPermutationTypes;
-};
-
-/**
- * An action setting the filtered permutations on a template node
- */
-USTRUCT()
-struct FRigVMSetLibraryTemplateAction : public FRigVMBaseAction
-{
-	GENERATED_BODY()
-
-public:
-
-	FRigVMSetLibraryTemplateAction() {}
-	FRigVMSetLibraryTemplateAction(URigVMLibraryNode* InNode, FRigVMTemplate& InNewTemplate);
-	virtual UScriptStruct* GetScriptStruct() const override { return FRigVMSetTemplateFilteredPermutationsAction::StaticStruct(); }
-	virtual bool Undo(URigVMController* InController) override;
-	virtual bool Redo(URigVMController* InController) override;
-
-	UPROPERTY()
-	FString NodePath;
-
-	UPROPERTY()
-	TArray<uint8> OldTemplateBytes;
-
-	UPROPERTY()
-	TArray<uint8> NewTemplateBytes;
-};
-
 /**
  * An action inserting a new array pin in the graph.
  */
diff --git a/Engine/Source/Runtime/RigVM/Private/RigVMCore/RigVMTemplate.cpp b/Engine/Source/Runtime/RigVM/Private/RigVMCore/RigVMTemplate.cpp
--- a/Engine/Source/Runtime/RigVM/Private/RigVMCore/RigVMTemplate.cpp
+++ b/Engine/Source/Runtime/RigVM/Private/RigVMCore/RigVMTemplate.cpp
@@ -232,7 +232,7 @@ bool FRigVMTemplateArgument::SupportsTypeIndex(TRigVMTypeIndex InTypeIndex, TRig
 		{
 			(*OutTypeIndex) = TypeIndices[(*Permutations)[0]];
 		}
-		return true;		
+		return true;
 	}
 
 	// Try to find compatible type
@@ -245,11 +245,11 @@ bool FRigVMTemplateArgument::SupportsTypeIndex(TRigVMTypeIndex InTypeIndex, TRig
 			{
 				(*OutTypeIndex) = TypeIndices[(*Permutations)[0]];
 			}
-			return true;		
+			return true;
 		}
 	}
 
-	return false;	
+	return false;
 }
 
 bool FRigVMTemplateArgument::IsSingleton(const TArray<int32>& InPermutationIndices) const
@@ -808,24 +808,11 @@ FText FRigVMTemplate::GetTooltipText(const TArray<int32>& InPermutationIndices)
 		return true;
 	};
 
-	if(InPermutationIndices.IsEmpty())
-	{
-		for(int32 PermutationIndex = 0; PermutationIndex < Permutations.Num(); PermutationIndex++)
-		{
-			if(!VisitPermutation(PermutationIndex))
-			{
-				break;
-			}
-		}
-	}
-	else
+	for(int32 PermutationIndex = 0; PermutationIndex < Permutations.Num(); PermutationIndex++)
 	{
-		for(const int32 PermutationIndex : InPermutationIndices)
+		if(!VisitPermutation(PermutationIndex))
 		{
-			if(!VisitPermutation(PermutationIndex))
-			{
-				break;
-			}
+			break;
 		}
 	}
 
@@ -1139,7 +1126,7 @@ bool FRigVMTemplate::ArgumentSupportsTypeIndex(const FName& InArgumentName, TRig
 {
 	if (const FRigVMTemplateArgument* Argument = FindArgument(InArgumentName))
 	{
-		return Argument->SupportsTypeIndex(InTypeIndex, OutTypeIndex);		
+		return Argument->SupportsTypeIndex(InTypeIndex, OutTypeIndex);
 	}
 	return false;
 }
@@ -1317,23 +1304,6 @@ bool FRigVMTemplate::Resolve(FTypeMap& InOutTypes, TArray<int32>& OutPermutation
 			{
 				InOutTypes.Add(Argument.Name,MatchedType);
 
-				// if we found a perfect match - remove all permutations which don't match this one
-				if(bFoundPerfectMatch)
-				{
-					const TArray<int32> BeforePermutationIndices = OutPermutationIndices;
-					OutPermutationIndices.RemoveAll([Argument, MatchedType](int32 PermutationIndex) -> bool
-					{
-						return Argument.TypeIndices[PermutationIndex] != MatchedType;
-					});
-					if (OutPermutationIndices.IsEmpty() && bAllowFloatingPointCasts)
-					{
-						OutPermutationIndices = BeforePermutationIndices;
-						OutPermutationIndices.RemoveAll([Argument, MatchedType, InputType, &Registry](int32 PermutationIndex) -> bool
-						{
-							return !Registry.CanMatchTypes(Argument.TypeIndices[PermutationIndex], *InputType, true);
-						});
-					}
-				}
 				continue;
 			}
 		}
