Control Rig: Fix crash when using pin type of user defined struct that no longer exists

diff --git a/Engine/Source/Developer/RigVMDeveloper/Private/RigVMModel/RigVMController.cpp b/Engine/Source/Developer/RigVMDeveloper/Private/RigVMModel/RigVMController.cpp
--- a/Engine/Source/Developer/RigVMDeveloper/Private/RigVMModel/RigVMController.cpp
+++ b/Engine/Source/Developer/RigVMDeveloper/Private/RigVMModel/RigVMController.cpp
@@ -17366,6 +17366,11 @@ bool URigVMController::EnsurePinValidity(URigVMPin* InPin, bool bRecursive)
 			InPin->CPPType = RigVMTypeUtils::PostProcessCPPType(InPin->CPPType, InPin->GetCPPTypeObject());
 		}
 	}
+
+	if (InPin->GetCPPType().IsEmpty() || InPin->GetCPPType() == FName().ToString())
+	{
+		return false;
+	}
 	
 	if(bRecursive)
 	{
