Control Rig: Link creation performance improvements

diff --git a/Engine/Plugins/Animation/ControlRig/Source/ControlRigDeveloper/Private/Graph/ControlRigGraphSchema.cpp b/Engine/Plugins/Animation/ControlRig/Source/ControlRigDeveloper/Private/Graph/ControlRigGraphSchema.cpp
--- a/Engine/Plugins/Animation/ControlRig/Source/ControlRigDeveloper/Private/Graph/ControlRigGraphSchema.cpp
+++ b/Engine/Plugins/Animation/ControlRig/Source/ControlRigDeveloper/Private/Graph/ControlRigGraphSchema.cpp
@@ -827,6 +827,11 @@ static bool HasChildConnection_Recursive(const UEdGraphPin* InPin)
 
 const FPinConnectionResponse UControlRigGraphSchema::CanCreateConnection(const UEdGraphPin* A, const UEdGraphPin* B) const
 {
+	if (A == nullptr || B == nullptr)
+	{
+		return FPinConnectionResponse(CONNECT_RESPONSE_DISALLOW, TEXT("One of the Pins is NULL"));
+	}
+	
 	UControlRigGraphNode* RigNodeA = Cast<UControlRigGraphNode>(A->GetOwningNode());
 	UControlRigGraphNode* RigNodeB = Cast<UControlRigGraphNode>(B->GetOwningNode());
 
diff --git a/Engine/Plugins/Animation/ControlRig/Source/ControlRigEditor/Private/Graph/ControlRigConnectionDrawingPolicy.h b/Engine/Plugins/Animation/ControlRig/Source/ControlRigEditor/Private/Graph/ControlRigConnectionDrawingPolicy.h
--- a/Engine/Plugins/Animation/ControlRig/Source/ControlRigEditor/Private/Graph/ControlRigConnectionDrawingPolicy.h
+++ b/Engine/Plugins/Animation/ControlRig/Source/ControlRigEditor/Private/Graph/ControlRigConnectionDrawingPolicy.h
@@ -31,6 +31,7 @@ public:
 	/*out*/ FArrangedWidget*& StartWidgetGeometry,
 	/*out*/ FArrangedWidget*& EndWidgetGeometry) override;
 	virtual void DetermineWiringStyle(UEdGraphPin* OutputPin, UEdGraphPin* InputPin, /*inout*/ FConnectionParams& Params) override;
+	virtual bool UseDrawStateCaching() const override { return true; }
 
 private:
 	// Each time a reroute node is encountered, input geometry is compared to output geometry to see if the pins on the reroute node need to be reversed
diff --git a/Engine/Source/Developer/RigVMDeveloper/Private/RigVMCompiler/RigVMAST.cpp b/Engine/Source/Developer/RigVMDeveloper/Private/RigVMCompiler/RigVMAST.cpp
--- a/Engine/Source/Developer/RigVMDeveloper/Private/RigVMCompiler/RigVMAST.cpp
+++ b/Engine/Source/Developer/RigVMDeveloper/Private/RigVMCompiler/RigVMAST.cpp
@@ -1907,8 +1907,8 @@ bool FRigVMParserAST::CanLink(URigVMPin* InSourcePin, URigVMPin* InTargetPin, FS
 		return true;
 	}
 
-	FRigVMASTProxy SourceNodeProxy = FRigVMASTProxy::MakeFromUObject(SourceNode);
-	FRigVMASTProxy TargetNodeProxy = FRigVMASTProxy::MakeFromUObject(TargetNode);
+	const FRigVMASTProxy SourceNodeProxy = FRigVMASTProxy::MakeFromUObject(SourceNode);
+	const FRigVMASTProxy TargetNodeProxy = FRigVMASTProxy::MakeFromUObject(TargetNode);
 
 	const FRigVMExprAST* SourceExpression = nullptr;
 	if (FRigVMExprAST* const* SourceExpressionPtr = SubjectToExpression.Find(SourceNodeProxy))
@@ -1959,19 +1959,46 @@ bool FRigVMParserAST::CanLink(URigVMPin* InSourcePin, URigVMPin* InTargetPin, FS
 	{
 		return true;
 	}
+	const FRigVMBlockExprAST* SourceRoot = SourceBlock->GetRootBlock();
+	const FRigVMBlockExprAST* TargetRoot = TargetBlock->GetRootBlock();
+	
+	bool bNeedsCycleChecking = (SourceBlock == TargetBlock);
 
-	if (SourceBlock == TargetBlock ||
-		SourceBlock->Contains(TargetBlock) ||
-		TargetBlock->Contains(SourceBlock) ||
-		TargetBlock->GetRootBlock()->Contains(SourceBlock) ||
-		SourceBlock->GetRootBlock()->Contains(TargetBlock))
+	// check if source block/root contains the target
+	if (!bNeedsCycleChecking)
 	{
-		if (SourceVarExpression)
+		TMap<const FRigVMExprAST*, bool> SourceCache;
+
+		// check root first
+		bNeedsCycleChecking = SourceRoot->Contains(TargetBlock, &SourceCache);
+
+		// check block if needed (avoid doing it twice)
+		if (!bNeedsCycleChecking && SourceBlock != SourceRoot)
 		{
-			if (SourceVarExpression->SupportsSoftLinks())
-			{
-				return true;
-			}
+			bNeedsCycleChecking = SourceBlock->Contains(TargetBlock, &SourceCache);
+		}		
+	}
+
+	// check if target block/root contains the source
+	if (!bNeedsCycleChecking)
+	{
+		TMap<const FRigVMExprAST*, bool> TargetCache;
+
+		// check root first
+		bNeedsCycleChecking = TargetRoot->Contains(SourceBlock, &TargetCache);
+
+		// check block if needed (avoid doing it twice)
+		if (!bNeedsCycleChecking && TargetBlock != TargetRoot)
+		{
+			bNeedsCycleChecking = TargetBlock->Contains(SourceBlock, &TargetCache);
+		}
+	}
+
+	if (bNeedsCycleChecking)
+	{
+		if (SourceVarExpression && SourceVarExpression->SupportsSoftLinks())
+		{
+			return true;
 		}
 
 		if (LastCycleCheckExpr != SourceExpression && LastCycleCheckExpr != TargetExpression)
@@ -2043,8 +2070,7 @@ bool FRigVMParserAST::CanLink(URigVMPin* InSourcePin, URigVMPin* InTargetPin, FS
 	{
 		// if one of the blocks is not part of the current 
 		// execution - that's fine.
-		if (SourceBlock->GetRootBlock()->ContainsEntry() != 
-			TargetBlock->GetRootBlock()->ContainsEntry())
+		if (SourceRoot->ContainsEntry() != TargetRoot->ContainsEntry())
 		{
 			return true;
 		}
diff --git a/Engine/Source/Developer/RigVMDeveloper/Private/RigVMModel/RigVMGraph.cpp b/Engine/Source/Developer/RigVMDeveloper/Private/RigVMModel/RigVMGraph.cpp
--- a/Engine/Source/Developer/RigVMDeveloper/Private/RigVMModel/RigVMGraph.cpp
+++ b/Engine/Source/Developer/RigVMDeveloper/Private/RigVMModel/RigVMGraph.cpp
@@ -444,12 +444,96 @@ void URigVMGraph::PrepareCycleChecking(URigVMPin* InPin, bool bAsInput)
 	GetDiagnosticsAST(false, LinksToSkip)->PrepareCycleChecking(InPin);
 }
 
+
+{
+	bool AreNodesLinked(const URigVMNode* InNodeToGetLinksFrom, const URigVMNode* InNodeToCheck, TSet<URigVMNode*>& OutNodes, const bool bSource)
+	{
+		const TArray<URigVMLink*> RigVMLinks = InNodeToGetLinksFrom->GetLinks();
+		for (const URigVMLink* Link: RigVMLinks)
+		{
+			const URigVMPin* Pin = bSource ? Link->GetSourcePin() : Link->GetTargetPin();
+			if (URigVMNode* Node = Pin ? Pin->GetNode() : nullptr)
+			{
+				if (Node != InNodeToGetLinksFrom && !OutNodes.Contains(Node))
+				{
+					if (Node == InNodeToCheck)
+					{
+						return true;
+					}
+					OutNodes.Add(Node);
+					
+					if (AreNodesLinked(Node, InNodeToCheck, OutNodes, bSource))
+					{
+						return true;
+					}
+				}
+			}
+		}
+		return false;
+	}
+
+	bool AreNodesGoingToCycle(const URigVMNode* InSourceNode, const URigVMNode* InTargetNode)
+	{
+		TSet<URigVMNode*> CheckedNodes;
+
+		const bool bIsTargetAlreadyASource = AreNodesLinked(InSourceNode, InTargetNode, CheckedNodes, true);
+		if (bIsTargetAlreadyASource)
+		{
+			return true;
+		}
+
+		CheckedNodes.Reset();
+		const bool bIsSourceAlreadyATarget = AreNodesLinked(InTargetNode, InSourceNode, CheckedNodes, false);
+		if(bIsSourceAlreadyATarget)
+		{
+			return true;
+		}
+
+		return false;
+	}
+
+	bool AreNodesAlreadyConnected(const URigVMNode* InSourceNode, const URigVMNode* InTargetNode)
+	{
+		TSet<URigVMNode*> CheckedNodes;
+
+		const bool bIsTargetAlreadyATarget = AreNodesLinked(InSourceNode, InTargetNode, CheckedNodes, false);
+		if (bIsTargetAlreadyATarget)
+		{
+			return true;
+		}
+
+		CheckedNodes.Reset();
+		const bool bIsSourceAlreadyASource = AreNodesLinked(InTargetNode, InSourceNode, CheckedNodes, true);
+		if (bIsSourceAlreadyASource)
+		{
+			return true;
+		}
+
+		return false;
+	}
+}
+
 bool URigVMGraph::CanLink(URigVMPin* InSourcePin, URigVMPin* InTargetPin, FString* OutFailureReason, const FRigVMByteCode* InByteCode, ERigVMPinDirection InUserLinkDirection, bool bEnableTypeCasting)
 {
 	if (!URigVMPin::CanLink(InSourcePin, InTargetPin, OutFailureReason, InByteCode, InUserLinkDirection, false, bEnableTypeCasting))
 	{
 		return false;
 	}
+
+	const URigVMNode* SourceNode = InSourcePin->GetNode();
+	const URigVMNode* TargetNode = InTargetPin->GetNode();
+
+	// check for cycling nodes dependencies before going into the AST 
+	if (AreNodesGoingToCycle(SourceNode, TargetNode))
+	{
+		if (OutFailureReason)
+		{
+			*OutFailureReason = TEXT("Cycles are not allowed.");
+		}
+		return false;
+	}
+
+	// full check needed
 	return GetDiagnosticsAST()->CanLink(InSourcePin, InTargetPin, OutFailureReason);
 }
 
diff --git a/Engine/Source/Editor/GraphEditor/Private/BlueprintConnectionDrawingPolicy.cpp b/Engine/Source/Editor/GraphEditor/Private/BlueprintConnectionDrawingPolicy.cpp
--- a/Engine/Source/Editor/GraphEditor/Private/BlueprintConnectionDrawingPolicy.cpp
+++ b/Engine/Source/Editor/GraphEditor/Private/BlueprintConnectionDrawingPolicy.cpp
@@ -610,17 +610,27 @@ void FKismetConnectionDrawingPolicy::DetermineWiringStyle(UEdGraphPin* OutputPin
 
 void FKismetConnectionDrawingPolicy::SetIncompatiblePinDrawState(const TSharedPtr<SGraphPin>& StartPin, const TSet< TSharedRef<SWidget> >& VisiblePins)
 {
+	if (VisiblePins.IsEmpty())
+	{
+		return;
+	}
+	
+	static const FLinearColor DisallowColor(0.25f, 0.25f, 0.25f, 0.5f);
+	
 	ResetIncompatiblePinDrawState(VisiblePins);
 
+	const UEdGraphPin* StartPinObj = StartPin->GetPinObj();
+	const UEdGraphSchema* EdGraphSchema = StartPinObj->GetSchema();
+	
 	for (auto VisiblePinIterator = VisiblePins.CreateConstIterator(); VisiblePinIterator; ++VisiblePinIterator)
 	{
 		TSharedPtr<SGraphPin> CheckPin = StaticCastSharedRef<SGraphPin>(*VisiblePinIterator);
 		if (CheckPin != StartPin)
 		{
-			const FPinConnectionResponse Response = StartPin->GetPinObj()->GetSchema()->CanCreateConnection(StartPin->GetPinObj(), CheckPin->GetPinObj());
+			const FPinConnectionResponse Response = EdGraphSchema->CanCreateConnection(StartPinObj, CheckPin->GetPinObj());
 			if (Response.Response == CONNECT_RESPONSE_DISALLOW) //-V1051
 			{
-				CheckPin->SetPinColorModifier(FLinearColor(0.25f, 0.25f, 0.25f, 0.5f));
+				CheckPin->SetPinColorModifier(DisallowColor);
 			}
 		}
 	}
diff --git a/Engine/Source/Editor/GraphEditor/Private/SGraphPanel.cpp b/Engine/Source/Editor/GraphEditor/Private/SGraphPanel.cpp
--- a/Engine/Source/Editor/GraphEditor/Private/SGraphPanel.cpp
+++ b/Engine/Source/Editor/GraphEditor/Private/SGraphPanel.cpp
@@ -481,6 +481,8 @@ int32 SGraphPanel::OnPaint( const FPaintArgs& Args, const FGeometry& AllottedGeo
 			ConnectionDrawingPolicy = FNodeFactory::CreateConnectionPolicy(Schema, WireLayerId, MaxLayerId, ZoomFactor, MyCullingRect, OutDrawElements, GraphObj);
 		}
 
+		const bool bUseDrawStateCaching = ConnectionDrawingPolicy->UseDrawStateCaching();
+		
 		TArray<TSharedPtr<SGraphPin>> OverridePins;
 		for (const FGraphPinHandle& Handle : PreviewConnectorFromPins)
 		{
@@ -575,9 +577,13 @@ int32 SGraphPanel::OnPaint( const FPaintArgs& Args, const FGeometry& AllottedGeo
 					ConnectionDrawingPolicy->DrawPreviewConnector(PinGeometry->Geometry, StartPoint, EndPoint, CurrentStartPin.Get()->GetPinObj());
 				}
 
-				//@TODO: Re-evaluate this incompatible mojo; it's mutating every pin state every frame to accomplish a visual effect
-				ConnectionDrawingPolicy->SetIncompatiblePinDrawState(CurrentStartPin, VisiblePins);
+				if (!bUseDrawStateCaching || !bIsDrawStateCached)
+				{
+					//@TODO: Re-evaluate this incompatible mojo; it's mutating every pin state every frame to accomplish a visual effect
+					ConnectionDrawingPolicy->SetIncompatiblePinDrawState(CurrentStartPin, VisiblePins);
+				}
 			}
+			bIsDrawStateCached = true;
 		}
 		else
 		{
@@ -1295,6 +1301,7 @@ void SGraphPanel::OnBeginMakingConnection(FGraphPinHandle PinHandle)
 		DismissContextMenu();
 
 		PreviewConnectorFromPins.Add(PinHandle);
+		bIsDrawStateCached = false;
 	}
 }
 
@@ -1304,6 +1311,7 @@ void SGraphPanel::OnStopMakingConnection(bool bForceStop)
 	{
 		PreviewConnectorFromPins.Reset();
 		bPreservePinPreviewConnection = false;
+		bIsDrawStateCached = false;
 	}
 }
 
diff --git a/Engine/Source/Editor/GraphEditor/Public/ConnectionDrawingPolicy.h b/Engine/Source/Editor/GraphEditor/Public/ConnectionDrawingPolicy.h
--- a/Engine/Source/Editor/GraphEditor/Public/ConnectionDrawingPolicy.h
+++ b/Engine/Source/Editor/GraphEditor/Public/ConnectionDrawingPolicy.h
@@ -175,6 +175,9 @@ public:
 		/*out*/ FArrangedWidget*& EndWidgetGeometry
 		);
 
+	// Choose whether we want to cache the pins draw state to avoid resetting it for every tick 
+	virtual bool UseDrawStateCaching() const { return false; }
+	
 	virtual void SetIncompatiblePinDrawState(const TSharedPtr<SGraphPin>& StartPin, const TSet< TSharedRef<SWidget> >& VisiblePins);
 	virtual void ResetIncompatiblePinDrawState(const TSet< TSharedRef<SWidget> >& VisiblePins);
 
diff --git a/Engine/Source/Editor/GraphEditor/Public/SGraphPanel.h b/Engine/Source/Editor/GraphEditor/Public/SGraphPanel.h
--- a/Engine/Source/Editor/GraphEditor/Public/SGraphPanel.h
+++ b/Engine/Source/Editor/GraphEditor/Public/SGraphPanel.h
@@ -286,6 +286,7 @@ protected:
 	/** Sometimes the panel draws a preview connector; e.g. when the user is connecting pins */
 	TArray< FGraphPinHandle > PreviewConnectorFromPins;
 	FVector2D PreviewConnectorEndpoint;
+	mutable bool bIsDrawStateCached = false;
 
 	/** Last mouse position seen, used for paint-centric highlighting */
 	FVector2D SavedMousePosForOnPaintEventLocalSpace;
