Control Rig: Repopulate pins to deal with missing ref functions

diff --git a/Engine/Source/Developer/RigVMDeveloper/Private/RigVMModel/RigVMController.cpp b/Engine/Source/Developer/RigVMDeveloper/Private/RigVMModel/RigVMController.cpp
--- a/Engine/Source/Developer/RigVMDeveloper/Private/RigVMModel/RigVMController.cpp
+++ b/Engine/Source/Developer/RigVMDeveloper/Private/RigVMModel/RigVMController.cpp
@@ -14875,6 +14875,12 @@ void URigVMController::RepopulatePinsOnNode(URigVMNode* InNode, bool bFollowCore
 		{
 			NewPinInfos = FPinInfoArray(FunctionHeader, this, &PreviousPinInfos);
 		}
+		// if we can't find referenced node anymore
+		// let's keep all pins
+		else
+		{
+			NewPinInfos = FPinInfoArray(FunctionRefNode, this, &PreviousPinInfos);
+		}
 	}
 	else
 	{
@@ -15307,7 +15313,7 @@ void URigVMController::RepopulatePinsOnNode(URigVMNode* InNode, bool bFollowCore
 	}
 	else if(CollapseNode)
 	{
-		if (CollapseNode->GetOuter()->IsA<URigVMFunctionLibrary>())
+		if (!CollapseNode->GetOuter()->IsA<URigVMFunctionLibrary>())
 		{
 			// no need to notify since the function library graph is invisible anyway
 			RemoveUnusedOrphanedPins(CollapseNode);
@@ -15698,6 +15704,10 @@ void URigVMController::ApplyPinStates(URigVMNode* InNode, const TMap<FString, UR
 		{
 			for (URigVMInjectionInfo* InjectionInfo : PinState.InjectionInfos)
 			{
+				if(URigVMPin* OuterPin = Cast<URigVMPin>(InjectionInfo->GetOuter()))
+				{
+					OuterPin->InjectionInfos.Remove(InjectionInfo);
+				}
 				RenameObject(InjectionInfo->Node, nullptr, InNode->GetGraph());
 				DestroyObject(InjectionInfo);
 			}
