RigVM: Avoid type dialog for pins supporting both float and double

diff --git a/Engine/Source/Developer/RigVMDeveloper/Private/RigVMModel/RigVMController.cpp b/Engine/Source/Developer/RigVMDeveloper/Private/RigVMModel/RigVMController.cpp
--- a/Engine/Source/Developer/RigVMDeveloper/Private/RigVMModel/RigVMController.cpp
+++ b/Engine/Source/Developer/RigVMDeveloper/Private/RigVMModel/RigVMController.cpp
@@ -16782,6 +16782,46 @@ bool URigVMController::PrepareToLink(URigVMPin* FirstToResolve, URigVMPin* Secon
 		return false;
 	}
 
+	// reduce matching types by duplicate entries for float / double
+	{
+		TArray<TRigVMTypeIndex> FilteredMatchingTypes;
+		FilteredMatchingTypes.Reserve(MatchingTypes.Num());
+		for(const TRigVMTypeIndex& MatchingType : MatchingTypes)
+		{
+			// special case float singe & array
+			if(MatchingType == RigVMTypeUtils::TypeIndex::Float)
+			{
+				if(MatchingTypes.Contains(RigVMTypeUtils::TypeIndex::Double))
+				{
+					continue;
+				}
+			}
+			if(MatchingType == RigVMTypeUtils::TypeIndex::FloatArray)
+			{
+				if(MatchingTypes.Contains(RigVMTypeUtils::TypeIndex::DoubleArray))
+				{
+					continue;
+				}
+			}
+
+			bool bAlreadyContainsMatch = false;
+			for(const TRigVMTypeIndex& FilteredType : FilteredMatchingTypes)
+			{
+				if(Registry.CanMatchTypes(MatchingType, FilteredType, true))
+				{
+					bAlreadyContainsMatch = true;
+					break;
+				}
+			}
+			if(bAlreadyContainsMatch)
+			{
+				continue;
+			}
+			FilteredMatchingTypes.Add(MatchingType);
+		}
+		Swap(FilteredMatchingTypes, MatchingTypes);
+	}
+
 	TRigVMTypeIndex FinalType = INDEX_NONE;
 	if (MatchingTypes.Num() > 1)
 	{
