RigVM: Support FMatrix in code generator

diff --git a/Engine/Source/Developer/RigVMDeveloper/Private/RigVMCompiler/RigVMCodeGenerator.cpp b/Engine/Source/Developer/RigVMDeveloper/Private/RigVMCompiler/RigVMCodeGenerator.cpp
--- a/Engine/Source/Developer/RigVMDeveloper/Private/RigVMCompiler/RigVMCodeGenerator.cpp
+++ b/Engine/Source/Developer/RigVMDeveloper/Private/RigVMCompiler/RigVMCodeGenerator.cpp
@@ -806,7 +806,9 @@ FString FRigVMCodeGenerator::DumpInstructions(const FString& InPrefix, const TAr
 				{
 					const FString CastPrefix = TargetCPPType != SourceCPPType ? Format(RigVM_BracesFormat, *TargetCPPType) : FString();  
 					
-					if(TargetOperand.EndsWith(TEXT("(")))
+					if(TargetOperand.EndsWith(TEXT("(")) ||
+						TargetOperand.EndsWith(TEXT(",")) ||
+						TargetOperand.EndsWith(TEXT(", ")))
 					{
 						Lines.Add(Prefix + Format(RigVM_CopyOpMethodFormat, *TargetOperand, *CastPrefix, *SourceOperand));
 					}
@@ -1780,6 +1782,28 @@ FString FRigVMCodeGenerator::GetOperandName(const FRigVMOperand& InOperand, bool
 			TransformMap.Add(TEXT("Scale3D.X"), TTuple<FString, FString>(TEXT("{0}.GetScale3D().X"), TEXT("FTransformSetter({0}).SetScaleX(")));
 			TransformMap.Add(TEXT("Scale3D.Y"), TTuple<FString, FString>(TEXT("{0}.GetScale3D().Y"), TEXT("FTransformSetter({0}).SetScaleY(")));
 			TransformMap.Add(TEXT("Scale3D.Z"), TTuple<FString, FString>(TEXT("{0}.GetScale3D().Z"), TEXT("FTransformSetter({0}).SetScaleZ(")));
+
+			TMap<FString, TTuple<FString, FString>>& MatrixMap = MappedProperties.Add(TBaseStructure<FMatrix>::Get());
+			MatrixMap.Add(TEXT("XPlane"), TTuple<FString, FString>(TEXT("(*(FPlane*){0}.M[0])"), TEXT("FMatrixSetter({0}).SetPlane(0, ")));
+			MatrixMap.Add(TEXT("XPlane.X"), TTuple<FString, FString>(TEXT("{0}.M[0][0]"), TEXT("FMatrixSetter({0}).SetComponent(0, 0, ")));
+			MatrixMap.Add(TEXT("XPlane.Y"), TTuple<FString, FString>(TEXT("{0}.M[0][1]"), TEXT("FMatrixSetter({0}).SetComponent(0, 1, ")));
+			MatrixMap.Add(TEXT("XPlane.Z"), TTuple<FString, FString>(TEXT("{0}.M[0][2]"), TEXT("FMatrixSetter({0}).SetComponent(0, 2, ")));
+			MatrixMap.Add(TEXT("XPlane.W"), TTuple<FString, FString>(TEXT("{0}.M[0][3]"), TEXT("FMatrixSetter({0}).SetComponent(0, 3, ")));
+			MatrixMap.Add(TEXT("YPlane"), TTuple<FString, FString>(TEXT("(*(FPlane*){0}.M[1])"), TEXT("FMatrixSetter({0}).SetPlane(1, ")));
+			MatrixMap.Add(TEXT("YPlane.X"), TTuple<FString, FString>(TEXT("{0}.M[1][0]"), TEXT("FMatrixSetter({0}).SetComponent(1, 0, ")));
+			MatrixMap.Add(TEXT("YPlane.Y"), TTuple<FString, FString>(TEXT("{0}.M[1][1]"), TEXT("FMatrixSetter({0}).SetComponent(1, 1, ")));
+			MatrixMap.Add(TEXT("YPlane.Z"), TTuple<FString, FString>(TEXT("{0}.M[1][2]"), TEXT("FMatrixSetter({0}).SetComponent(1, 2, ")));
+			MatrixMap.Add(TEXT("YPlane.W"), TTuple<FString, FString>(TEXT("{0}.M[1][3]"), TEXT("FMatrixSetter({0}).SetComponent(1, 3, ")));
+			MatrixMap.Add(TEXT("ZPlane"), TTuple<FString, FString>(TEXT("(*(FPlane*){0}.M[2])"), TEXT("FMatrixSetter({0}).SetPlane(2, ")));
+			MatrixMap.Add(TEXT("ZPlane.X"), TTuple<FString, FString>(TEXT("{0}.M[2][0]"), TEXT("FMatrixSetter({0}).SetComponent(2, 0, ")));
+			MatrixMap.Add(TEXT("ZPlane.Y"), TTuple<FString, FString>(TEXT("{0}.M[2][1]"), TEXT("FMatrixSetter({0}).SetComponent(2, 1, ")));
+			MatrixMap.Add(TEXT("ZPlane.Z"), TTuple<FString, FString>(TEXT("{0}.M[2][2]"), TEXT("FMatrixSetter({0}).SetComponent(2, 2, ")));
+			MatrixMap.Add(TEXT("ZPlane.W"), TTuple<FString, FString>(TEXT("{0}.M[2][3]"), TEXT("FMatrixSetter({0}).SetComponent(2, 3, ")));
+			MatrixMap.Add(TEXT("WPlane"), TTuple<FString, FString>(TEXT("(*(FPlane*){0}.M[3])"), TEXT("FMatrixSetter({0}).SetPlane(3, ")));
+			MatrixMap.Add(TEXT("WPlane.X"), TTuple<FString, FString>(TEXT("{0}.M[3][0]"), TEXT("FMatrixSetter({0}).SetComponent(3, 0, ")));
+			MatrixMap.Add(TEXT("WPlane.Y"), TTuple<FString, FString>(TEXT("{0}.M[3][1]"), TEXT("FMatrixSetter({0}).SetComponent(3, 1, ")));
+			MatrixMap.Add(TEXT("WPlane.Z"), TTuple<FString, FString>(TEXT("{0}.M[3][2]"), TEXT("FMatrixSetter({0}).SetComponent(3, 2, ")));
+			MatrixMap.Add(TEXT("WPlane.W"), TTuple<FString, FString>(TEXT("{0}.M[3][3]"), TEXT("FMatrixSetter({0}).SetComponent(3, 3, ")));
 		}
 
 		const FProperty* CurrentProperty = Property.Property;
@@ -2404,6 +2428,10 @@ const TMap<FName, FRigVMCodeGenerator::TStructConstGenerator>& FRigVMCodeGenerat
 	static constexpr TCHAR Roll[] = TEXT("Roll");
 	static constexpr TCHAR Type[] = TEXT("Type");
 	static constexpr TCHAR Name[] = TEXT("Name");
+	static constexpr TCHAR XPlane[] = TEXT("XPlane");
+	static constexpr TCHAR YPlane[] = TEXT("YPlane");
+	static constexpr TCHAR ZPlane[] = TEXT("ZPlane");
+	static constexpr TCHAR WPlane[] = TEXT("WPlane");
 
 	StructConstGenerators.Add(TEXT("FVector2D"), [](const FString& InDefault) -> FString
 	{
@@ -2511,5 +2539,31 @@ const TMap<FName, FRigVMCodeGenerator::TStructConstGenerator>& FRigVMCodeGenerat
 		return FString();
 	});
 
+	StructConstGenerators.Add(TEXT("FPlane"), [](const FString& InDefault) -> FString
+	{
+		static constexpr TCHAR Constructor[] = TEXT("FPlane({0}, {1}, {2}, {3})");
+		const FStringMap Defaults = DefaultValueHelpers::SplitIntoMap(InDefault);
+		if(Defaults.Contains(X) && Defaults.Contains(Y) && Defaults.Contains(Z) && Defaults.Contains(W))
+		{
+			return Format(Constructor, Defaults.FindChecked(X), Defaults.FindChecked(Y), Defaults.FindChecked(Z), Defaults.FindChecked(W));
+		}
+		return FString();
+	});
+
+	StructConstGenerators.Add(TEXT("FMatrix"), [](const FString& InDefault) -> FString
+	{
+		static constexpr TCHAR Constructor[] = TEXT("FMatrix({0}, {1}, {2}, {3})");
+		const FStringMap Defaults = DefaultValueHelpers::SplitIntoMap(InDefault);
+		if(Defaults.Contains(XPlane) && Defaults.Contains(YPlane) && Defaults.Contains(ZPlane) && Defaults.Contains(WPlane))
+		{
+			const FString XPlaneValue = StructConstGenerators.FindChecked(TEXT("FPlane"))(Defaults.FindChecked(XPlane));
+			const FString YPlaneValue = StructConstGenerators.FindChecked(TEXT("FPlane"))(Defaults.FindChecked(YPlane));
+			const FString ZPlaneValue = StructConstGenerators.FindChecked(TEXT("FPlane"))(Defaults.FindChecked(ZPlane));
+			const FString WPlaneValue = StructConstGenerators.FindChecked(TEXT("FPlane"))(Defaults.FindChecked(WPlane));
+			return Format(Constructor, XPlaneValue, YPlaneValue, ZPlaneValue, WPlaneValue);
+		}
+		return FString();
+	});
+
 	return StructConstGenerators;
 }
diff --git a/Engine/Source/Runtime/RigVM/Public/RigVMCore/RigVMNativized.h b/Engine/Source/Runtime/RigVM/Public/RigVMCore/RigVMNativized.h
--- a/Engine/Source/Runtime/RigVM/Public/RigVMCore/RigVMNativized.h
+++ b/Engine/Source/Runtime/RigVM/Public/RigVMCore/RigVMNativized.h
@@ -470,6 +470,31 @@ protected:
 		FTransform& Transform;
 	};
 
+	class FMatrixSetter
+	{
+	public:
+		FMatrixSetter(FMatrix& InOutMatrix)
+			: Matrix(InOutMatrix)
+		{
+		}
+
+		void SetPlane(int32 InPlaneIndex, const FPlane& InPlane)
+		{
+			Matrix.M[InPlaneIndex][0] = InPlane.X;
+			Matrix.M[InPlaneIndex][1] = InPlane.Y;
+			Matrix.M[InPlaneIndex][2] = InPlane.Z;
+			Matrix.M[InPlaneIndex][3] = InPlane.W;
+		}
+
+		void SetComponent(int32 InPlaneIndex, int32 InComponentIndex, double InValue)
+		{
+			Matrix.M[InPlaneIndex][InComponentIndex] = InValue;
+		}
+
+	private:
+		FMatrix& Matrix;
+	};
+
 	template <typename T,
 			  typename TEnableIf<!TIsArray<T>::Value, int>::Type = 0>
 	bool ImportDefaultValue(T& OutValue, const FProperty* Property, const FString& InBuffer)
