Control Rig: Fix unresolving execute pins

diff --git a/Engine/Source/Developer/RigVMDeveloper/Private/RigVMModel/RigVMController.cpp b/Engine/Source/Developer/RigVMDeveloper/Private/RigVMModel/RigVMController.cpp
--- a/Engine/Source/Developer/RigVMDeveloper/Private/RigVMModel/RigVMController.cpp
+++ b/Engine/Source/Developer/RigVMDeveloper/Private/RigVMModel/RigVMController.cpp
@@ -16482,6 +16482,8 @@ bool URigVMController::UpdateTemplateNodePinTypes(URigVMTemplateNode* InNode, bo
 		return false;
 	}
 
+	const FRigVMRegistry& Registry = FRigVMRegistry::Get();
+
 	if (InNode->IsA<URigVMFunctionEntryNode>() || InNode->IsA<URigVMFunctionReturnNode>())
 	{
 		return false;
@@ -16510,7 +16512,13 @@ bool URigVMController::UpdateTemplateNodePinTypes(URigVMTemplateNode* InNode, bo
 			}
 		}
 		return false;
-	});	
+	});
+
+	FRigVMDispatchContext DispatchContext;
+	if(const URigVMDispatchNode* DispatchNode = Cast<URigVMDispatchNode>(InNode))
+	{
+		DispatchContext = DispatchNode->GetDispatchContext();
+	}
 
 	// Find the types for each possible permutation
 	TMap<int32, TArray<TRigVMTypeIndex>> PinTypes; // permutation to pin types
@@ -16532,6 +16540,11 @@ bool URigVMController::UpdateTemplateNodePinTypes(URigVMTemplateNode* InNode, bo
 				Types.Add(Argument->TypeIndices[ResolvedPermutation]);
 				bAddedType = true;
 			}
+			else if (const FRigVMExecuteArgument* ExecuteArgument = Template->FindExecuteArgument(Pin->GetFName(), DispatchContext))
+			{
+				Types.Add(ExecuteArgument->TypeIndex);
+				bAddedType = true;
+			}
 			
 			if (!bAddedType)
 			{
@@ -16682,7 +16695,7 @@ bool URigVMController::UpdateTemplateNodePinTypes(URigVMTemplateNode* InNode, bo
 					bAnyTypeChanged = !Pin->IsExecuteContext();
 					if(bAnyTypeChanged)
 					{
-						const FString CPPType = FRigVMRegistry::Get().GetType(FinalPinTypes[PinIndex]).CPPType.ToString();
+						const FString CPPType = Registry.GetType(FinalPinTypes[PinIndex]).CPPType.ToString();
 						ReportPinTypeChange(Pin, CPPType);
 					}
 					ChangePinType(Pin, FinalPinTypes[PinIndex], bSetupUndoRedo, false, false, false, bInitializeDefaultValue);
