Control Rig: Update UI node after creating sub pins

diff --git a/Engine/Plugins/Animation/ControlRig/Source/ControlRigDeveloper/Private/Graph/ControlRigGraphNode.cpp b/Engine/Plugins/Animation/ControlRig/Source/ControlRigDeveloper/Private/Graph/ControlRigGraphNode.cpp
--- a/Engine/Plugins/Animation/ControlRig/Source/ControlRigDeveloper/Private/Graph/ControlRigGraphNode.cpp
+++ b/Engine/Plugins/Animation/ControlRig/Source/ControlRigDeveloper/Private/Graph/ControlRigGraphNode.cpp
@@ -703,6 +703,8 @@ bool UControlRigGraphNode::CreateGraphPinFromModelPin(const URigVMPin* InModelPi
 			{
 				CreateGraphPinFromModelPin(ModelSubPin, InDirection, GraphPin);
 			}
+
+			SetupPinDefaultsFromModel(GraphPin, InModelPin);
 		}
 
 		return GraphPin;
diff --git a/Engine/Source/Developer/RigVMDeveloper/Private/RigVMModel/RigVMController.cpp b/Engine/Source/Developer/RigVMDeveloper/Private/RigVMModel/RigVMController.cpp
--- a/Engine/Source/Developer/RigVMDeveloper/Private/RigVMModel/RigVMController.cpp
+++ b/Engine/Source/Developer/RigVMDeveloper/Private/RigVMModel/RigVMController.cpp
@@ -15896,7 +15896,7 @@ void URigVMController::CreateDefaultValueForStructIfRequired(UScriptStruct* InSt
 		// before ExportText() because ExportText() appends to it.
 		InOutDefaultValue.Reset();
 
-		InStruct->ExportText(InOutDefaultValue, TempBuffer.GetData(), nullptr, nullptr, PPF_None, nullptr);
+		InStruct->ExportText(InOutDefaultValue, TempBuffer.GetData(), TempBuffer.GetData(), nullptr, PPF_None, nullptr);
 		InStruct->DestroyStruct(TempBuffer.GetData());
 	}
 }
