Control Rig: Fix folding of connection between entry and return nodes inside function

diff --git a/Engine/Source/Developer/RigVMDeveloper/Private/RigVMCompiler/RigVMAST.cpp b/Engine/Source/Developer/RigVMDeveloper/Private/RigVMCompiler/RigVMAST.cpp
--- a/Engine/Source/Developer/RigVMDeveloper/Private/RigVMCompiler/RigVMAST.cpp
+++ b/Engine/Source/Developer/RigVMDeveloper/Private/RigVMCompiler/RigVMAST.cpp
@@ -1220,6 +1220,17 @@ FRigVMExprAST* FRigVMParserAST::TraverseLink(int32 InLinkIndex, FRigVMExprAST* I
 			bRequiresCopy = true;
 		}
 	}
+	if (!bRequiresCopy)
+	{
+		// Connections between entry and return in a function requires a copy
+		if (LibraryNodeBeingCompiled)
+		{
+			if (SourceRootPin->GetNode()->IsA<URigVMFunctionEntryNode>() && TargetRootPin->GetNode()->IsA<URigVMFunctionReturnNode>())
+			{
+				bRequiresCopy = true;
+			}
+		}
+	}
 
 	FRigVMAssignExprAST* AssignExpr = nullptr;
 	if (bRequiresCopy)
diff --git a/Engine/Source/Developer/RigVMDeveloper/Private/RigVMCompiler/RigVMCompiler.cpp b/Engine/Source/Developer/RigVMDeveloper/Private/RigVMCompiler/RigVMCompiler.cpp
--- a/Engine/Source/Developer/RigVMDeveloper/Private/RigVMCompiler/RigVMCompiler.cpp
+++ b/Engine/Source/Developer/RigVMDeveloper/Private/RigVMCompiler/RigVMCompiler.cpp
@@ -3261,6 +3261,7 @@ const FRigVMCompilerWorkData::FRigVMASTProxyArray& URigVMCompiler::FindProxiesWi
 
 	const FString CPPType = InProxy.GetSubjectChecked<URigVMPin>()->GetCPPType();
 
+	bool bContainsInterfacePin = false;
 	for(int32 ProxyIndex = 0; ProxyIndex < PinProxiesToProcess.Num(); ProxyIndex++)
 	{
 		if (PinProxiesToProcess[ProxyIndex].IsValid())
@@ -3280,6 +3281,22 @@ const FRigVMCompilerWorkData::FRigVMASTProxyArray& URigVMCompiler::FindProxiesWi
 				{
 					continue;
 				}
+
+				// when compiling a function, avoid folding pins from entry and return nodes
+				if (CurrentCompilationFunction)
+				{
+					if (URigVMNode* Node = Pin->GetNode())
+					{
+						if (Node->IsA<URigVMFunctionEntryNode>() || Node->IsA<URigVMFunctionReturnNode>())
+						{
+							if (bContainsInterfacePin)
+							{
+								continue;
+							}
+							bContainsInterfacePin = true;
+						}
+					}
+				}
 			}
 			PinProxies.Add(PinProxiesToProcess[ProxyIndex]);
 		}
diff --git a/Engine/Source/Runtime/RigVM/Private/RigVMCore/RigVMMemoryStorage.cpp b/Engine/Source/Runtime/RigVM/Private/RigVMCore/RigVMMemoryStorage.cpp
--- a/Engine/Source/Runtime/RigVM/Private/RigVMCore/RigVMMemoryStorage.cpp
+++ b/Engine/Source/Runtime/RigVM/Private/RigVMCore/RigVMMemoryStorage.cpp
@@ -1119,7 +1119,7 @@ bool URigVMMemoryStorage::CopyProperty(
 
 			if (!ReportedErrors.Contains(Message))
 			{
-				UE_LOG(LogRigVM, Warning, TEXT("%s"), *Message);
+				UE_LOG(LogRigVM, Error, TEXT("%s"), *Message);
 				ReportedErrors.Add(Message);
 			}
 			
