Control Rig: Fix postload breaking all links and renable postload hashing by default

diff --git a/Engine/Plugins/Animation/ControlRig/Source/ControlRigDeveloper/Private/ControlRigBlueprint.cpp b/Engine/Plugins/Animation/ControlRig/Source/ControlRigDeveloper/Private/ControlRigBlueprint.cpp
--- a/Engine/Plugins/Animation/ControlRig/Source/ControlRigDeveloper/Private/ControlRigBlueprint.cpp
+++ b/Engine/Plugins/Animation/ControlRig/Source/ControlRigDeveloper/Private/ControlRigBlueprint.cpp
@@ -57,7 +57,7 @@
 
 TAutoConsoleVariable<bool> CVarEnablePostLoadHashing(
 	TEXT("ControlRig.EnablePostLoadHashing"),
-	false,
+	true,
 	TEXT("When true refreshing the RigVMGraphs will be skipped if the hash matches the serialized hash."));
 
 static TArray<UClass*> GetClassObjectsInPackage(UPackage* InPackage)
@@ -1335,7 +1335,7 @@ void UControlRigBlueprint::RefreshAllModels(EControlRigBlueprintLoadType InLoadT
 			for (URigVMGraph* ModelGraph : ModelGraphs)
 			{
 				URigVMController* Controller = GetOrCreateController(ModelGraph);
-				Controller->ReattachLinksToPinObjects(true /* follow redirectors */, nullptr, true, true);
+				Controller->ReattachLinksToPinObjects(true /* follow redirectors */, nullptr, false, true, false);
 			}
 		}
 		return;
diff --git a/Engine/Source/Developer/RigVMDeveloper/Private/RigVMModel/RigVMController.cpp b/Engine/Source/Developer/RigVMDeveloper/Private/RigVMModel/RigVMController.cpp
--- a/Engine/Source/Developer/RigVMDeveloper/Private/RigVMModel/RigVMController.cpp
+++ b/Engine/Source/Developer/RigVMDeveloper/Private/RigVMModel/RigVMController.cpp
@@ -14171,7 +14171,7 @@ int32 URigVMController::DetachLinksFromPinObjects(const TArray<URigVMLink*>* InL
 	return Links.Num();
 }
 
-int32 URigVMController::ReattachLinksToPinObjects(bool bFollowCoreRedirectors, const TArray<URigVMLink*>* InLinks, bool bSetupOrphanedPins, bool bAllowNonArgumentLinks)
+int32 URigVMController::ReattachLinksToPinObjects(bool bFollowCoreRedirectors, const TArray<URigVMLink*>* InLinks, bool bSetupOrphanedPins, bool bAllowNonArgumentLinks, bool bRecursive)
 {
 	URigVMGraph* Graph = GetGraph();
 	check(Graph);
@@ -14232,6 +14232,7 @@ int32 URigVMController::ReattachLinksToPinObjects(bool bFollowCoreRedirectors, c
 			// ignore duplicated links that have been processed
 			if (SourcePin->IsLinkedTo(TargetPin))
 			{
+				NewLinks.Add(Link);
 				continue;
 			}
 
@@ -14324,6 +14325,10 @@ int32 URigVMController::ReattachLinksToPinObjects(bool bFollowCoreRedirectors, c
 
 	if (bReplacingAllLinks)
 	{
+		if(Graph->Links.Num() != NewLinks.Num())
+		{
+			ReportWarningf(TEXT("Number of links changed during ReattachLinksToPinObjects in graph %s in project %s"), *Graph->GetPathName(), *GetPackage()->GetPathName());
+		}
 		Graph->Links = NewLinks;
 
 		for (URigVMLink* Link : Graph->Links)
@@ -14350,7 +14355,7 @@ int32 URigVMController::ReattachLinksToPinObjects(bool bFollowCoreRedirectors, c
 		}
 	}
 
-	if (InLinks == nullptr)
+	if (bRecursive && InLinks == nullptr)
 	{
 		for (URigVMNode* Node : Graph->Nodes)
 		{
@@ -14358,7 +14363,7 @@ int32 URigVMController::ReattachLinksToPinObjects(bool bFollowCoreRedirectors, c
 			{
 				FRigVMControllerGraphGuard GraphGuard(this, CollapseNode->GetContainedGraph(), false);
 				TGuardValue<bool> GuardEditGraph(CollapseNode->ContainedGraph->bEditable, true);
-				ReattachLinksToPinObjects(bFollowCoreRedirectors, nullptr, bSetupOrphanedPins, bAllowNonArgumentLinks);
+				ReattachLinksToPinObjects(bFollowCoreRedirectors, nullptr, bSetupOrphanedPins, bAllowNonArgumentLinks, bRecursive);
 			}
 		}
 	}
diff --git a/Engine/Source/Developer/RigVMDeveloper/Public/RigVMModel/RigVMController.h b/Engine/Source/Developer/RigVMDeveloper/Public/RigVMModel/RigVMController.h
--- a/Engine/Source/Developer/RigVMDeveloper/Public/RigVMModel/RigVMController.h
+++ b/Engine/Source/Developer/RigVMDeveloper/Public/RigVMModel/RigVMController.h
@@ -909,7 +909,7 @@ public:
 	FRigVMController_ConfigureWorkflowOptionsDelegate ConfigureWorkflowOptionsDelegate; 
 
 	int32 DetachLinksFromPinObjects(const TArray<URigVMLink*>* InLinks = nullptr);
-	int32 ReattachLinksToPinObjects(bool bFollowCoreRedirectors = false, const TArray<URigVMLink*>* InLinks = nullptr, bool bSetupOrphanedPins = false, bool bAllowNonArgumentLinks = false);
+	int32 ReattachLinksToPinObjects(bool bFollowCoreRedirectors = false, const TArray<URigVMLink*>* InLinks = nullptr, bool bSetupOrphanedPins = false, bool bAllowNonArgumentLinks = false, bool bRecursive = true);
 	void AddPinRedirector(bool bInput, bool bOutput, const FString& OldPinPath, const FString& NewPinPath);
 
 	// Removes nodes which went stale.
