Control Rig: Fix functions not compiling when pin deafult values change

diff --git a/Engine/Source/Developer/RigVMDeveloper/Private/RigVMCompiler/RigVMCompiler.cpp b/Engine/Source/Developer/RigVMDeveloper/Private/RigVMCompiler/RigVMCompiler.cpp
--- a/Engine/Source/Developer/RigVMDeveloper/Private/RigVMCompiler/RigVMCompiler.cpp
+++ b/Engine/Source/Developer/RigVMDeveloper/Private/RigVMCompiler/RigVMCompiler.cpp
@@ -372,7 +372,7 @@ bool URigVMCompiler::Compile(TArray<URigVMGraph*> InGraphs, URigVMController* In
 								{
 									if (FRigVMGraphFunctionData* DependencyData = HostObj->GetRigVMGraphFunctionStore()->FindFunction(Pair.Key))
 									{
-										if (Pair.Value != DependencyData->CompilationData.Hash)
+										if (DependencyData->CompilationData.Hash == 0 || Pair.Value != DependencyData->CompilationData.Hash)
 										{
 											FunctionData->ClearCompilationData();
 											break;
@@ -1110,6 +1110,8 @@ bool URigVMCompiler::Compile(TArray<URigVMGraph*> InGraphs, URigVMController* In
 				}
 			}
 		}
+
+		OutFunctionCompilationData->Hash = GetTypeHash(OutFunctionCompilationData);
 	}
 
 	if (!CurrentCompilationFunction)
@@ -1134,6 +1136,7 @@ bool URigVMCompiler::CompileFunction(const URigVMLibraryNode* InLibraryNode, URi
 		return false;
 	}
 
+	OutFunctionCompilationData->Hash = 0;
 	OutFunctionCompilationData->ByteCode.Reset();
 
 	TArray<FRigVMExternalVariable> ExternalVariables;
@@ -1171,6 +1174,8 @@ bool URigVMCompiler::CompileFunction(const URigVMLibraryNode* InLibraryNode, URi
 							}
 						}
 					}
+
+					ClientHost->GetRigVMClient()->UpdateFunctionReferences(Data->Header, true, false);
 				}
 		
 				Store->UpdateFunctionCompilationData(InLibraryNode->GetFunctionIdentifier(), *OutFunctionCompilationData);
diff --git a/Engine/Source/Developer/RigVMDeveloper/Private/RigVMModel/RigVMClient.cpp b/Engine/Source/Developer/RigVMDeveloper/Private/RigVMModel/RigVMClient.cpp
--- a/Engine/Source/Developer/RigVMDeveloper/Private/RigVMModel/RigVMClient.cpp
+++ b/Engine/Source/Developer/RigVMDeveloper/Private/RigVMModel/RigVMClient.cpp
@@ -858,17 +858,20 @@ void FRigVMClient::HandleGraphModifiedEvent(ERigVMGraphNotifType InNotifType, UR
 		case ERigVMGraphNotifType::PinTypeChanged: // A pin's data type has changed (Subject == URigVMPin)
 		case ERigVMGraphNotifType::PinIndexChanged: // A pin's index has changed (Subject == URigVMPin)
 		{
-			if (URigVMNode* Node = Cast<URigVMNode>(InSubject->GetOuter()))
+			if (URigVMPin* Pin = Cast<URigVMPin>(InSubject))
 			{
-				if (URigVMLibraryNode* LibraryNode = Node->FindFunctionForNode())
+				if (URigVMNode* Node = Cast<URigVMNode>(Pin->GetNode()))
 				{
-					DirtyGraphFunctionCompilationData(LibraryNode);
-				}
-				if(Node->GetOuter()->IsA<URigVMFunctionLibrary>())
-				{
-					if (URigVMCollapseNode* CollapseNode = Cast<URigVMCollapseNode>(Node))
+					if (URigVMLibraryNode* LibraryNode = Node->FindFunctionForNode())
+					{
+						DirtyGraphFunctionCompilationData(LibraryNode);
+					}
+					if(Node->GetOuter()->IsA<URigVMFunctionLibrary>())
 					{
-						UpdateGraphFunctionData(CollapseNode);
+						if (URigVMCollapseNode* CollapseNode = Cast<URigVMCollapseNode>(Node))
+						{
+							UpdateGraphFunctionData(CollapseNode);
+						}
 					}
 				}
 			}
diff --git a/Engine/Source/Runtime/RigVM/Public/RigVMCore/RigVMGraphFunctionDefinition.h b/Engine/Source/Runtime/RigVM/Public/RigVMCore/RigVMGraphFunctionDefinition.h
--- a/Engine/Source/Runtime/RigVM/Public/RigVMCore/RigVMGraphFunctionDefinition.h
+++ b/Engine/Source/Runtime/RigVM/Public/RigVMCore/RigVMGraphFunctionDefinition.h
@@ -129,7 +129,7 @@ struct RIGVM_API FRigVMFunctionCompilationData
 
 	bool IsValid() const
 	{
-		return ByteCode.Num() > 0 || WorkPropertyDescriptions.Num() > 0;
+		return Hash != 0;
 	}
 
 	friend uint32 GetTypeHash(const FRigVMFunctionCompilationData& Data) 
