Control Rig: Fix crash when opening rig with corrupted aggregate node

diff --git a/Engine/Plugins/Animation/ControlRig/Source/ControlRigDeveloper/Private/Graph/ControlRigGraphNode.cpp b/Engine/Plugins/Animation/ControlRig/Source/ControlRigDeveloper/Private/Graph/ControlRigGraphNode.cpp
--- a/Engine/Plugins/Animation/ControlRig/Source/ControlRigDeveloper/Private/Graph/ControlRigGraphNode.cpp
+++ b/Engine/Plugins/Animation/ControlRig/Source/ControlRigDeveloper/Private/Graph/ControlRigGraphNode.cpp
@@ -1309,9 +1309,18 @@ FSlateIcon UControlRigGraphNode::GetIconAndTint(FLinearColor& OutColor) const
 			return EventIcon;
 		}
 
-		while(const URigVMAggregateNode* AggregateNode = Cast<URigVMAggregateNode>(ModelNode))
+		URigVMNode* InnerNode = ModelNode;
+		while(const URigVMAggregateNode* AggregateNode = Cast<URigVMAggregateNode>(InnerNode))
 		{
-			ModelNode = AggregateNode->GetFirstInnerNode();
+			InnerNode = AggregateNode->GetFirstInnerNode();
+		}
+		if (InnerNode)
+		{
+			ModelNode = InnerNode;
+		}
+		else
+		{
+			checkf(false, TEXT("Could not find first inner node in aggregate node %s in package %s"), *ModelNode->GetPathName(), *GetPackage()->GetPathName());
 		}
 
 		if (ModelNode->IsA<URigVMFunctionReferenceNode>())
diff --git a/Engine/Source/Developer/RigVMDeveloper/Private/RigVMModel/Nodes/RigVMAggregateNode.cpp b/Engine/Source/Developer/RigVMDeveloper/Private/RigVMModel/Nodes/RigVMAggregateNode.cpp
--- a/Engine/Source/Developer/RigVMDeveloper/Private/RigVMModel/Nodes/RigVMAggregateNode.cpp
+++ b/Engine/Source/Developer/RigVMDeveloper/Private/RigVMModel/Nodes/RigVMAggregateNode.cpp
@@ -48,7 +48,20 @@ URigVMNode* URigVMAggregateNode::GetFirstInnerNode() const
 			TArray<URigVMNode*> ConnectedNodes;
 			for (const URigVMPin* EntryPin : EntryNode->GetPins())
 			{
-				const TArray<URigVMPin*> TargetPins = EntryPin->GetLinkedTargetPins();
+				TArray<URigVMPin*> TargetPins = EntryPin->GetLinkedTargetPins();
+				if (TargetPins.IsEmpty())
+				{
+					// If there are no target pins, they might not be reattached yet
+					const TArray<URigVMLink*>& Links = GetContainedLinks();
+					for (const URigVMLink* Link : Links)
+					{
+						if (Link->GetSourcePin() == EntryPin)
+						{
+							TargetPins.Add(Link->GetTargetPin());
+							break;
+						}
+					}
+				}
 				if (TargetPins.Num() > 0)
 				{
 					if (TargetPins[0]->GetName() == Arg1Name || TargetPins[0]->GetName() == Arg2Name)
@@ -72,7 +85,20 @@ URigVMNode* URigVMAggregateNode::GetFirstInnerNode() const
 			const URigVMFunctionEntryNode* EntryNode = GetEntryNode();
 			if (const URigVMPin* EntryPin = EntryNode->FindPin(ArgOppositeName))
 			{
-				const TArray<URigVMPin*> TargetPins = EntryPin->GetLinkedTargetPins();
+				TArray<URigVMPin*> TargetPins = EntryPin->GetLinkedTargetPins();
+				if (TargetPins.IsEmpty())
+				{
+					// If there are no target pins, they might not be reattached yet
+					const TArray<URigVMLink*>& Links = GetContainedLinks();
+					for (const URigVMLink* Link : Links)
+					{
+						if (Link->GetSourcePin() == EntryPin)
+						{
+							TargetPins.Add(Link->GetTargetPin());
+							break;
+						}
+					}
+				}
 				if (TargetPins.Num() > 0)
 				{
 					FirstInnerNodeCache = TargetPins[0]->GetNode();
