RigVM: Follow up for 24537844 - fixing crash in aggregate node

diff --git a/Engine/Plugins/Animation/ControlRig/Source/ControlRigDeveloper/Private/Graph/ControlRigGraphNode.cpp b/Engine/Plugins/Animation/ControlRig/Source/ControlRigDeveloper/Private/Graph/ControlRigGraphNode.cpp
--- a/Engine/Plugins/Animation/ControlRig/Source/ControlRigDeveloper/Private/Graph/ControlRigGraphNode.cpp
+++ b/Engine/Plugins/Animation/ControlRig/Source/ControlRigDeveloper/Private/Graph/ControlRigGraphNode.cpp
@@ -1320,8 +1320,11 @@ FSlateIcon UControlRigGraphNode::GetIconAndTint(FLinearColor& OutColor) const
 		}
 		else
 		{
-			const FString Message = FString::Printf(TEXT("Could not find first inner node in aggregate node %s in package %s"), *ModelNode->GetPathName(), *GetPackage()->GetPathName());
-			FScriptExceptionHandler::Get().HandleException(ELogVerbosity::Warning, *Message, *FString());
+			if(!ModelNode->HasAnyFlags(RF_NeedLoad))
+			{
+				const FString Message = FString::Printf(TEXT("Could not find first inner node in aggregate node %s in package %s"), *ModelNode->GetPathName(), *GetPackage()->GetPathName());
+				FScriptExceptionHandler::Get().HandleException(ELogVerbosity::Warning, *Message, *FString());
+			}
 		}
 
 		if (ModelNode->IsA<URigVMFunctionReferenceNode>())
diff --git a/Engine/Source/Developer/RigVMDeveloper/Private/RigVMModel/Nodes/RigVMAggregateNode.cpp b/Engine/Source/Developer/RigVMDeveloper/Private/RigVMModel/Nodes/RigVMAggregateNode.cpp
--- a/Engine/Source/Developer/RigVMDeveloper/Private/RigVMModel/Nodes/RigVMAggregateNode.cpp
+++ b/Engine/Source/Developer/RigVMDeveloper/Private/RigVMModel/Nodes/RigVMAggregateNode.cpp
@@ -41,68 +41,77 @@ URigVMNode* URigVMAggregateNode::GetFirstInnerNode() const
 		
 		if (IsInputAggregate())
 		{
-			// Find node connected twice to the entry (through aggregate arguments)
-			const FString Arg1Name = GetFirstAggregatePin()->GetName();
-			const FString Arg2Name = GetSecondAggregatePin()->GetName();
-			const URigVMFunctionEntryNode* EntryNode = GetEntryNode();
-			TArray<URigVMNode*> ConnectedNodes;
-			for (const URigVMPin* EntryPin : EntryNode->GetPins())
+			const URigVMPin* FirstAggregatePin = GetFirstAggregatePin();
+			const URigVMPin* SecondAggregatePin = GetSecondAggregatePin();
+			if(FirstAggregatePin && SecondAggregatePin)
 			{
-				TArray<URigVMPin*> TargetPins = EntryPin->GetLinkedTargetPins();
-				if (TargetPins.IsEmpty())
+				// Find node connected twice to the entry (through aggregate arguments)
+				const FString Arg1Name = FirstAggregatePin->GetName();
+				const FString Arg2Name = SecondAggregatePin->GetName();
+				const URigVMFunctionEntryNode* EntryNode = GetEntryNode();
+				TArray<URigVMNode*> ConnectedNodes;
+				for (const URigVMPin* EntryPin : EntryNode->GetPins())
 				{
-					// If there are no target pins, they might not be reattached yet
-					const TArray<URigVMLink*>& Links = GetContainedLinks();
-					for (const URigVMLink* Link : Links)
+					TArray<URigVMPin*> TargetPins = EntryPin->GetLinkedTargetPins();
+					if (TargetPins.IsEmpty())
 					{
-						if (Link->GetSourcePin() == EntryPin)
+						// If there are no target pins, they might not be reattached yet
+						const TArray<URigVMLink*>& Links = GetContainedLinks();
+						for (const URigVMLink* Link : Links)
 						{
-							TargetPins.Add(Link->GetTargetPin());
-							break;
+							if (Link->GetSourcePin() == EntryPin)
+							{
+								TargetPins.Add(Link->GetTargetPin());
+								break;
+							}
 						}
 					}
-				}
-				if (TargetPins.Num() > 0)
-				{
-					if (TargetPins[0]->GetName() == Arg1Name || TargetPins[0]->GetName() == Arg2Name)
+					if (TargetPins.Num() > 0)
 					{
-						URigVMNode* TargetNode = TargetPins[0]->GetNode();
-						if (ConnectedNodes.Contains(TargetNode))
+						if (TargetPins[0]->GetName() == Arg1Name || TargetPins[0]->GetName() == Arg2Name)
 						{
-							FirstInnerNodeCache = TargetNode;
-							return FirstInnerNodeCache;
+							URigVMNode* TargetNode = TargetPins[0]->GetNode();
+							if (ConnectedNodes.Contains(TargetNode))
+							{
+								FirstInnerNodeCache = TargetNode;
+								return FirstInnerNodeCache;
+							}
+	
+							ConnectedNodes.Add(TargetNode);
 						}
-
-						ConnectedNodes.Add(TargetNode);
 					}
 				}
 			}
 		}
 		else
 		{
-			// Find node connected to entry throught the opposite aggregate argument 
-			const FString ArgOppositeName = GetOppositeAggregatePin()->GetName();
-			const URigVMFunctionEntryNode* EntryNode = GetEntryNode();
-			if (const URigVMPin* EntryPin = EntryNode->FindPin(ArgOppositeName))
+			// Find node connected to entry throught the opposite aggregate argument
+			const URigVMPin* OppositeAggregatePin = GetOppositeAggregatePin();
+			if(OppositeAggregatePin)
 			{
-				TArray<URigVMPin*> TargetPins = EntryPin->GetLinkedTargetPins();
-				if (TargetPins.IsEmpty())
+				const FString ArgOppositeName = OppositeAggregatePin->GetName();
+				const URigVMFunctionEntryNode* EntryNode = GetEntryNode();
+				if (const URigVMPin* EntryPin = EntryNode->FindPin(ArgOppositeName))
 				{
-					// If there are no target pins, they might not be reattached yet
-					const TArray<URigVMLink*>& Links = GetContainedLinks();
-					for (const URigVMLink* Link : Links)
+					TArray<URigVMPin*> TargetPins = EntryPin->GetLinkedTargetPins();
+					if (TargetPins.IsEmpty())
 					{
-						if (Link->GetSourcePin() == EntryPin)
+						// If there are no target pins, they might not be reattached yet
+						const TArray<URigVMLink*>& Links = GetContainedLinks();
+						for (const URigVMLink* Link : Links)
 						{
-							TargetPins.Add(Link->GetTargetPin());
-							break;
+							if (Link->GetSourcePin() == EntryPin)
+							{
+								TargetPins.Add(Link->GetTargetPin());
+								break;
+							}
 						}
 					}
-				}
-				if (TargetPins.Num() > 0)
-				{
-					FirstInnerNodeCache = TargetPins[0]->GetNode();
-					return FirstInnerNodeCache;
+					if (TargetPins.Num() > 0)
+					{
+						FirstInnerNodeCache = TargetPins[0]->GetNode();
+						return FirstInnerNodeCache;
+					}
 				}
 			}
 		}
