Backing out 24538523

#deadend

diff --git a/Engine/Plugins/Animation/ControlRig/Source/ControlRigDeveloper/Private/Graph/ControlRigGraphNode.cpp b/Engine/Plugins/Animation/ControlRig/Source/ControlRigDeveloper/Private/Graph/ControlRigGraphNode.cpp
--- a/Engine/Plugins/Animation/ControlRig/Source/ControlRigDeveloper/Private/Graph/ControlRigGraphNode.cpp
+++ b/Engine/Plugins/Animation/ControlRig/Source/ControlRigDeveloper/Private/Graph/ControlRigGraphNode.cpp
@@ -1309,22 +1309,9 @@ FSlateIcon UControlRigGraphNode::GetIconAndTint(FLinearColor& OutColor) const
 			return EventIcon;
 		}
 
-		URigVMNode* InnerNode = ModelNode;
-		while(const URigVMAggregateNode* AggregateNode = Cast<URigVMAggregateNode>(InnerNode))
+		while(const URigVMAggregateNode* AggregateNode = Cast<URigVMAggregateNode>(ModelNode))
 		{
-			InnerNode = AggregateNode->GetFirstInnerNode();
-		}
-		if (InnerNode)
-		{
-			ModelNode = InnerNode;
-		}
-		else
-		{
-			if(!ModelNode->HasAnyFlags(RF_NeedLoad))
-			{
-				const FString Message = FString::Printf(TEXT("Could not find first inner node in aggregate node %s in package %s"), *ModelNode->GetPathName(), *GetPackage()->GetPathName());
-				FScriptExceptionHandler::Get().HandleException(ELogVerbosity::Warning, *Message, *FString());
-			}
+			ModelNode = AggregateNode->GetFirstInnerNode();
 		}
 
 		if (ModelNode->IsA<URigVMFunctionReferenceNode>())
diff --git a/Engine/Source/Developer/RigVMDeveloper/Private/RigVMModel/Nodes/RigVMAggregateNode.cpp b/Engine/Source/Developer/RigVMDeveloper/Private/RigVMModel/Nodes/RigVMAggregateNode.cpp
--- a/Engine/Source/Developer/RigVMDeveloper/Private/RigVMModel/Nodes/RigVMAggregateNode.cpp
+++ b/Engine/Source/Developer/RigVMDeveloper/Private/RigVMModel/Nodes/RigVMAggregateNode.cpp
@@ -41,77 +41,42 @@ URigVMNode* URigVMAggregateNode::GetFirstInnerNode() const
 		
 		if (IsInputAggregate())
 		{
-			const URigVMPin* FirstAggregatePin = GetFirstAggregatePin();
-			const URigVMPin* SecondAggregatePin = GetSecondAggregatePin();
-			if(FirstAggregatePin && SecondAggregatePin)
+			// Find node connected twice to the entry (through aggregate arguments)
+			const FString Arg1Name = GetFirstAggregatePin()->GetName();
+			const FString Arg2Name = GetSecondAggregatePin()->GetName();
+			const URigVMFunctionEntryNode* EntryNode = GetEntryNode();
+			TArray<URigVMNode*> ConnectedNodes;
+			for (const URigVMPin* EntryPin : EntryNode->GetPins())
 			{
-				// Find node connected twice to the entry (through aggregate arguments)
-				const FString Arg1Name = FirstAggregatePin->GetName();
-				const FString Arg2Name = SecondAggregatePin->GetName();
-				const URigVMFunctionEntryNode* EntryNode = GetEntryNode();
-				TArray<URigVMNode*> ConnectedNodes;
-				for (const URigVMPin* EntryPin : EntryNode->GetPins())
+				const TArray<URigVMPin*> TargetPins = EntryPin->GetLinkedTargetPins();
+				if (TargetPins.Num() > 0)
 				{
-					TArray<URigVMPin*> TargetPins = EntryPin->GetLinkedTargetPins();
-					if (TargetPins.IsEmpty())
+					if (TargetPins[0]->GetName() == Arg1Name || TargetPins[0]->GetName() == Arg2Name)
 					{
-						// If there are no target pins, they might not be reattached yet
-						const TArray<URigVMLink*>& Links = GetContainedLinks();
-						for (const URigVMLink* Link : Links)
+						URigVMNode* TargetNode = TargetPins[0]->GetNode();
+						if (ConnectedNodes.Contains(TargetNode))
 						{
-							if (Link->GetSourcePin() == EntryPin)
-							{
-								TargetPins.Add(Link->GetTargetPin());
-								break;
-							}
-						}
-					}
-					if (TargetPins.Num() > 0)
-					{
-						if (TargetPins[0]->GetName() == Arg1Name || TargetPins[0]->GetName() == Arg2Name)
-						{
-							URigVMNode* TargetNode = TargetPins[0]->GetNode();
-							if (ConnectedNodes.Contains(TargetNode))
-							{
-								FirstInnerNodeCache = TargetNode;
-								return FirstInnerNodeCache;
-							}
-	
-							ConnectedNodes.Add(TargetNode);
+							FirstInnerNodeCache = TargetNode;
+							return FirstInnerNodeCache;
 						}
+
+						ConnectedNodes.Add(TargetNode);
 					}
 				}
 			}
 		}
 		else
 		{
-			// Find node connected to entry throught the opposite aggregate argument
-			const URigVMPin* OppositeAggregatePin = GetOppositeAggregatePin();
-			if(OppositeAggregatePin)
+			// Find node connected to entry throught the opposite aggregate argument 
+			const FString ArgOppositeName = GetOppositeAggregatePin()->GetName();
+			const URigVMFunctionEntryNode* EntryNode = GetEntryNode();
+			if (const URigVMPin* EntryPin = EntryNode->FindPin(ArgOppositeName))
 			{
-				const FString ArgOppositeName = OppositeAggregatePin->GetName();
-				const URigVMFunctionEntryNode* EntryNode = GetEntryNode();
-				if (const URigVMPin* EntryPin = EntryNode->FindPin(ArgOppositeName))
+				const TArray<URigVMPin*> TargetPins = EntryPin->GetLinkedTargetPins();
+				if (TargetPins.Num() > 0)
 				{
-					TArray<URigVMPin*> TargetPins = EntryPin->GetLinkedTargetPins();
-					if (TargetPins.IsEmpty())
-					{
-						// If there are no target pins, they might not be reattached yet
-						const TArray<URigVMLink*>& Links = GetContainedLinks();
-						for (const URigVMLink* Link : Links)
-						{
-							if (Link->GetSourcePin() == EntryPin)
-							{
-								TargetPins.Add(Link->GetTargetPin());
-								break;
-							}
-						}
-					}
-					if (TargetPins.Num() > 0)
-					{
-						FirstInnerNodeCache = TargetPins[0]->GetNode();
-						return FirstInnerNodeCache;
-					}
+					FirstInnerNodeCache = TargetPins[0]->GetNode();
+					return FirstInnerNodeCache;
 				}
 			}
 		}
