Control Rig: Fix loading of pin flag bIsDynamicArray

diff --git a/Engine/Source/Developer/RigVMDeveloper/Private/RigVMModel/RigVMController.cpp b/Engine/Source/Developer/RigVMDeveloper/Private/RigVMModel/RigVMController.cpp
--- a/Engine/Source/Developer/RigVMDeveloper/Private/RigVMModel/RigVMController.cpp
+++ b/Engine/Source/Developer/RigVMDeveloper/Private/RigVMModel/RigVMController.cpp
@@ -13810,6 +13810,23 @@ void URigVMController::AddPinsForTemplate(const FRigVMTemplate* InTemplate, cons
 		}
 		Pin->Direction = Arg->GetDirection();
 
+		Pin->bIsDynamicArray = FRigVMRegistry::Get().IsArrayType(TypeIndex);
+		if (Pin->Direction == ERigVMPinDirection::Hidden)
+		{
+			if (InTemplate->GetArgumentMetaData(Arg->Name, FRigVMStruct::ArraySizeMetaName).IsEmpty())
+			{
+				Pin->bIsDynamicArray = true;
+			}
+		}
+
+		if (Pin->bIsDynamicArray)
+		{
+			if (!InTemplate->GetArgumentMetaData(Arg->Name, FRigVMStruct::SingletonMetaName).IsEmpty())
+			{
+				Pin->bIsDynamicArray = false;
+			}
+		}
+
 		AddNodePin(InNode, Pin);
 
 		if(!Pin->IsWildCard() && !Pin->IsArray())
@@ -13890,10 +13907,17 @@ void URigVMController::ConfigurePinFromProperty(FProperty* InProperty, URigVMPin
 	InOutPin->CPPType += ExtendedCppType;
 
 	InOutPin->bIsDynamicArray = false;
+	FProperty* PropertyForType = InProperty;
+	FArrayProperty* ArrayProperty = CastField<FArrayProperty>(PropertyForType);
+	if (ArrayProperty)
+	{
+		PropertyForType = ArrayProperty->Inner;
+		InOutPin->bIsDynamicArray = true;
+	}
 #if WITH_EDITOR
 	if (InOutPin->Direction == ERigVMPinDirection::Hidden)
 	{
-		if (!InProperty->HasMetaData(TEXT("ArraySize")))
+		if (!InProperty->HasMetaData(FRigVMStruct::ArraySizeMetaName))
 		{
 			InOutPin->bIsDynamicArray = true;
 		}
@@ -13908,13 +13932,6 @@ void URigVMController::ConfigurePinFromProperty(FProperty* InProperty, URigVMPin
 	}
 #endif
 
-	FProperty* PropertyForType = InProperty;
-	FArrayProperty* ArrayProperty = CastField<FArrayProperty>(PropertyForType);
-	if (ArrayProperty)
-	{
-		PropertyForType = ArrayProperty->Inner;
-	}
-
 	if (FStructProperty* StructProperty = CastField<FStructProperty>(PropertyForType))
 	{
 		InOutPin->CPPTypeObject = StructProperty->Struct;
diff --git a/Engine/Source/Runtime/RigVM/Private/RigVMCore/RigVMStruct.cpp b/Engine/Source/Runtime/RigVM/Private/RigVMCore/RigVMStruct.cpp
--- a/Engine/Source/Runtime/RigVM/Private/RigVMCore/RigVMStruct.cpp
+++ b/Engine/Source/Runtime/RigVM/Private/RigVMCore/RigVMStruct.cpp
@@ -295,6 +295,7 @@ const FName FRigVMStruct::IconMetaName("Icon");
 const FName FRigVMStruct::KeywordsMetaName("Keywords");
 const FName FRigVMStruct::TemplateNameMetaName = FRigVMRegistry::TemplateNameMetaName;
 const FName FRigVMStruct::FixedSizeArrayMetaName("FixedSizeArray");
+const FName FRigVMStruct::ArraySizeMetaName("ArraySize");
 const FName FRigVMStruct::AggregateMetaName("Aggregate");
 const FName FRigVMStruct::ExpandPinByDefaultMetaName("ExpandByDefault");
 const FName FRigVMStruct::DefaultArraySizeMetaName("DefaultArraySize");
diff --git a/Engine/Source/Runtime/RigVM/Public/RigVMCore/RigVMStruct.h b/Engine/Source/Runtime/RigVM/Public/RigVMCore/RigVMStruct.h
--- a/Engine/Source/Runtime/RigVM/Public/RigVMCore/RigVMStruct.h
+++ b/Engine/Source/Runtime/RigVM/Public/RigVMCore/RigVMStruct.h
@@ -248,6 +248,7 @@ public:
 	static const FName KeywordsMetaName;
 	static const FName TemplateNameMetaName;
 	static const FName FixedSizeArrayMetaName;
+	static const FName ArraySizeMetaName;
 	static const FName AggregateMetaName;
 	static const FName ExpandPinByDefaultMetaName;
 	static const FName DefaultArraySizeMetaName;
