Control Rig: Add copy operation between float and double properties

diff --git a/Engine/Source/Developer/RigVMDeveloper/Private/RigVMCompiler/RigVMAST.cpp b/Engine/Source/Developer/RigVMDeveloper/Private/RigVMCompiler/RigVMAST.cpp
--- a/Engine/Source/Developer/RigVMDeveloper/Private/RigVMCompiler/RigVMAST.cpp
+++ b/Engine/Source/Developer/RigVMDeveloper/Private/RigVMCompiler/RigVMAST.cpp
@@ -1220,6 +1220,7 @@ FRigVMExprAST* FRigVMParserAST::TraverseLink(int32 InLinkIndex, FRigVMExprAST* I
 			bRequiresCopy = true;
 		}
 	}
+
 	if (!bRequiresCopy)
 	{
 		// Connections between entry and return in a function requires a copy
@@ -1232,6 +1233,14 @@ FRigVMExprAST* FRigVMParserAST::TraverseLink(int32 InLinkIndex, FRigVMExprAST* I
 		}
 	}
 
+	if (!bRequiresCopy)
+	{
+		if (SourcePin->GetTypeIndex() != TargetPin->GetTypeIndex())
+		{
+			bRequiresCopy = true;
+		}
+	}
+
 	FRigVMAssignExprAST* AssignExpr = nullptr;
 	if (bRequiresCopy)
 	{
