Control Rig: Composition reroutes not behaving correctly

diff --git a/Engine/Source/Developer/RigVMDeveloper/Private/RigVMCompiler/RigVMAST.cpp b/Engine/Source/Developer/RigVMDeveloper/Private/RigVMCompiler/RigVMAST.cpp
--- a/Engine/Source/Developer/RigVMDeveloper/Private/RigVMCompiler/RigVMAST.cpp
+++ b/Engine/Source/Developer/RigVMDeveloper/Private/RigVMCompiler/RigVMAST.cpp
@@ -1233,6 +1233,18 @@ FRigVMExprAST* FRigVMParserAST::TraverseLink(int32 InLinkIndex, FRigVMExprAST* I
 		}
 	}
 
+	if (!bRequiresCopy)
+	{
+		// Connections from a composition reroute always require a copy
+		if (URigVMRerouteNode* RerouteNode = Cast<URigVMRerouteNode>(SourceRootPin->GetNode()))
+		{
+			if (RerouteNode->IsCompositionNode())
+			{
+				bRequiresCopy = true;
+			}
+		}
+	}
+
 	if (!bRequiresCopy)
 	{
 		if (SourcePin->GetTypeIndex() != TargetPin->GetTypeIndex())
@@ -1667,9 +1679,13 @@ void FRigVMParserAST::FoldAssignments()
 			}
 		}
 		
-		// non-input pins on anything but a reroute / array node should be skipped
-		if (TargetPin->GetDirection() != ERigVMPinDirection::Input &&
-			(Cast<URigVMRerouteNode>(TargetPin->GetNode()) == nullptr))
+		// non-input pins on anything but a reroute (passthrough or literal) or array node should be skipped
+		bool bIsReroute = false;
+		if (URigVMRerouteNode* RerouteNode = Cast<URigVMRerouteNode>(TargetPin->GetNode()))
+		{
+			bIsReroute = !RerouteNode->IsCompositionNode();
+		}
+		if (TargetPin->GetDirection() != ERigVMPinDirection::Input && !bIsReroute)
 		{
 			continue;
 		}
@@ -2561,11 +2577,18 @@ void FRigVMParserAST::Inline(TArray<URigVMGraph*> InGraphs, const TArray<FRigVMA
 			// If its a function reference, don't recurse
 			if (Node->IsA<URigVMFunctionReferenceNode>())
 			{
-				return false;				
+				return false;
+			}
+			
+			if (URigVMRerouteNode* RerouteNode = Cast<URigVMRerouteNode>(Node))
+			{
+				if (!RerouteNode->IsCompositionNode())
+				{
+					return true;
+				}
 			}
 			
-			return Node->IsA<URigVMRerouteNode>() ||
-				Node->IsA<URigVMLibraryNode>() ||
+			return Node->IsA<URigVMLibraryNode>() ||
 				Node->IsA<URigVMFunctionEntryNode>() ||
 				Node->IsA<URigVMFunctionReturnNode>();
 		}
@@ -2882,7 +2905,14 @@ void FRigVMParserAST::Inline(TArray<URigVMGraph*> InGraphs, const TArray<FRigVMA
 
 				URigVMPin* SourcePin = SourcePinProxy.GetSubjectChecked<URigVMPin>();
 				URigVMNode* SourceNode = SourcePin->GetNode();
-				if (SourceNode->IsA<URigVMRerouteNode>() ||
+
+				bool bIsPassThroughOrLiteralReroute = false;
+				if (URigVMRerouteNode* RerouteNode = Cast<URigVMRerouteNode>(SourceNode))
+				{
+					bIsPassThroughOrLiteralReroute = RerouteNode->IsPassThrough() || RerouteNode->IsLiteral();
+				}
+				
+				if (bIsPassThroughOrLiteralReroute ||
 					SourceNode->IsA<URigVMCollapseNode>() ||
 					SourceNode->IsA<URigVMFunctionReturnNode>())
 				{
@@ -2934,9 +2964,12 @@ void FRigVMParserAST::Inline(TArray<URigVMGraph*> InGraphs, const TArray<FRigVMA
 
 		static void VisitNode(const FRigVMASTProxy& InNodeProxy, LocalPinTraversalInfo& OutTraversalInfo)
 		{
-			if (InNodeProxy.IsA<URigVMRerouteNode>())
+			if (URigVMRerouteNode* RerouteNode = Cast<URigVMRerouteNode>(InNodeProxy.GetSubject()))
 			{
-				return;
+				if (RerouteNode->IsPassThrough() || RerouteNode->IsLiteral())
+				{
+					return;
+				}
 			}
 			
 			const bool bIsCompilingFunction = OutTraversalInfo.LibraryNodeBeingCompiled != nullptr;
diff --git a/Engine/Source/Developer/RigVMDeveloper/Private/RigVMCompiler/RigVMCompiler.cpp b/Engine/Source/Developer/RigVMDeveloper/Private/RigVMCompiler/RigVMCompiler.cpp
--- a/Engine/Source/Developer/RigVMDeveloper/Private/RigVMCompiler/RigVMCompiler.cpp
+++ b/Engine/Source/Developer/RigVMDeveloper/Private/RigVMCompiler/RigVMCompiler.cpp
@@ -856,7 +856,7 @@ bool URigVMCompiler::Compile(TArray<URigVMGraph*> InGraphs, URigVMController* In
 			{
 				if(EntryNode == nullptr)
 				{
-					ReportError(TEXT("Corrupt library node '%s' - Missing entry node."));
+					ReportErrorf(TEXT("Corrupt library node '%s' - Missing entry node."), *CurrentCompilationFunction->GetPathName());
 					return false;
 				}
 				InterfacePin = EntryNode->FindPin(Pin->GetName());
@@ -865,7 +865,7 @@ bool URigVMCompiler::Compile(TArray<URigVMGraph*> InGraphs, URigVMController* In
 			{
 				if(ReturnNode == nullptr)
 				{
-					ReportError(TEXT("Corrupt library node '%s' - Missing return node."));
+					ReportErrorf(TEXT("Corrupt library node '%s' - Missing return node."), *CurrentCompilationFunction->GetPathName());
 					return false;
 				}
 				InterfacePin = ReturnNode->FindPin(Pin->GetName());
@@ -873,7 +873,7 @@ bool URigVMCompiler::Compile(TArray<URigVMGraph*> InGraphs, URigVMController* In
 
 			if(InterfacePin == nullptr)
 			{
-				ReportError(TEXT("Corrupt library node '%s' - Pin '%s' is not part of the entry / return node."));
+				ReportErrorf(TEXT("Corrupt library node '%s' - Pin '%s' is not part of the entry / return node."), *CurrentCompilationFunction->GetPathName(), *Pin->GetPathName());
 				return false;
 			}
 			
diff --git a/Engine/Source/Developer/RigVMDeveloper/Private/RigVMModel/Nodes/RigVMRerouteNode.cpp b/Engine/Source/Developer/RigVMDeveloper/Private/RigVMModel/Nodes/RigVMRerouteNode.cpp
--- a/Engine/Source/Developer/RigVMDeveloper/Private/RigVMModel/Nodes/RigVMRerouteNode.cpp
+++ b/Engine/Source/Developer/RigVMDeveloper/Private/RigVMModel/Nodes/RigVMRerouteNode.cpp
@@ -114,3 +114,23 @@ const FRigVMTemplate* URigVMRerouteNode::GetTemplate() const
 	return CachedTemplate;
 }
 
+bool URigVMRerouteNode::IsPassThrough() const
+{
+	return !GetPins()[0]->GetSourceLinks().IsEmpty();
+}
+
+bool URigVMRerouteNode::IsLiteral() const
+{
+	return GetPins()[0]->GetSourceLinks(true).IsEmpty();
+}
+
+bool URigVMRerouteNode::IsCompositionNode() const
+{
+	if (IsPassThrough() || IsLiteral())
+	{
+		return false;
+	}
+	
+	return true;
+}
+
diff --git a/Engine/Source/Developer/RigVMDeveloper/Public/RigVMModel/Nodes/RigVMRerouteNode.h b/Engine/Source/Developer/RigVMDeveloper/Public/RigVMModel/Nodes/RigVMRerouteNode.h
--- a/Engine/Source/Developer/RigVMDeveloper/Public/RigVMModel/Nodes/RigVMRerouteNode.h
+++ b/Engine/Source/Developer/RigVMDeveloper/Public/RigVMModel/Nodes/RigVMRerouteNode.h
@@ -36,6 +36,15 @@ public:
 	virtual const FRigVMTemplate* GetTemplate() const override;
 	virtual bool IsSingleton() const override { return false; }
 
+	// Is a dot node with a root pin source connection
+	bool IsPassThrough() const;
+
+	// Has no source connections
+	bool IsLiteral() const;
+
+	// Is not a passthrough or a literal reroute
+	bool IsCompositionNode() const;
+
 private:
 
 	static const FString RerouteName;
diff --git a/Engine/Source/Runtime/RigVM/Private/RigVMCore/RigVM.cpp b/Engine/Source/Runtime/RigVM/Private/RigVMCore/RigVM.cpp
--- a/Engine/Source/Runtime/RigVM/Private/RigVMCore/RigVM.cpp
+++ b/Engine/Source/Runtime/RigVM/Private/RigVMCore/RigVM.cpp
@@ -1759,16 +1759,20 @@ ERigVMExecuteResult URigVM::ExecuteInstructions(int32 InFirstInstruction, int32
 				Labels.Add(GetOperandLabel(Operand));
 			}
 
-			ContextPublicData.DebugMemoryString += FString::Printf(TEXT("Instruction %d: %s(%s)\n"), ContextPublicData.InstructionIndex, *FunctionNames[Op.FunctionIndex].ToString(), *FString::Join(Labels, TEXT(", ")));				
+			ContextPublicData.DebugMemoryString += FString::Printf(TEXT("Instruction %d: %s(%s)\n"), ContextPublicData.InstructionIndex, *FunctionNames[Op.FunctionIndex].ToString(), *FString::Join(Labels, TEXT(", ")));
 		}
 		else if(Instruction.OpCode == ERigVMOpCode::Copy)
 		{
+			static auto FormatFunction = [](const FString& RegisterName, const FString& RegisterOffsetName) -> FString
+			{
+				return FString::Printf(TEXT("%s%s"), *RegisterName, *RegisterOffsetName);
+			};
 			const FRigVMCopyOp& Op = ByteCode.GetOpAt<FRigVMCopyOp>(Instruction);
-			ContextPublicData.DebugMemoryString += FString::Printf(TEXT("Instruction %d: Copy %s -> %s\n"), ContextPublicData.InstructionIndex, *GetOperandLabel(Op.Source), *GetOperandLabel(Op.Target));				
+			ContextPublicData.DebugMemoryString += FString::Printf(TEXT("Instruction %d: Copy %s -> %s\n"), ContextPublicData.InstructionIndex, *GetOperandLabel(Op.Source, FormatFunction), *GetOperandLabel(Op.Target, FormatFunction));
 		}
 		else
 		{
-			ContextPublicData.DebugMemoryString += FString::Printf(TEXT("Instruction %d: %s\n"), ContextPublicData.InstructionIndex, *ContextPublicData.InstanceOpCodeEnum->GetNameByIndex((uint8)Instruction.OpCode).ToString());				
+			ContextPublicData.DebugMemoryString += FString::Printf(TEXT("Instruction %d: %s\n"), ContextPublicData.InstructionIndex, *ContextPublicData.InstanceOpCodeEnum->GetNameByIndex((uint8)Instruction.OpCode).ToString());
 		}
 		
 #endif
