Control Rig: Fix entry/return nodes being created in the same position when creating a function

diff --git a/Engine/Source/Developer/RigVMDeveloper/Private/RigVMModel/RigVMController.cpp b/Engine/Source/Developer/RigVMDeveloper/Private/RigVMModel/RigVMController.cpp
--- a/Engine/Source/Developer/RigVMDeveloper/Private/RigVMModel/RigVMController.cpp
+++ b/Engine/Source/Developer/RigVMDeveloper/Private/RigVMModel/RigVMController.cpp
@@ -11361,10 +11361,12 @@ URigVMLibraryNode* URigVMController::AddFunctionToLibrary(const FName& InFunctio
 				TGuardValue<bool> SuspendNotifications(bSuspendNotifications, true);
 				RefreshFunctionPins(EntryNode);
 			}
-			EntryNode->Position = FVector2D(-250.f, 0.f);
 			Notify(ERigVMGraphNotifType::NodeAdded, EntryNode);
 		}
 
+		EntryNode->Position = FVector2D(-250.f, 0.f);
+		Notify(ERigVMGraphNotifType::NodePositionChanged, EntryNode);
+
 		if (ReturnNode == nullptr)
 		{
 			ReturnNode = NewObject<URigVMFunctionReturnNode>(CollapseNode->ContainedGraph, TEXT("Return"));
@@ -11373,10 +11375,11 @@ URigVMLibraryNode* URigVMController::AddFunctionToLibrary(const FName& InFunctio
 				TGuardValue<bool> SuspendNotifications(bSuspendNotifications, true);
 				RefreshFunctionPins(ReturnNode);
 			}
-			ReturnNode->Position = FVector2D(250.f, 0.f);
 			Notify(ERigVMGraphNotifType::NodeAdded, ReturnNode);
 		}
 
+		ReturnNode->Position = FVector2D(250.f, 0.f);
+		Notify(ERigVMGraphNotifType::NodePositionChanged, ReturnNode);
 
 		if (bMutable)
 		{
