Control Rig: Fix crash when opening corrupted rig

diff --git a/Engine/Source/Developer/RigVMDeveloper/Private/RigVMCompiler/RigVMASTProxy.cpp b/Engine/Source/Developer/RigVMDeveloper/Private/RigVMCompiler/RigVMASTProxy.cpp
--- a/Engine/Source/Developer/RigVMDeveloper/Private/RigVMCompiler/RigVMASTProxy.cpp
+++ b/Engine/Source/Developer/RigVMDeveloper/Private/RigVMCompiler/RigVMASTProxy.cpp
@@ -15,13 +15,16 @@ FString FRigVMCallstack::GetCallPath(bool bIncludeLast) const
 	{
 		if (URigVMNode* Node = Cast<URigVMNode>(Entry))
 		{
-			if(Node->GetGraph()->IsRootGraph())
+			if (URigVMGraph* Graph = Node->GetGraph())
 			{
-				Segments.Add(Node->GetNodePath(true));
-			}
-			else
-			{
-				Segments.Add(Node->GetName());
+				if(Graph->IsRootGraph())
+				{
+					Segments.Add(Node->GetNodePath(true));
+				}
+				else
+				{
+					Segments.Add(Node->GetName());
+				}
 			}
 		}
 		else if (URigVMPin* Pin = Cast<URigVMPin>(Entry))
@@ -34,7 +37,7 @@ FString FRigVMCallstack::GetCallPath(bool bIncludeLast) const
 		}
 	}
 
-	if (!bIncludeLast)
+	if (!bIncludeLast && !Segments.IsEmpty())
 	{
 		Segments.Pop();
 	}
diff --git a/Engine/Source/Developer/RigVMDeveloper/Private/RigVMModel/RigVMController.cpp b/Engine/Source/Developer/RigVMDeveloper/Private/RigVMModel/RigVMController.cpp
--- a/Engine/Source/Developer/RigVMDeveloper/Private/RigVMModel/RigVMController.cpp
+++ b/Engine/Source/Developer/RigVMDeveloper/Private/RigVMModel/RigVMController.cpp
@@ -15107,9 +15107,13 @@ void URigVMController::RepopulatePinsOnNode(URigVMNode* InNode, bool bFollowCore
 		{
 			AddSubPin(ParentPin, Pin);
 		}
+		else if (URigVMNode* OwnerNode = Cast<URigVMNode>(InOuter))
+		{
+			AddNodePin(OwnerNode, Pin);
+		}
 		else
 		{
-			AddNodePin(CastChecked<URigVMNode>(InOuter), Pin);
+			ensureMsgf(false, TEXT("Outer %s of pin info %s is not a pin or a node"), *InOuter->GetPathName(), *InPinPath);
 		}
 
 		Notify(ERigVMGraphNotifType::PinAdded, Pin);
@@ -17417,14 +17421,14 @@ void URigVMController::MakeExecutePin(URigVMPin* InOutPin)
 void URigVMController::AddNodePin(URigVMNode* InNode, URigVMPin* InPin)
 {
 	ValidatePin(InPin);
-	check(!InNode->Pins.Contains(InPin));
+	checkf(!InNode->Pins.Contains(InPin), TEXT("Node %s already contains pin %s"), *InNode->GetPathName(), *InPin->GetName());
 	InNode->Pins.Add(InPin);
 }
 
 void URigVMController::AddSubPin(URigVMPin* InParentPin, URigVMPin* InPin)
 {
 	ValidatePin(InPin);
-	check(!InParentPin->SubPins.Contains(InPin));
+	checkf(!InParentPin->SubPins.Contains(InPin), TEXT("Parent pin %s already contains subpin %s"), *InParentPin->GetPathName(), *InPin->GetName());
 	InParentPin->SubPins.Add(InPin);
 }
 
