RigVM: Refactoring RigVMController Pass 1

* Remove inverse action
* Consolidate undo actions to rely on import text
* Avoid unnessary actions that cancel each other out
* Remove notion of detach links and use BreakLinkedPaths / RestoreLinkedPaths as the single mechanism

diff --git a/Engine/Plugins/Animation/ControlRig/Source/ControlRigDeveloper/Private/ControlRigBlueprint.cpp b/Engine/Plugins/Animation/ControlRig/Source/ControlRigDeveloper/Private/ControlRigBlueprint.cpp
--- a/Engine/Plugins/Animation/ControlRig/Source/ControlRigDeveloper/Private/ControlRigBlueprint.cpp
+++ b/Engine/Plugins/Animation/ControlRig/Source/ControlRigDeveloper/Private/ControlRigBlueprint.cpp
@@ -979,6 +979,10 @@ void UControlRigBlueprint::PostLoad()
 			RefreshAllModels();
 		}
 
+		// at this point we may still have links which are detached. we may or may not be able to 
+		// reattach them.
+		GetRigVMClient()->ProcessDetachedLinks();
+
 		GetControlRigBlueprintGeneratedClass()->GetRigVMGraphFunctionStore()->RemoveAllCompilationData();
 
 		// perform backwards compat value upgrades
@@ -995,7 +999,7 @@ void UControlRigBlueprint::PostLoad()
 			{
 				URigVMController* Controller = GetOrCreateController(GraphToValidate);
 				FRigVMControllerNotifGuard NotifGuard(Controller, true);
-				Controller->RemoveUnusedOrphanedPins(Node);
+				Controller->RemoveUnusedOrphanedPins(Node, true);
 			}
 				
 			for(URigVMNode* Node : GraphToValidate->GetNodes())
@@ -1411,7 +1415,10 @@ void UControlRigBlueprint::RefreshAllModels(EControlRigBlueprintLoadType InLoadT
 			for (URigVMGraph* ModelGraph : ModelGraphs)
 			{
 				URigVMController* Controller = GetOrCreateController(ModelGraph);
-				Controller->ReattachLinksToPinObjects(true /* follow redirectors */, nullptr, false, true, false);
+				URigVMController::FRestoreLinkedPathSettings Settings;
+				Settings.bFollowCoreRedirectors = true;
+				Settings.bRelayToOrphanPins = true;
+				Controller->ProcessDetachedLinks(Settings);
 			}
 		}
 		return;
@@ -1420,18 +1427,20 @@ void UControlRigBlueprint::RefreshAllModels(EControlRigBlueprintLoadType InLoadT
 	TGuardValue<bool> IsCompilingGuard(bIsCompiling, true);
 	TGuardValue<bool> ClientIgnoreModificationsGuard(RigVMClient.bIgnoreModelNotifications, true);
 	
-	TArray<URigVMGraph*> GraphsToDetach = RigVMClient.GetAllModels(true, false);
+	TArray<URigVMGraph*> AllModelsLeavesFirst = RigVMClient.GetAllModelsLeavesFirst(true);
+	TMap<const URigVMGraph*, TArray<URigVMController::FLinkedPath>> LinkedPaths;
 
 	if (ensure(IsInGameThread()))
 	{
-		for (URigVMGraph* GraphToDetach : GraphsToDetach)
+		for (URigVMGraph* Graph : AllModelsLeavesFirst)
 		{
-			URigVMController* Controller = GetOrCreateController(GraphToDetach);
+			URigVMController* Controller = GetOrCreateController(Graph);
 			// temporarily disable default value validation during load time, serialized values should always be accepted
 			TGuardValue<bool> PerGraphDisablePinDefaultValueValidation(Controller->bValidatePinDefaults, false);
 			FRigVMControllerNotifGuard NotifGuard(Controller, true);
-			Controller->DetachLinksFromPinObjects();
-			TArray<URigVMNode*> Nodes = GraphToDetach->GetNodes();
+			LinkedPaths.Add(Graph, Controller->GetLinkedPaths());
+			Controller->FastBreakLinkedPaths(LinkedPaths.FindChecked(Graph));
+			TArray<URigVMNode*> Nodes = Graph->GetNodes();
 			for (URigVMNode* Node : Nodes)
 			{
 				Controller->RepopulatePinsOnNode(Node, true, true);
@@ -1440,88 +1449,26 @@ void UControlRigBlueprint::RefreshAllModels(EControlRigBlueprintLoadType InLoadT
 		SetupPinRedirectorsForBackwardsCompatibility();
 	}
 
-	for (URigVMGraph* GraphToDetach : GraphsToDetach)
+	for (URigVMGraph* Graph : AllModelsLeavesFirst)
 	{
-		URigVMController* Controller = GetOrCreateController(GraphToDetach);
-		// at this stage, allow all links to be reattached,
-		// RecomputeAllTemplateFilteredPermutations() later should break any invalid links
+		URigVMController* Controller = GetOrCreateController(Graph);
 		FRigVMControllerNotifGuard NotifGuard(Controller, true);
-		Controller->ReattachLinksToPinObjects(true /* follow redirectors */, nullptr, true, true);
-	}
-
-	TArray<URigVMGraph*> GraphsToClean = GetAllModels();
-
-	// Sort from leaf graphs to root
-	TArray<URigVMGraph*> SortedGraphsToClean;
-	SortedGraphsToClean.Reserve(GraphsToClean.Num());
-	while (SortedGraphsToClean.Num() < GraphsToClean.Num())
-	{
-		bool bGraphAdded = false;
-		for (URigVMGraph* Graph : GraphsToClean)
 		{
-			if (SortedGraphsToClean.Contains(Graph))
-			{
-				continue;
-			}
-			
-			TArray<URigVMGraph*> ContainedGraphs;
-			for (URigVMNode* Node : Graph->GetNodes())
-			{
-				if (URigVMLibraryNode* LibraryNode = Cast<URigVMLibraryNode>(Node))
-				{
-					if (URigVMFunctionReferenceNode* FunctionReferenceNode = Cast<URigVMFunctionReferenceNode>(LibraryNode))
-					{
-						if (FunctionReferenceNode->GetReferencedFunctionHeader().LibraryPointer.LibraryNode.GetLongPackageName() != GetPackage()->GetPathName())
-						{
-							continue;
-						}
-						if (URigVMLibraryNode* ReferencedNode = FunctionReferenceNode->LoadReferencedNode())
-						{
-							ContainedGraphs.Add(ReferencedNode->GetContainedGraph());
-							continue;
-						}
-					}
-
-					if (URigVMGraph* ContainedGraph = LibraryNode->GetContainedGraph())
-					{
-						ContainedGraphs.Add(ContainedGraph);
-					}
-				}
-			}
-			
-			bool bAllContained = true;
-			for (URigVMGraph* Contained : ContainedGraphs)
-			{
-				if (!SortedGraphsToClean.Contains(Contained))
-				{
-					bAllContained = false;
-					break;
-				}
-			}
-			if (bAllContained)
-			{
-				SortedGraphsToClean.Add(Graph);
-				bGraphAdded = true;
-			}
+			URigVMController::FRestoreLinkedPathSettings Settings;
+			Settings.bFollowCoreRedirectors = true;
+			Controller->RestoreLinkedPaths(LinkedPaths.FindChecked(Graph));
 		}
-		ensure(bGraphAdded);
-	}
-	
-	for(int32 GraphIndex=0; GraphIndex<SortedGraphsToClean.Num(); GraphIndex++)
-	{
-		URigVMGraph* GraphToClean = SortedGraphsToClean[GraphIndex];
-		URigVMController* Controller = GetOrCreateController(GraphToClean);
-		TGuardValue<bool> GuardEditGraph(GraphToClean->bEditable, true);
-		FRigVMControllerNotifGuard NotifGuard(Controller, true);
+
+		TGuardValue<bool> GuardEditGraph(Graph->bEditable, true);
 		
-		for(URigVMNode* ModelNode : GraphToClean->GetNodes())
+		for(URigVMNode* ModelNode : Graph->GetNodes())
 		{
 			Controller->RemoveUnusedOrphanedPins(ModelNode);
 		}
 
 		if(bIsPostLoad)
 		{
-			for(URigVMNode* ModelNode : GraphToClean->GetNodes())
+			for(URigVMNode* ModelNode : Graph->GetNodes())
 			{
 				if (URigVMTemplateNode* TemplateNode = Cast<URigVMTemplateNode>(ModelNode))
 				{
@@ -1535,56 +1482,42 @@ void UControlRigBlueprint::RefreshAllModels(EControlRigBlueprintLoadType InLoadT
 
 		if(bIsPostLoad)
 		{
-			for(URigVMNode* ModelNode : GraphToClean->GetNodes())
+			for(URigVMNode* ModelNode : Graph->GetNodes())
 			{
-				if (ModelNode->HasWildCardPin())
-				{
-					continue;
-				}
 				if(URigVMUnitNode* UnitNode = Cast<URigVMUnitNode>(ModelNode))
 				{
-					UScriptStruct* ScriptStruct = UnitNode->GetScriptStruct(); 
-					if(ScriptStruct == nullptr)
-					{
-						Controller->FullyResolveTemplateNode(UnitNode, INDEX_NONE, false);
-					}
-
-					// Try to find a deprecated template
-					if (UnitNode->GetScriptStruct() == nullptr && !UnitNode->TemplateNotation.IsNone())
+					if (!UnitNode->HasWildCardPin())
 					{
-						const FRigVMTemplate* Template = FRigVMRegistry::Get().FindTemplate(UnitNode->TemplateNotation, true);
-						FRigVMTemplate::FTypeMap TypeMap = UnitNode->GetTemplatePinTypeMap(true);
-						int32 Permutation;
-						if (Template->FullyResolve(TypeMap, Permutation))
+						UScriptStruct* ScriptStruct = UnitNode->GetScriptStruct(); 
+						if(ScriptStruct == nullptr)
 						{
-							const FRigVMFunction* Function = Template->GetPermutation(Permutation);
-							UnitNode->ResolvedFunctionName = Function->GetName();
+							Controller->FullyResolveTemplateNode(UnitNode, INDEX_NONE, false);
 						}
-					}
 
-					if (UnitNode->GetScriptStruct() == nullptr)
-					{
-						static const TCHAR UnresolvedUnitNodeMessage[] = TEXT("Node %s could not be resolved.");
-						Controller->ReportErrorf(UnresolvedUnitNodeMessage, *ModelNode->GetNodePath(true));
-					}
-				}
-				if (URigVMDispatchNode* DispatchNode = Cast<URigVMDispatchNode>(ModelNode))
-				{
-					if (const FRigVMDispatchFactory* Factory = DispatchNode->GetFactory())
-					{
-						FRigVMTemplate::FTypeMap TypeMap = DispatchNode->GetTemplatePinTypeMap(true);
-						const FString ResolvedPermutationName = Factory->GetPermutationName(TypeMap);
-						if(const FRigVMFunction* Function = FRigVMRegistry::Get().FindFunction(*ResolvedPermutationName))
+						// Try to find a deprecated template
+						if (UnitNode->GetScriptStruct() == nullptr && !UnitNode->TemplateNotation.IsNone())
 						{
-							DispatchNode->ResolvedFunctionName = Function->GetName();
+							const FRigVMTemplate* Template = FRigVMRegistry::Get().FindTemplate(UnitNode->TemplateNotation, true);
+							FRigVMTemplate::FTypeMap TypeMap = UnitNode->GetTemplatePinTypeMap();
+
+							int32 Permutation;
+							if (Template->FullyResolve(TypeMap, Permutation))
+							{
+								const FRigVMFunction* Function = Template->GetPermutation(Permutation);
+								UnitNode->ResolvedFunctionName = Function->GetName();
+							}
 						}
-						else
+
+						if (UnitNode->GetScriptStruct() == nullptr)
 						{
-							static const TCHAR UnresolvedDispatchPermutationMessage[] = TEXT("Dispatch factory %s could not find permutation..");
-							Controller->ReportErrorf(UnresolvedDispatchPermutationMessage, *Factory->GetFactoryName().ToString());
+							static const TCHAR UnresolvedUnitNodeMessage[] = TEXT("Node %s could not be resolved.");
+							Controller->ReportErrorf(UnresolvedUnitNodeMessage, *ModelNode->GetNodePath(true));
 						}
 					}
-					else
+				}
+				if (URigVMDispatchNode* DispatchNode = Cast<URigVMDispatchNode>(ModelNode))
+				{
+					if (DispatchNode->GetFactory() == nullptr)
 					{
 						static const TCHAR UnresolvedDispatchNodeMessage[] = TEXT("Dispatch node %s has no factory..");
 						Controller->ReportErrorf(UnresolvedDispatchNodeMessage, *ModelNode->GetNodePath(true));
@@ -4206,12 +4139,6 @@ void UControlRigBlueprint::PatchVariableNodesOnLoad()
 	{
 		TGuardValue<bool> GuardNotifsSelf(bSuspendModelNotificationsForSelf, true);
 
-		{
-			URigVMController* Controller = GetOrCreateController();
-			FRigVMControllerNotifGuard NotifGuard(Controller, true);
-			Controller->ReattachLinksToPinObjects();
-		}
-
 		check(GetDefaultModel());
 
 		for(URigVMGraph* Model : RigVMClient)
@@ -4229,7 +4156,9 @@ void UControlRigBlueprint::PatchVariableNodesOnLoad()
 
 					int32 VariableIndex = AddedMemberVariableMap.FindChecked(Description.Name);
 					FName VarName = NewVariables[VariableIndex].VarName;
-					GetOrCreateController()->RefreshVariableNode(VariableNode->GetFName(), VarName, Description.CPPType, Description.CPPTypeObject, false);
+
+					URigVMController* Controller = GetOrCreateController(Model);
+					Controller->RefreshVariableNode(VariableNode->GetFName(), VarName, Description.CPPType, Description.CPPTypeObject, false);
 					bDirtyDuringLoad = true;
 				}
 
@@ -4243,7 +4172,8 @@ void UControlRigBlueprint::PatchVariableNodesOnLoad()
 
 					int32 VariableIndex = AddedMemberVariableMap.FindChecked(Description.Name);
 					FName VarName = NewVariables[VariableIndex].VarName;
-					GetOrCreateController()->ReplaceParameterNodeWithVariable(ParameterNode->GetFName(), VarName, Description.CPPType, Description.CPPTypeObject, false);
+					URigVMController* Controller = GetOrCreateController(Model);
+					Controller->ReplaceParameterNodeWithVariable(ParameterNode->GetFName(), VarName, Description.CPPType, Description.CPPTypeObject, false);
 					bDirtyDuringLoad = true;
 				}
 			}
@@ -4494,12 +4424,6 @@ void UControlRigBlueprint::PatchParameterNodesOnLoad()
 		
 		TGuardValue<bool> GuardNotifsSelf(bSuspendModelNotificationsForSelf, true);
 
-		{
-			URigVMController* Controller = GetOrCreateController();
-			FRigVMControllerNotifGuard NotifGuard(Controller, true);
-			Controller->ReattachLinksToPinObjects();
-		}
-
 		check(GetDefaultModel());
 
 		for(URigVMGraph* Model : RigVMClient)
@@ -4609,8 +4533,8 @@ void UControlRigBlueprint::PatchLinksWithCast()
 				URigVMPin* TargetPin = Link->GetTargetPin();
 				if(!SourcePin->IsLinkedTo(TargetPin))
 				{
-					const TArray<URigVMLink*> LinksToReattach = {Link};
-					Controller->ReattachLinksToPinObjects(true, &LinksToReattach, false, false);
+					const TArray<URigVMController::FLinkedPath> LinkedPaths = Controller->GetLinkedPaths({Link});
+					Controller->RestoreLinkedPaths(LinkedPaths);
 				}
 			}
 
diff --git a/Engine/Plugins/Animation/ControlRig/Source/ControlRigDeveloper/Private/Graph/ControlRigGraphSchema.cpp b/Engine/Plugins/Animation/ControlRig/Source/ControlRigDeveloper/Private/Graph/ControlRigGraphSchema.cpp
--- a/Engine/Plugins/Animation/ControlRig/Source/ControlRigDeveloper/Private/Graph/ControlRigGraphSchema.cpp
+++ b/Engine/Plugins/Animation/ControlRig/Source/ControlRigDeveloper/Private/Graph/ControlRigGraphSchema.cpp
@@ -1624,7 +1624,7 @@ bool UControlRigGraphSchema::TryDeleteGraph(UEdGraph* GraphToDelete) const
 							}
 						}
 						
-						return Controller->RemoveNode(LibraryNode, bSetupUndoRedo, false, true);
+						return Controller->RemoveNode(LibraryNode, bSetupUndoRedo, false);
 					}
 				}
 			}
@@ -2051,7 +2051,7 @@ bool UControlRigGraphSchema::SafeDeleteNodeFromGraph(UEdGraph* Graph, UEdGraphNo
 		{
 			GEditor->CancelTransaction(0);
 		}
-		return RigNode->GetController()->RemoveNode(RigNode->GetModelNode(), true, true, true);
+		return RigNode->GetController()->RemoveNode(RigNode->GetModelNode(), true, true);
 	}
 	return false;
 }
diff --git a/Engine/Plugins/Animation/ControlRig/Source/ControlRigEditor/Private/ControlRigEditorModule.cpp b/Engine/Plugins/Animation/ControlRig/Source/ControlRigEditor/Private/ControlRigEditorModule.cpp
--- a/Engine/Plugins/Animation/ControlRig/Source/ControlRigEditor/Private/ControlRigEditorModule.cpp
+++ b/Engine/Plugins/Animation/ControlRig/Source/ControlRigEditor/Private/ControlRigEditorModule.cpp
@@ -1835,7 +1835,7 @@ void FControlRigEditorModule::GetContextMenuActions(const UControlRigGraphSchema
 										LOCTEXT("RemoveVisualDebug_Tooltip", "Removes the visual debugging node"),
 										FSlateIcon(),
 											FUIAction(FExecuteAction::CreateLambda([Controller, InGraphPin, ModelPin, VisualDebugNode]() {
-											Controller->RemoveNodeByName(VisualDebugNode->GetFName(), true, false, true);
+											Controller->RemoveNodeByName(VisualDebugNode->GetFName(), true, false);
 										})
 									));
 								}
diff --git a/Engine/Plugins/Animation/ControlRig/Source/ControlRigEditor/Private/Editor/ControlRigEditor.cpp b/Engine/Plugins/Animation/ControlRig/Source/ControlRigEditor/Private/Editor/ControlRigEditor.cpp
--- a/Engine/Plugins/Animation/ControlRig/Source/ControlRigEditor/Private/Editor/ControlRigEditor.cpp
+++ b/Engine/Plugins/Animation/ControlRig/Source/ControlRigEditor/Private/Editor/ControlRigEditor.cpp
@@ -2572,8 +2572,8 @@ void FControlRigEditor::DeleteSelectedNodes()
 	const FGraphPanelSelectionSet SelectedNodes = GetSelectedNodes();
 	SetUISelectionState(NAME_None);
 
-	bool DeletedAnything = false;
-	GetFocusedController()->OpenUndoBracket(TEXT("Delete selected nodes"));
+	bool bRelinkPins = false;
+	TArray<URigVMNode*> NodesToRemove;
 
 	for (FGraphPanelSelectionSet::TConstIterator NodeIt(SelectedNodes); NodeIt; ++NodeIt)
 	{
@@ -2584,17 +2584,18 @@ void FControlRigEditor::DeleteSelectedNodes()
 				AnalyticsTrackNodeEvent(GetBlueprintObj(), Node, true);
 				if (UControlRigGraphNode* RigNode = Cast<UControlRigGraphNode>(Node))
 				{
-					const bool bReconnectPins = (FSlateApplication::Get().GetModifierKeys().IsShiftDown());
-					if(GetFocusedController()->RemoveNodeByName(*RigNode->ModelNodePath, true, false, true, bReconnectPins))
+					bRelinkPins = bRelinkPins || FSlateApplication::Get().GetModifierKeys().IsShiftDown();
+
+					if(URigVMNode* ModelNode = GetFocusedController()->GetGraph()->FindNodeByName(*RigNode->ModelNodePath))
 					{
-						DeletedAnything = true;
+						NodesToRemove.Add(ModelNode);
 					}
 				}
 				else if (UEdGraphNode_Comment* CommentNode = Cast<UEdGraphNode_Comment>(Node))
 				{
-					if(GetFocusedController()->RemoveNodeByName(CommentNode->GetFName(), true, false, true))
+					if(URigVMNode* ModelNode = GetFocusedController()->GetGraph()->FindNodeByName(CommentNode->GetFName()))
 					{
-						DeletedAnything = true;
+						NodesToRemove.Add(ModelNode);
 					}
 				}
 				else
@@ -2605,14 +2606,18 @@ void FControlRigEditor::DeleteSelectedNodes()
 		}
 	}
 
-	if(DeletedAnything)
+	if(NodesToRemove.IsEmpty())
 	{
-		GetFocusedController()->CloseUndoBracket();
+		return;
 	}
-	else
+
+	GetFocusedController()->OpenUndoBracket(TEXT("Delete selected nodes"));
+	if(bRelinkPins && NodesToRemove.Num() == 1)
 	{
-		GetFocusedController()->CancelUndoBracket();
+		GetFocusedController()->RelinkSourceAndTargetPins(NodesToRemove[0], true);;
 	}
+	GetFocusedController()->RemoveNodes(NodesToRemove, true);
+	GetFocusedController()->CloseUndoBracket();
 }
 
 bool FControlRigEditor::CanDeleteNodes() const
diff --git a/Engine/Plugins/Experimental/AnimNext/Source/AnimNextGraphUncookedOnly/Private/AnimNextGraph_EditorData.cpp b/Engine/Plugins/Experimental/AnimNext/Source/AnimNextGraphUncookedOnly/Private/AnimNextGraph_EditorData.cpp
--- a/Engine/Plugins/Experimental/AnimNext/Source/AnimNextGraphUncookedOnly/Private/AnimNextGraph_EditorData.cpp
+++ b/Engine/Plugins/Experimental/AnimNext/Source/AnimNextGraphUncookedOnly/Private/AnimNextGraph_EditorData.cpp
@@ -125,16 +125,18 @@ void UAnimNextGraph_EditorData::RefreshAllModels(EAnimNextGraphLoadType InLoadTy
 	TGuardValue<bool> ClientIgnoreModificationsGuard(RigVMClient.bIgnoreModelNotifications, true);
 
 	TArray<URigVMGraph*> GraphsToDetach = RigVMClient.GetAllModels(true, false);
-
+	TMap<const URigVMGraph*, TArray<URigVMController::FLinkedPath>> LinkedPaths;
+	
 	if (ensure(IsInGameThread()))
 	{
-		for (URigVMGraph* GraphToDetach : GraphsToDetach)
+		for (const URigVMGraph* GraphToDetach : GraphsToDetach)
 		{
 			URigVMController* Controller = RigVMClient.GetOrCreateController(GraphToDetach);
 			// temporarily disable default value validation during load time, serialized values should always be accepted
 			TGuardValue<bool> PerGraphDisablePinDefaultValueValidation(Controller->bValidatePinDefaults, false);
 			FRigVMControllerNotifGuard NotifGuard(Controller, true);
-			Controller->DetachLinksFromPinObjects();
+			LinkedPaths.Add(GraphToDetach, Controller->GetLinkedPaths());
+			Controller->FastBreakLinkedPaths(LinkedPaths.FindChecked(GraphToDetach));
 			TArray<URigVMNode*> Nodes = GraphToDetach->GetNodes();
 			for (URigVMNode* Node : Nodes)
 			{
@@ -144,13 +146,13 @@ void UAnimNextGraph_EditorData::RefreshAllModels(EAnimNextGraphLoadType InLoadTy
 		//SetupPinRedirectorsForBackwardsCompatibility();
 	}
 
-	for (URigVMGraph* GraphToDetach : GraphsToDetach)
+	for (const URigVMGraph* GraphToDetach : GraphsToDetach)
 	{
 		URigVMController* Controller = RigVMClient.GetOrCreateController(GraphToDetach);
 		// at this stage, allow all links to be reattached,
 		// RecomputeAllTemplateFilteredPermutations() later should break any invalid links
 		FRigVMControllerNotifGuard NotifGuard(Controller, true);
-		Controller->ReattachLinksToPinObjects(true /* follow redirectors */, nullptr, true, true);
+		Controller->RestoreLinkedPaths(LinkedPaths.FindChecked(GraphToDetach));
 	}
 
 	if (bIsPostLoad)
diff --git a/Engine/Source/Developer/RigVMDeveloper/Private/RigVMModel/Nodes/RigVMTemplateNode.cpp b/Engine/Source/Developer/RigVMDeveloper/Private/RigVMModel/Nodes/RigVMTemplateNode.cpp
--- a/Engine/Source/Developer/RigVMDeveloper/Private/RigVMModel/Nodes/RigVMTemplateNode.cpp
+++ b/Engine/Source/Developer/RigVMDeveloper/Private/RigVMModel/Nodes/RigVMTemplateNode.cpp
@@ -288,6 +288,11 @@ TArray<URigVMPin*> URigVMTemplateNode::GetAggregatePins(const ERigVMPinDirection
 		AggregatePins = FindAggregatePins(Permutations[0]);
 		for (int32 i=1; i<Permutations.Num(); ++i)
 		{
+			if(Permutations[i] >= Template->NumPermutations())
+			{
+				continue;
+			}
+			
 			TArray<URigVMPin*> OtherAggregateInputs = FindAggregatePins(Permutations[i]);
 			if (OtherAggregateInputs.Num() != AggregatePins.Num() ||
 				OtherAggregateInputs.FilterByPredicate([&](const URigVMPin* OtherPin)
diff --git a/Engine/Source/Developer/RigVMDeveloper/Private/RigVMModel/RigVMClient.cpp b/Engine/Source/Developer/RigVMDeveloper/Private/RigVMModel/RigVMClient.cpp
--- a/Engine/Source/Developer/RigVMDeveloper/Private/RigVMModel/RigVMClient.cpp
+++ b/Engine/Source/Developer/RigVMDeveloper/Private/RigVMModel/RigVMClient.cpp
@@ -3,6 +3,7 @@
 #include "RigVMModel/RigVMClient.h"
 #include "Misc/TransactionObjectEvent.h"
 #include "Exporters/Exporter.h"
+#include "Algo/Sort.h"
 #include "UnrealExporter.h"
 
 #include UE_INLINE_GENERATED_CPP_BY_NAME(RigVMClient)
@@ -136,6 +137,13 @@ TArray<URigVMGraph*> FRigVMClient::GetAllModels(bool bIncludeFunctionLibrary, bo
 	return AllModels;
 }
 
+TArray<URigVMGraph*> FRigVMClient::GetAllModelsLeavesFirst(bool bIncludeFunctionLibrary) const
+{
+	TArray<URigVMGraph*> SortedModels = GetAllModels(bIncludeFunctionLibrary, true);
+	URigVMController::SortGraphElementsByGraphDepth(SortedModels, true);
+	return SortedModels;
+}
+
 URigVMController* FRigVMClient::GetController(int32 InIndex) const
 {
 	return GetController(GetModel(InIndex));
@@ -687,14 +695,16 @@ FRigVMClientPatchResult FRigVMClient::PatchModelsOnLoad()
 {
 	FRigVMClientPatchResult Result;
 	
-	TArray<URigVMGraph*> AllModels = GetAllModels(true, true);
+	TArray<URigVMGraph*> AllModels = GetAllModelsLeavesFirst(true);
 	TGuardValue<bool> ClientIgnoreModificationsGuard(bIgnoreModelNotifications, true);
 	for(URigVMGraph* Model : AllModels)
 	{
+		Model->PostLoad();
+		
 		URigVMController* Controller = GetOrCreateController(Model);
 		TGuardValue<bool> GuardSuspendTemplateComputation(Controller->bSuspendTemplateComputation, true);
 		TGuardValue<bool> GuardIsTransacting(Controller->bIsTransacting, true);
-		
+
 		Result.Merge(Controller->PatchUnitNodesOnLoad());
 		Result.Merge(Controller->PatchDispatchNodesOnLoad());
 		Result.Merge(Controller->PatchBranchNodesOnLoad());
@@ -706,6 +716,16 @@ FRigVMClientPatchResult FRigVMClient::PatchModelsOnLoad()
 	return Result;
 }
 
+void FRigVMClient::ProcessDetachedLinks()
+{
+	TArray<URigVMGraph*> AllModels = GetAllModels(true, true);
+	for(URigVMGraph* Model : AllModels)
+	{
+		URigVMController* Controller = GetOrCreateController(Model);
+		Controller->ProcessDetachedLinks();
+	}
+}
+
 void FRigVMClient::PostDuplicateHost(const FString& InOldPathName, const FString& InNewPathName)
 {
 	TArray<URigVMGraph*> AllModels = GetAllModels(true, true);
diff --git a/Engine/Source/Developer/RigVMDeveloper/Private/RigVMModel/RigVMController.cpp b/Engine/Source/Developer/RigVMDeveloper/Private/RigVMModel/RigVMController.cpp
--- a/Engine/Source/Developer/RigVMDeveloper/Private/RigVMModel/RigVMController.cpp
+++ b/Engine/Source/Developer/RigVMDeveloper/Private/RigVMModel/RigVMController.cpp
@@ -18,8 +18,8 @@
 #include "UObject/Package.h"
 #include "UObject/StrongObjectPtr.h"
 #include "Misc/CoreMisc.h"
-#include "Algo/Sort.h"
 #include "Algo/Count.h"
+#include "Algo/Reverse.h"
 #include "Algo/Transform.h"
 #include "RigVMPythonUtils.h"
 #include "RigVMTypeUtils.h"
@@ -149,8 +149,7 @@ void URigVMController::SetGraph(URigVMGraph* InGraph)
 {
 	ensure(Graphs.Num() < 2);
 
-	URigVMGraph* LastGraph = GetGraph();
-	if (LastGraph)
+	if (URigVMGraph* LastGraph = GetGraph())
 	{
 		if(LastGraph == InGraph)
 		{
@@ -168,11 +167,14 @@ void URigVMController::SetGraph(URigVMGraph* InGraph)
 	HandleModifiedEvent(ERigVMGraphNotifType::GraphChanged, GetGraph(), nullptr);
 }
 
-void URigVMController::PushGraph(URigVMGraph* InGraph, bool bSetupUndoRedo)
+bool URigVMController::PushGraph(URigVMGraph* InGraph, bool bSetupUndoRedo)
 {
-	URigVMGraph* LastGraph = GetGraph();
-	if (LastGraph)
+	if (URigVMGraph* LastGraph = GetGraph())
 	{
+		if(LastGraph == InGraph)
+		{
+			return false;
+		}
 		LastGraph->OnModified().RemoveAll(this);
 	}
 
@@ -185,11 +187,16 @@ void URigVMController::PushGraph(URigVMGraph* InGraph, bool bSetupUndoRedo)
 	{
 		ActionStack->AddAction(FRigVMPushGraphAction(InGraph));
 	}
+
+	return true;
 }
 
 URigVMGraph* URigVMController::PopGraph(bool bSetupUndoRedo)
 {
-	ensure(Graphs.Num() > 1);
+	if(Graphs.Num() == 1)
+	{
+		return nullptr;
+	}
 	
 	URigVMGraph* LastGraph = GetGraph();
 	if (LastGraph)
@@ -199,8 +206,7 @@ URigVMGraph* URigVMController::PopGraph(bool bSetupUndoRedo)
 
 	Graphs.Pop();
 
-	URigVMGraph* CurrentGraph = GetGraph();
-	if (CurrentGraph)
+	if (URigVMGraph* CurrentGraph = GetGraph())
 	{
 		CurrentGraph->OnModified().AddUObject(this, &URigVMController::HandleModifiedEvent);
 	}
@@ -1337,12 +1343,13 @@ URigVMUnitNode* URigVMController::AddUnitNode(UScriptStruct* InScriptStruct, con
 	}
 
 	FRigVMControllerCompileBracketScope CompileScope(this);
-	FRigVMAddUnitNodeAction Action;
+
+	FRigVMBaseAction Action;
 	if (bSetupUndoRedo)
 	{
-		Action = FRigVMAddUnitNodeAction(Node);
 		Action.Title = FString::Printf(TEXT("Add %s Node"), *Node->GetNodeTitle());
 		ActionStack->BeginAction(Action);
+		ActionStack->AddAction(FRigVMImportFromTextAction(Node, this));
 	}
 
 	Notify(ERigVMGraphNotifType::NodeAdded, Node);
@@ -1609,7 +1616,15 @@ URigVMVariableNode* URigVMController::AddVariableNode(const FName& InVariableNam
 	}
 
 	FString CPPType = RigVMTypeUtils::PostProcessCPPType(InCPPType, InCPPTypeObject);
-	
+
+	FRigVMExternalVariable ExternalVariable = GetVariableByName(InVariableName, true);
+	if(!ExternalVariable.IsValid(true))
+	{
+		static constexpr TCHAR Format[] = TEXT("Cannot add variable '%s' with type '%s' - variable does not exist.");
+		ReportErrorf(Format, *InVariableName.ToString(), *CPPType);
+		return nullptr;
+	}
+
 	FString Name = GetValidNodeName(InNodeName.IsEmpty() ? FString(TEXT("VariableNode")) : InNodeName);
 	URigVMVariableNode* Node = NewObject<URigVMVariableNode>(Graph, *Name);
 	Node->Position = InPosition;
@@ -1629,42 +1644,17 @@ URigVMVariableNode* URigVMController::AddVariableNode(const FName& InVariableNam
 	AddNodePin(Node, VariablePin);
 
 	URigVMPin* ValuePin = NewObject<URigVMPin>(Node, *URigVMVariableNode::ValueName);
-
-	FRigVMExternalVariable ExternalVariable = GetVariableByName(InVariableName);
-	if(ExternalVariable.IsValid(true))
+	ValuePin->CPPType = ExternalVariable.TypeName.ToString();
+	ValuePin->CPPTypeObject = ExternalVariable.TypeObject;
+	if (ValuePin->CPPTypeObject)
 	{
-		ValuePin->CPPType = ExternalVariable.TypeName.ToString();
-		ValuePin->CPPTypeObject = ExternalVariable.TypeObject;
-		if (ValuePin->CPPTypeObject)
-		{
-			ValuePin->CPPTypeObjectPath = *ValuePin->CPPTypeObject->GetPathName();
-		}
-		ValuePin->bIsDynamicArray = ExternalVariable.bIsArray;
-
-		if(ValuePin->bIsDynamicArray && !RigVMTypeUtils::IsArrayType(ValuePin->CPPType))
-		{
-			ValuePin->CPPType = RigVMTypeUtils::ArrayTypeFromBaseType(*ValuePin->CPPType);
-		}
+		ValuePin->CPPTypeObjectPath = *ValuePin->CPPTypeObject->GetPathName();
 	}
-	else
-	{
-		ValuePin->CPPType = CPPType;
+	ValuePin->bIsDynamicArray = ExternalVariable.bIsArray;
 
-		if (UClass* Class = Cast<UClass>(InCPPTypeObject))
-		{
-			ValuePin->CPPTypeObject = Class;
-			ValuePin->CPPTypeObjectPath = *ValuePin->CPPTypeObject->GetPathName();
-		}
-		else if (UScriptStruct* ScriptStruct = Cast<UScriptStruct>(InCPPTypeObject))
-		{
-			ValuePin->CPPTypeObject = ScriptStruct;
-			ValuePin->CPPTypeObjectPath = *ValuePin->CPPTypeObject->GetPathName();
-		}
-		else if (UEnum* Enum = Cast<UEnum>(InCPPTypeObject))
-		{
-			ValuePin->CPPTypeObject = Enum;
-			ValuePin->CPPTypeObjectPath = *ValuePin->CPPTypeObject->GetPathName();
-		}
+	if(ValuePin->bIsDynamicArray && !RigVMTypeUtils::IsArrayType(ValuePin->CPPType))
+	{
+		ValuePin->CPPType = RigVMTypeUtils::ArrayTypeFromBaseType(*ValuePin->CPPType);
 	}
 
 	ValuePin->Direction = bIsGetter ? ERigVMPinDirection::Output : ERigVMPinDirection::Input;
@@ -1696,12 +1686,12 @@ URigVMVariableNode* URigVMController::AddVariableNode(const FName& InVariableNam
 	}
 
 	FRigVMControllerCompileBracketScope CompileScope(this);
-	FRigVMAddVariableNodeAction Action;
+	FRigVMBaseAction Action;
 	if (bSetupUndoRedo)
 	{
-		Action = FRigVMAddVariableNodeAction(Node);
 		Action.Title = FString::Printf(TEXT("Add %s Variable"), *InVariableName.ToString());
 		ActionStack->BeginAction(Action);
+		ActionStack->AddAction(FRigVMImportFromTextAction(Node, this));
 	}
 
 	Notify(ERigVMGraphNotifType::NodeAdded, Node);
@@ -2445,8 +2435,8 @@ TArray<URigVMNode*> URigVMController::UpgradeNodes(const TArray<URigVMNode*>& In
 	}
 
 	// find all links affecting the nodes to upgrade
-	TArray<TPair<FString, FString>> LinkedPaths = GetLinkedPinPaths(InNodes);
-	if(!BreakLinkedPaths(LinkedPaths, bSetupUndoRedo))
+	TArray<FLinkedPath> LinkedPaths = GetLinkedPaths(InNodes, true);
+	if(!FastBreakLinkedPaths(LinkedPaths, bSetupUndoRedo))
 	{
 		if(bSetupUndoRedo)
 		{
@@ -2471,7 +2461,9 @@ TArray<URigVMNode*> URigVMController::UpgradeNodes(const TArray<URigVMNode*>& In
 		}
 	}
 
-	RestoreLinkedPaths(LinkedPaths, {}, RemapPinDelegates, bSetupUndoRedo);
+	FRestoreLinkedPathSettings Settings;
+	Settings.RemapDelegates = RemapPinDelegates;
+	RestoreLinkedPaths(LinkedPaths, Settings, bSetupUndoRedo);
 
 	if(bRecursive)
 	{
@@ -2531,7 +2523,7 @@ URigVMNode* URigVMController::UpgradeNode(URigVMNode* InNode, bool bSetupUndoRed
 		});
 	}
 
-	if(!RemoveNode(InNode, bSetupUndoRedo, true, false, false))
+	if(!RemoveNode(InNode, bSetupUndoRedo, false))
 	{
 		if(bSetupUndoRedo)
 		{
@@ -2706,12 +2698,12 @@ URigVMCommentNode* URigVMController::AddCommentNode(const FString& InCommentText
 	}
 
 	FRigVMControllerCompileBracketScope CompileScope(this);
-	FRigVMAddCommentNodeAction Action;
+	FRigVMBaseAction Action;
 	if (bSetupUndoRedo)
 	{
-		Action = FRigVMAddCommentNodeAction(Node);
 		Action.Title = FString::Printf(TEXT("Add Comment"));
 		ActionStack->BeginAction(Action);
+		ActionStack->AddAction(FRigVMImportFromTextAction(Node, this));
 	}
 
 	Notify(ERigVMGraphNotifType::NodeAdded, Node);
@@ -2898,7 +2890,7 @@ URigVMRerouteNode* URigVMController::AddRerouteNodeOnPin(const FString& InPinPat
 
 	if (bSetupUndoRedo)
 	{
-		ActionStack->AddAction(FRigVMAddRerouteNodeAction(Node));
+		ActionStack->AddAction(FRigVMImportFromTextAction(Node, this, false));
 	}
 
 	Notify(ERigVMGraphNotifType::NodeAdded, Node);
@@ -3477,13 +3469,13 @@ URigVMNode* URigVMController::EjectNodeFromPin(URigVMPin* InPin, bool bSetupUndo
 
 	
 	FRigVMControllerCompileBracketScope CompileScope(this);
-	FRigVMInverseAction InverseAction;
+	FRigVMBaseAction Action;
 	if (bSetupUndoRedo)
 	{
-		InverseAction.Title = TEXT("Eject node");
+		Action.Title = TEXT("Eject node");
 
-		ActionStack->BeginAction(InverseAction);
-		ActionStack->AddAction(FRigVMInjectNodeIntoPinAction(Injection));
+		ActionStack->BeginAction(Action);
+		ActionStack->AddAction(FRigVMEjectNodeFromPinAction(Injection));
 	}
 
 	FVector2D Position = InPin->GetNode()->GetPosition() + FVector2D(0.f, 12.f) * float(InPin->GetPinIndex());
@@ -3549,7 +3541,7 @@ URigVMNode* URigVMController::EjectNodeFromPin(URigVMPin* InPin, bool bSetupUndo
 
 	if (bSetupUndoRedo)
 	{
-		ActionStack->EndAction(InverseAction);
+		ActionStack->EndAction(Action);
 	}
 
 	if (bPrintPythonCommand)
@@ -3656,7 +3648,7 @@ bool URigVMController::CancelUndoBracket()
 	return ActionStack->CancelUndoBracket(this);
 }
 
-FString URigVMController::ExportNodesToText(const TArray<FName>& InNodeNames)
+FString URigVMController::ExportNodesToText(const TArray<FName>& InNodeNames, bool bIncludeExteriorLinks)
 {
 	if (!IsValidGraph())
 	{
@@ -3697,27 +3689,28 @@ FString URigVMController::ExportNodesToText(const TArray<FName>& InNodeNames)
 
 	for (URigVMLink* Link : Graph->Links)
 	{
-		URigVMPin* SourcePin = Link->GetSourcePin();
-		URigVMPin* TargetPin = Link->GetTargetPin();
+		const URigVMPin* SourcePin = Link->GetSourcePin();
+		const URigVMPin* TargetPin = Link->GetTargetPin();
 		if (SourcePin && TargetPin && IsValid(SourcePin) && IsValid(TargetPin))
 		{
-			if (!AllNodeNames.Contains(SourcePin->GetNode()->GetFName()))
-			{
-				continue;
-			}
-			if (!AllNodeNames.Contains(TargetPin->GetNode()->GetFName()))
+			const bool bSourceNodeIsPartOfExport = AllNodeNames.Contains(SourcePin->GetNode()->GetFName());
+			const bool bTargetNodeIsPartOfExport = AllNodeNames.Contains(TargetPin->GetNode()->GetFName());
+
+			const bool bShouldExport = (bSourceNodeIsPartOfExport && bTargetNodeIsPartOfExport) ||
+				(bIncludeExteriorLinks && (bSourceNodeIsPartOfExport != bTargetNodeIsPartOfExport));
+
+			if(bShouldExport)
 			{
-				continue;
+				Link->UpdatePinPaths();
+				UExporter::ExportToOutputDevice(&Context, Link, NULL, Archive, TEXT("copy"), 0, PPF_ExportsNotFullyQualified | PPF_Copy | PPF_Delimited, false, Link->GetOuter());
 			}
-			Link->PrepareForCopy();
-			UExporter::ExportToOutputDevice(&Context, Link, NULL, Archive, TEXT("copy"), 0, PPF_ExportsNotFullyQualified | PPF_Copy | PPF_Delimited, false, Link->GetOuter());
 		}
 	}
 
 	return MoveTemp(Archive);
 }
 
-FString URigVMController::ExportSelectedNodesToText()
+FString URigVMController::ExportSelectedNodesToText(bool bIncludeExteriorLinks)
 {
 	if (!IsValidGraph())
 	{
@@ -3727,7 +3720,7 @@ FString URigVMController::ExportSelectedNodesToText()
 	URigVMGraph* Graph = GetGraph();
 	check(Graph);
 
-	return ExportNodesToText(Graph->GetSelectNodes());
+	return ExportNodesToText(Graph->GetSelectNodes(), bIncludeExteriorLinks);
 }
 
 struct FRigVMControllerObjectFactory : public FCustomizableTextObjectFactory
@@ -3735,44 +3728,57 @@ struct FRigVMControllerObjectFactory : public FCustomizableTextObjectFactory
 public:
 	URigVMController* Controller;
 	TArray<URigVMNode*> CreatedNodes;
-	TArray<FName> CreateNodeNames;
-	TMap<FName, FName> NodeNameMap;
+	TMap<const URigVMGraph*, TArray<FName>> CreateNodeNamesPerGraph;
+	TMap<const URigVMGraph*, TMap<FString, FString>> NodeNameMapPerGraph;
+	URigVMController::FRestoreLinkedPathSettings RestoreLinksSettings; 
 	TArray<URigVMLink*> CreatedLinks;
+	TArray<URigVMGraph*> CreatedGraphs;
 public:
 	FRigVMControllerObjectFactory(URigVMController* InController)
 		: FCustomizableTextObjectFactory(GWarn)
 		, Controller(InController)
 	{
+		RestoreLinksSettings.bIsImportingFromText = true;
 	}
 
 protected:
 	virtual bool CanCreateClass(UClass* ObjectClass, bool& bOmitSubObjs) const override
 	{
-		if (URigVMNode* DefaultNode = Cast<URigVMNode>(ObjectClass->GetDefaultObject()))
-		{
-			// bOmitSubObjs = true;
-			return true;
-		}
-		if (URigVMLink* DefaultLink = Cast<URigVMLink>(ObjectClass->GetDefaultObject()))
-		{
-			return true;
-		}
-
-		return false;
+		const UObject* DefaultObject = ObjectClass->GetDefaultObject(); 
+		return DefaultObject->IsA<URigVMNode>() || 
+			DefaultObject->IsA<URigVMGraph>() ||
+			DefaultObject->IsA<URigVMLink>() ||
+			DefaultObject->IsA<URigVMInjectionInfo>() ||
+			DefaultObject->IsA<URigVMPin>();
 	}
 
-	virtual void UpdateObjectName(UClass* ObjectClass, FName& InOutObjName) override
+	virtual void UpdateObjectName(UClass* ObjectClass, UObject* InParent, FName& InOutObjName) override
 	{
-		if (URigVMNode* DefaultNode = Cast<URigVMNode>(ObjectClass->GetDefaultObject()))
+		if (ObjectClass->GetDefaultObject()->IsA<URigVMNode>())
 		{
-			URigVMGraph* Graph = Controller->GetGraph();
+			const URigVMGraph* Graph = Cast<URigVMGraph>(InParent);
+			if(Graph == nullptr)
+			{
+				if(const URigVMInjectionInfo* InjectionInfo = Cast<URigVMInjectionInfo>(InParent))
+				{
+					Graph = InjectionInfo->GetGraph();
+				}
+			}
+
 			check(Graph);
 
-			const FName ValidName = Controller->GetUniqueName(InOutObjName, [Graph, this](const FName& InName) {
+			TArray<FName>& CreateNodeNames = CreateNodeNamesPerGraph.FindOrAdd(Graph);
+			TMap<FString, FString>& NodeNameMap = NodeNameMapPerGraph.FindOrAdd(Graph);
+
+			const FName ValidName = URigVMController::GetUniqueName(InOutObjName, [Graph, this, CreateNodeNames](const FName& InName) {
 				return !CreateNodeNames.Contains(InName) && Graph->IsNameAvailable(InName.ToString());
 			}, false, true);
+
+			if(!InOutObjName.IsEqual(ValidName, ENameCase::CaseSensitive))
+			{
+				NodeNameMap.Add(InOutObjName.ToString(), ValidName.ToString());
+			}
 			
-			NodeNameMap.Add(InOutObjName, ValidName);
 			CreateNodeNames.Add(ValidName);
 			InOutObjName = ValidName;
 		}
@@ -3780,27 +3786,56 @@ protected:
 
 	virtual void ProcessConstructedObject(UObject* CreatedObject) override
 	{
-		if (URigVMNode* CreatedNode = Cast<URigVMNode>(CreatedObject))
+		if (URigVMGraph* Graph = Cast<URigVMGraph>(CreatedObject))
 		{
-			CreatedNodes.AddUnique(CreatedNode);
+			CreatedGraphs.Add(Graph);
+			(void)CreateNodeNamesPerGraph.FindOrAdd(Graph);
+			(void)NodeNameMapPerGraph.FindOrAdd(Graph);
 
-			for (URigVMPin* Pin : CreatedNode->GetPins())
+			for(URigVMNode* Node : Graph->GetNodes())
 			{
-				for (URigVMInjectionInfo* Injection : Pin->GetInjectedNodes())
-				{
-					ProcessConstructedObject(Injection->Node);
+				ProcessConstructedObject(Node);
+			}
+			for(URigVMLink* Link : Graph->GetLinks())
+			{
+				ProcessConstructedObject(Link);
+			}
+		}
+		else if (URigVMNode* Node = Cast<URigVMNode>(CreatedObject))
+		{
+			CreatedNodes.Add(Node);
 
-					FName NewName = Injection->Node->GetFName();
-					UpdateObjectName(URigVMNode::StaticClass(), NewName);
-					Controller->RenameObject(Injection->Node, *NewName.ToString(), nullptr);
-					Injection->InputPin = Injection->InputPin ? Injection->Node->FindPin(Injection->InputPin->GetName()) : nullptr;
-					Injection->OutputPin = Injection->OutputPin ? Injection->Node->FindPin(Injection->OutputPin->GetName()) : nullptr;
-				}
+			for (URigVMPin* Pin : Node->GetPins())
+			{
+				ProcessConstructedObject(Pin);
+			}
+
+			if(const URigVMCollapseNode* CollapseNode = Cast<URigVMCollapseNode>(Node))
+			{
+				ProcessConstructedObject(CollapseNode->GetContainedGraph());
 			}
 		}
-		else if (URigVMLink* CreatedLink = Cast<URigVMLink>(CreatedObject))
+		else if (URigVMLink* Link = Cast<URigVMLink>(CreatedObject))
 		{
-			CreatedLinks.Add(CreatedLink);
+			CreatedLinks.Add(Link);
+		}
+		else if (const URigVMPin* Pin = Cast<URigVMPin>(CreatedObject))
+		{
+			for(URigVMInjectionInfo* InjectionInfo : Pin->GetInjectedNodes())
+			{
+				ProcessConstructedObject(InjectionInfo);
+			}
+		}
+		else if (URigVMInjectionInfo* Injection = Cast<URigVMInjectionInfo>(CreatedObject))
+		{
+			URigVMNode* InjectedNode = Injection->Node;
+			ProcessConstructedObject(InjectedNode);
+
+			FName NewName = InjectedNode->GetFName();
+			UpdateObjectName(URigVMNode::StaticClass(), Injection->GetGraph(), NewName);
+			Controller->RenameObject(InjectedNode, *NewName.ToString(), nullptr);
+			Injection->InputPin = Injection->InputPin ? Injection->Node->FindPin(Injection->InputPin->GetName()) : nullptr;
+			Injection->OutputPin = Injection->OutputPin ? Injection->Node->FindPin(Injection->OutputPin->GetName()) : nullptr;
 		}
 	}
 };
@@ -3818,377 +3853,268 @@ bool URigVMController::CanImportNodesFromText(const FString& InText)
 	}
 
 	FRigVMControllerObjectFactory Factory(this);
-	if (!Factory.CanCreateObjectsFromText(InText))
-	{
-		return false;
-	}
-
-	URigVMGraph* Graph = GetGraph();
-	check(Graph);
-	Factory.ProcessBuffer(Graph, RF_Transactional, InText);
-
-	if (Factory.CreatedNodes.Num() == 0)
-	{
-		return false;
-	}
-
-	// If the graph is a function library, make sure we are pasting unconnected collapse nodes
-	if (Graph->IsA<URigVMFunctionLibrary>())
-	{
-		if (!Factory.CreatedLinks.IsEmpty())
-		{
-			return false;
-		}
-		
-		for (URigVMNode* Node : Factory.CreatedNodes)
-		{
-			if (!Node->IsA<URigVMCollapseNode>())
-			{
-				return false;
-			}
-		}
-	}
-
-	return true;
+	return Factory.CanCreateObjectsFromText(InText);
 }
 
 TArray<FName> URigVMController::ImportNodesFromText(const FString& InText, bool bSetupUndoRedo, bool bPrintPythonCommands)
 {
-	TArray<FName> NodeNames;
 	if (!IsValidGraph())
 	{
-		return NodeNames;
+		return TArray<FName>();
 	}
 
 	if (!bIsTransacting && !IsGraphEditable())
 	{
-		return NodeNames;
+		return TArray<FName>();
 	}
 
-	URigVMGraph* Graph = GetGraph();
-	check(Graph);
+	check(GetGraph());
 
-	FRigVMControllerObjectFactory Factory(this);
-	Factory.ProcessBuffer(Graph, RF_Transactional, InText);
+	const TArray<URigVMNode*> NodesPriorImport = GetGraph()->GetNodes();
+	const TArray<URigVMLink*> LinksPriorImport = GetGraph()->GetLinks();
 
-	if (Factory.CreatedNodes.Num() == 0)
-	{
-		return NodeNames;
-	}
+	FRigVMControllerObjectFactory Factory(this);
+	Factory.ProcessBuffer(GetGraph(), RF_Transactional, InText);
 
-	if (bSetupUndoRedo)
+	if (Factory.CreatedNodes.IsEmpty() && Factory.CreatedLinks.IsEmpty())
 	{
-		OpenUndoBracket(TEXT("Importing Nodes from Text"));
+		return TArray<FName>();
 	}
+	
+	// sort the graphs, nodes and links based on the depth - first resolve the links in leaf graphs
+	// and then work your way outward
+	const TArray<URigVMGraph*> ImportedGraphs = Factory.CreatedGraphs;
+	const TArray<URigVMNode*> ImportedNodes = Factory.CreatedNodes;
+	const TArray<URigVMLink*> ImportedLinks = Factory.CreatedLinks;
+	SortGraphElementsByGraphDepth(Factory.CreatedGraphs, true);
+	SortGraphElementsByGraphDepth(Factory.CreatedNodes, true);
+	SortGraphElementsByGraphDepth(Factory.CreatedLinks, true);
 
 	FRigVMControllerCompileBracketScope CompileScope(this);
-	FRigVMInverseAction AddNodesAction;
-	if (bSetupUndoRedo)
+
+	TArray<TGuardValue<bool>> EditGuards;
+	for (URigVMGraph* CreatedGraph : Factory.CreatedGraphs)
 	{
-		ActionStack->BeginAction(AddNodesAction);
+		EditGuards.Emplace(CreatedGraph->bEditable, true);
 	}
 
-	TArray<TGuardValue<bool>> EditGuards;
-	TArray<URigVMNode*> CreatedNodes = Factory.CreatedNodes;
-	for (int32 i=0; i<CreatedNodes.Num(); ++i)
+	TMap<URigVMGraph*, TArray<FName>> NodeNamesPerGraph;
 	{
-		URigVMNode* CreatedNode = CreatedNodes[i];
-		if (URigVMCollapseNode* CollapseNode = Cast<URigVMCollapseNode>(CreatedNode))
+		TArray<URigVMNode*> FilteredNodes;
+		for (URigVMNode* CreatedNode : Factory.CreatedNodes)
 		{
-			if (URigVMGraph* ContainedGraph = CollapseNode->GetContainedGraph())
+			URigVMGraph* Graph = CreatedNode->GetTypedOuter<URigVMGraph>();
+			if(Graph == nullptr)
+			{
+				DestroyObject(CreatedNode);
+				continue;
+			}
+
+			FRigVMControllerGraphGuard GraphGuard(this, Graph, false);
+			if(!CanAddNode(CreatedNode, true))
 			{
-				EditGuards.Emplace(ContainedGraph->bEditable, true);
-				CreatedNodes.Append(ContainedGraph->GetNodes());
+				Graph->Nodes.Remove(CreatedNode);
+				DestroyObject(CreatedNode);
+				continue;
 			}
+
+			Graph->Nodes.AddUnique(CreatedNode);
+			FilteredNodes.Add(CreatedNode);
 		}
+
+		Swap(Factory.CreatedNodes, FilteredNodes);
 	}
 
-	FRigVMUnitNodeCreatedContext::FScope UnitNodeCreatedScope(UnitNodeCreatedContext, ERigVMNodeCreatedReason::Paste);
-	TMap<URigVMPin*, TRigVMTypeIndex> TypesWhenPasted;
-	for (URigVMNode* CreatedNode : Factory.CreatedNodes)
+	// the links may already be in the graph's links property
+	// due to the serialization - remove them since we need.
+	// but we also want to maintain the order of the links as they
+	// are getting created. we'll create an array 
+	for (URigVMLink* CreatedLink : Factory.CreatedLinks)
 	{
-		if(!CanAddNode(CreatedNode, true))
+		if(URigVMGraph* Graph = CreatedLink->GetTypedOuter<URigVMGraph>())
 		{
-			continue;
+			Graph->Links.Remove(CreatedLink);
 		}
+	}
 
-		Graph->Nodes.Add(CreatedNode);
+	FRigVMUnitNodeCreatedContext::FScope UnitNodeCreatedScope(UnitNodeCreatedContext, ERigVMNodeCreatedReason::Paste);
+	for (URigVMNode* CreatedNode : Factory.CreatedNodes)
+	{
+		URigVMGraph* Graph = CreatedNode->GetTypedOuter<URigVMGraph>();
+		check(Graph);
+		FRigVMControllerGraphGuard GraphGuard(this, Graph, false);
 
-		if (bSetupUndoRedo)
+		// Refresh the unit node to account for changes in node color, pin additions, pin order, etc
+		if (URigVMUnitNode* UnitNode = Cast<URigVMUnitNode>(CreatedNode))
 		{
-			if (!CreatedNode->IsInjected() || !CreatedNode->IsA<URigVMVariableNode>())
+			if(UnitNode->ResolvedFunctionName.IsEmpty())
 			{
-				ActionStack->AddAction(FRigVMRemoveNodeAction(CreatedNode, this));
-			}
-		}
-
-		// find all nodes affected by this
-		TArray<URigVMNode*> SubNodes;
-		SubNodes.Add(CreatedNode);
+				if(UScriptStruct* ScriptStruct = UnitNode->GetScriptStruct())
+				{
+					const FName MethodName = UnitNode->GetMethodName();
+					if(!MethodName.IsNone())
+					{
+						if(const FRigVMFunction* Function = FRigVMRegistry::Get().FindFunction(ScriptStruct, *MethodName.ToString()))
+						{
+							UnitNode->ResolvedFunctionName = Function->GetName();
 
-		// Refresh the unit nodes to account for changes in node color, pin additions, pin order, etc
-		for(int32 SubNodeIndex=0; SubNodeIndex < SubNodes.Num(); SubNodeIndex++)
-		{
-			if (URigVMUnitNode* UnitNode = Cast<URigVMUnitNode>(SubNodes[SubNodeIndex]))
-			{
-				TGuardValue<bool> SuspendNotifications(bSuspendNotifications, true);
-				RepopulatePinsOnNode(UnitNode);
+							if(const FRigVMTemplate* Template = Function->GetTemplate())
+							{
+								UnitNode->TemplateNotation = Template->GetNotation();
+							}
+						}
+					}
+				}
 			}
+			
+			TGuardValue<bool> SuspendNotifications(bSuspendNotifications, true);
+			RepopulatePinsOnNode(UnitNode);
 		}
 
-		for(int32 SubNodeIndex=0; SubNodeIndex < SubNodes.Num(); SubNodeIndex++)
+		if (URigVMUnitNode* UnitNode = Cast<URigVMUnitNode>(CreatedNode))
 		{
-			if(URigVMCollapseNode* CollapseNode = Cast<URigVMCollapseNode>(SubNodes[SubNodeIndex]))
+			if (UnitNodeCreatedContext.IsValid())
 			{
+				if (TSharedPtr<FStructOnScope> StructScope = UnitNode->ConstructStructInstance())
 				{
-					FRigVMControllerGraphGuard GraphGuard(this, CollapseNode->GetContainedGraph(), false);
-					TGuardValue<bool> GuardEditGraph(CollapseNode->ContainedGraph->bEditable, true);
-					TGuardValue<bool> SuspendNotifications(bSuspendNotifications, true);
-					ReattachLinksToPinObjects();
+					TGuardValue<FName> NodeNameScope(UnitNodeCreatedContext.NodeName, UnitNode->GetFName());
+					FRigVMStruct* StructInstance = (FRigVMStruct*)StructScope->GetStructMemory();
+					StructInstance->OnUnitNodeCreated(UnitNodeCreatedContext);
 				}
-				
-				SubNodes.Append(CollapseNode->GetContainedNodes());
 			}
 		}
 
-		if (URigVMLibraryNode* LibraryNode = Cast<URigVMLibraryNode>(CreatedNode))
+		if (URigVMFunctionReferenceNode* FunctionRefNode = Cast<URigVMFunctionReferenceNode>(CreatedNode))
 		{
-			TFunction<void(URigVMLibraryNode*)> Recompute = [&](URigVMLibraryNode* OuterNode)
+			if(URigVMBuildData* BuildData = URigVMBuildData::Get())
 			{
-				for (URigVMNode* Node : OuterNode->GetContainedNodes())
-				{
-					if (URigVMLibraryNode* ContainedNode = Cast<URigVMLibraryNode>(Node))
-					{
-						Recompute(ContainedNode);
-					}
-				}
+				BuildData->RegisterFunctionReference(FunctionRefNode->GetReferencedFunctionHeader().LibraryPointer, FunctionRefNode);
+			}
+		}
 
-				// If the OuterNode is referencing a non-existing function, the contained graph might be empty
-				if (OuterNode->GetContainedGraph() == nullptr)
-				{
-					return;
-				}
-			};
-			Recompute(LibraryNode);
+		for(URigVMPin* Pin : CreatedNode->Pins)
+		{
+			EnsurePinValidity(Pin, true);
 		}
 
-		for(URigVMNode* SubNode : SubNodes)
+		Notify(ERigVMGraphNotifType::NodeAdded, CreatedNode);
+
+		TArray<FName>& NodeNames = NodeNamesPerGraph.FindOrAdd(Graph);
+		NodeNames.Add(CreatedNode->GetFName());
+	}
+
+	if (Factory.CreatedLinks.Num() > 0)
+	{
+		URigVMGraph* LastGraph = nullptr;
+		for (URigVMLink* CreatedLink : Factory.CreatedLinks)
 		{
-			FRigVMControllerGraphGuard GraphGuard(this, SubNode->GetGraph(), false);
-			if (URigVMUnitNode* UnitNode = Cast<URigVMUnitNode>(SubNode))
+			URigVMGraph* Graph = CreatedLink->GetTypedOuter<URigVMGraph>();
+			if(Graph == nullptr)
 			{
-				if (UnitNodeCreatedContext.IsValid())
-				{
-					if (TSharedPtr<FStructOnScope> StructScope = UnitNode->ConstructStructInstance())
-					{
-						TGuardValue<FName> NodeNameScope(UnitNodeCreatedContext.NodeName, UnitNode->GetFName());
-						FRigVMStruct* StructInstance = (FRigVMStruct*)StructScope->GetStructMemory();
-						StructInstance->OnUnitNodeCreated(UnitNodeCreatedContext);
-					}
-				}
+				DestroyObject(CreatedLink);
+				continue;
 			}
 
-			if (URigVMFunctionReferenceNode* FunctionRefNode = Cast<URigVMFunctionReferenceNode>(SubNode))
+			if(Graph != LastGraph)
 			{
-				if(URigVMBuildData* BuildData = URigVMBuildData::Get())
+				Factory.RestoreLinksSettings.NodeNameMap = Factory.NodeNameMapPerGraph.FindChecked(Graph);
+				LastGraph = Graph;
+			}
+
+			FRigVMControllerGraphGuard GraphGuard(this, Graph, false);
+
+			// this takes care of mapping from the old to the new nodes etc
+			const TArray<FLinkedPath> LinkedPaths = RemapLinkedPaths(
+				GetLinkedPaths({CreatedLink}), 
+				Factory.RestoreLinksSettings,
+				false);
+			const FLinkedPath& LinkedPath = LinkedPaths[0]; 
+
+			URigVMPin* SourcePin = LinkedPath.GetSourcePin();
+			URigVMPin* TargetPin = LinkedPath.GetTargetPin();
+			
+			if (SourcePin && TargetPin)
+			{
+				// update the paths on the link to allow reuse of the UObject
+				CreatedLink->SetSourceAndTargetPinPaths(LinkedPath.SourcePinPath, LinkedPath.TargetPinPath);
+				
+				if (TargetPin->IsBoundToVariable())
+				{
+					const FString VariableNodeName = TargetPin->GetBoundVariableNode()->GetName();
+					const FString BindingPath = TargetPin->GetBoundVariablePath();
+
+					// The current situation is that the outer pin has an injection info, and the injected node exists
+					// but the injected node is not linked to the outer pin. BreakAllLinks will try to unbind the outer pin,
+					// for that to be successful, the binding needs to be complete
+					// Connect it so that the unbound is successful
+					check(!SourcePin->IsLinkedTo(TargetPin));
+					Graph->DetachedLinks.Add(CreatedLink);
+					AddLink(SourcePin, TargetPin, false);
+
+					// recreate binding
+					UnbindPinFromVariable(TargetPin, false);
+					BindPinToVariable(TargetPin, BindingPath, false, VariableNodeName);
+				}
+				else
 				{
-					BuildData->RegisterFunctionReference(FunctionRefNode->GetReferencedFunctionHeader().LibraryPointer, FunctionRefNode);
+					Graph->DetachedLinks.Add(CreatedLink);
+					AddLink(SourcePin, TargetPin, false);
 				}
 			}
 
-			for(URigVMPin* Pin : SubNode->Pins)
+			// if the link is still part of the detached link array
+			// the restore was not successful.
+			if(Graph->DetachedLinks.Remove(CreatedLink))
 			{
-				EnsurePinValidity(Pin, true);
+				ReportErrorf(TEXT("Cannot import link '%s'."), *URigVMLink::GetPinPathRepresentation(CreatedLink->GetSourcePinPath(), CreatedLink->GetTargetPinPath()));
+				DestroyObject(CreatedLink);
 			}
 		}
+	}
 
-		Notify(ERigVMGraphNotifType::NodeAdded, CreatedNode);
+	// order nodes and links per graph again in the order that they were originally created.
+	SortGraphElementsByImportOrder(GetGraph()->Nodes, NodesPriorImport, ImportedNodes);
+	for(URigVMGraph* CreatedGraph : Factory.CreatedGraphs)
+	{
+		SortGraphElementsByImportOrder(CreatedGraph->Nodes, {}, ImportedNodes);
+	}
 
-		NodeNames.Add(CreatedNode->GetFName());
+	SortGraphElementsByImportOrder(GetGraph()->Links, LinksPriorImport, ImportedLinks);
+	for(URigVMGraph* CreatedGraph : Factory.CreatedGraphs)
+	{
+		SortGraphElementsByImportOrder(CreatedGraph->Links, {}, ImportedLinks);
+	}
+
+	if(!NodeNamesPerGraph.Contains(GetGraph()))
+	{
+		ReportError(TEXT("Unexpected failure during ImportNodesFromText"));
+		return TArray<FName>();
 	}
+	
+	const TArray<FName>& NodeNames = NodeNamesPerGraph.FindChecked(GetGraph());
 
 	if (bSetupUndoRedo)
 	{
-		ActionStack->EndAction(AddNodesAction);
+		FRigVMImportFromTextAction Action(InText, NodeNames);
+		Action.Title = TEXT("Importing Nodes from Text");
+		ActionStack->AddAction(Action);
 	}
 
-	if (Factory.CreatedLinks.Num() > 0)
+#if WITH_EDITOR
+	if (bPrintPythonCommands && !NodeNames.IsEmpty())
 	{
-		FRigVMBaseAction AddLinksAction;
-		if (bSetupUndoRedo)
-		{
-			ActionStack->BeginAction(AddLinksAction);
-		}
+		FString PythonContent = InText.Replace(TEXT("\\\""), TEXT("\\\\\""));
+		PythonContent = InText.Replace(TEXT("'"), TEXT("\\'"));
+		PythonContent = PythonContent.Replace(TEXT("\r\n"), TEXT("\\r\\n'\r\n'"));
 
-		for (URigVMLink* CreatedLink : Factory.CreatedLinks)
-		{
-			FString SourceLeft, SourceRight, TargetLeft, TargetRight;
-			if (URigVMPin::SplitPinPathAtStart(CreatedLink->SourcePinPath, SourceLeft, SourceRight) &&
-				URigVMPin::SplitPinPathAtStart(CreatedLink->TargetPinPath, TargetLeft, TargetRight))
-			{
-				const FName* NewSourceNodeName = Factory.NodeNameMap.Find(*SourceLeft);
-				const FName* NewTargetNodeName = Factory.NodeNameMap.Find(*TargetLeft);
-				if (NewSourceNodeName && NewTargetNodeName)
-				{
-					CreatedLink->SourcePinPath = URigVMPin::JoinPinPath(NewSourceNodeName->ToString(), SourceRight);
-					CreatedLink->TargetPinPath = URigVMPin::JoinPinPath(NewTargetNodeName->ToString(), TargetRight);
-					URigVMPin* SourcePin = CreatedLink->GetSourcePin();
-					URigVMPin* TargetPin = CreatedLink->GetTargetPin();
+		const FString GraphName = GetSanitizedGraphName(GetGraph()->GetGraphName());
 
-					if (SourcePin == nullptr)
-					{
-						URigVMNode* OriginalNode = Graph->FindNode(SourceLeft);
-						if (OriginalNode && OriginalNode->IsA<URigVMFunctionEntryNode>())
-						{
-							CreatedLink->SourcePinPath = URigVMPin::JoinPinPath(SourceLeft, SourceRight);
-							SourcePin = CreatedLink->GetSourcePin();
-						}
-					}
-					if (TargetPin == nullptr)
-					{
-						URigVMNode* OriginalNode = Graph->FindNode(TargetLeft);
-						if (OriginalNode && OriginalNode->IsA<URigVMFunctionReturnNode>())
-						{
-							CreatedLink->TargetPinPath = URigVMPin::JoinPinPath(TargetLeft, TargetRight);
-							TargetPin = CreatedLink->GetTargetPin();
-						}
-					}
-
-					if (SourcePin == nullptr)
-					{
-						URigVMNode* OriginalNode = Graph->FindNode(SourceLeft);
-						if (OriginalNode)
-						{
-							FString OldSourcePinPath = URigVMPin::JoinPinPath(SourceLeft, SourceRight);
-							if (URigVMPin* OldPin = Graph->FindPin(OldSourcePinPath))
-							{
-								if (OldPin->IsStructMember())
-								{
-									URigVMPin* OldRootPin = OldPin->GetRootPin();
-									FString NewSourceRootPinPath = URigVMPin::JoinPinPath(NewSourceNodeName->ToString(), OldRootPin->GetName());
-									if (URigVMPin* NewRootPin = Graph->FindPin(NewSourceRootPinPath))
-									{
-										SourcePin = CreatedLink->GetSourcePin();
-									}
-								}
-							}
-						}
-					}
-					if (TargetPin == nullptr)
-					{
-						URigVMNode* OriginalNode = Graph->FindNode(TargetLeft);
-						if (OriginalNode)
-						{
-							FString OldTargetPinPath = URigVMPin::JoinPinPath(TargetLeft, TargetRight);
-							if (URigVMPin* OldPin = Graph->FindPin(OldTargetPinPath))
-							{
-								if (OldPin->IsStructMember())
-								{
-									URigVMPin* OldRootPin = OldPin->GetRootPin();
-									FString NewTargetRootPinPath = URigVMPin::JoinPinPath(NewTargetNodeName->ToString(), OldRootPin->GetName());
-									if (URigVMPin* NewRootPin = Graph->FindPin(NewTargetRootPinPath))
-									{
-										TargetPin = CreatedLink->GetTargetPin();
-									}
-								}
-							}
-						}
-					}
-					
-					if (SourcePin && TargetPin)
-					{
-						// BreakAllLinks will unbind and destroy the injected variable node
-						// We need to rebind to recreate the variable node with the same name
-						bool bWasBinded = TargetPin->IsBoundToVariable();
-						FString VariableNodeName, BindingPath;
-						if (bWasBinded)
-						{
-							VariableNodeName = TargetPin->GetBoundVariableNode()->GetName();
-							BindingPath = TargetPin->GetBoundVariablePath();
-
-							// The current situation is that the outer pin has an injection info, and the injected node exists
-							// but the injected node is not linked to the outer pin. BreakAllLinks will try to unbind the outer pin,
-							// for that to be successful, the binding needs to be complete
-							// Connect it so that the unbound is successful
-							if (!SourcePin->IsLinkedTo(TargetPin))
-							{
-								Graph->Links.Add(CreatedLink);
-								SourcePin->Links.Add(CreatedLink);
-								TargetPin->Links.Add(CreatedLink);
-							}
-
-							if (TargetPin->IsBoundToVariable())
-							{
-								UnbindPinFromVariable(TargetPin, bSetupUndoRedo);
-							}
-						}
-						
-						BreakAllLinksRecursive(TargetPin, true, true, bSetupUndoRedo);
-						BreakAllLinks(TargetPin, true, bSetupUndoRedo);
-						BreakAllLinksRecursive(TargetPin, true, false, bSetupUndoRedo);
-
-						// recreate binding if needed
-						if (bWasBinded)
-						{
-							BindPinToVariable(TargetPin, BindingPath, bSetupUndoRedo, VariableNodeName);
-						}
-						else
-						{
-							PrepareToLink(TargetPin, SourcePin, bSetupUndoRedo);
-
-							Graph->Links.Add(CreatedLink);
-							SourcePin->Links.Add(CreatedLink);
-							TargetPin->Links.Add(CreatedLink);
-
-							if (bSetupUndoRedo)
-							{
-								ActionStack->AddAction(FRigVMAddLinkAction(SourcePin, TargetPin));
-								if (SourcePin->GetNode()->IsInjected())
-								{
-									ActionStack->AddAction(FRigVMInjectNodeIntoPinAction(SourcePin->GetTypedOuter<URigVMInjectionInfo>()));	
-								}
-								if (TargetPin->GetNode()->IsInjected())
-								{
-									ActionStack->AddAction(FRigVMInjectNodeIntoPinAction(TargetPin->GetTypedOuter<URigVMInjectionInfo>()));	
-								}
-							}
-							Notify(ERigVMGraphNotifType::LinkAdded, CreatedLink);
-						}
-						continue;
-					}
-				}
-			}
-
-			ReportErrorf(TEXT("Cannot import link '%s'."), *URigVMLink::GetPinPathRepresentation(CreatedLink->SourcePinPath, CreatedLink->TargetPinPath));
-			DestroyObject(CreatedLink);
-		}
-
-		if (bSetupUndoRedo)
-		{
-			ActionStack->EndAction(AddLinksAction);
-		}
-	}
-
-	if (bSetupUndoRedo)
-	{
-		CloseUndoBracket();
-	}
-	
-#if WITH_EDITOR
-	if (bPrintPythonCommands && !NodeNames.IsEmpty())
-	{
-		FString PythonContent = InText.Replace(TEXT("\\\""), TEXT("\\\\\""));
-		PythonContent = InText.Replace(TEXT("'"), TEXT("\\'"));
-		PythonContent = PythonContent.Replace(TEXT("\r\n"), TEXT("\\r\\n'\r\n'"));
-
-		const FString GraphName = GetSanitizedGraphName(GetGraph()->GetGraphName());
-
-		RigVMPythonUtils::Print(GetGraphOuterName(),
-			FString::Printf(TEXT("blueprint.get_controller_by_name('%s').import_nodes_from_text('%s')"),
-			*GraphName,
-			*PythonContent));
-	}
-#endif
+		RigVMPythonUtils::Print(GetGraphOuterName(),
+			FString::Printf(TEXT("blueprint.get_controller_by_name('%s').import_nodes_from_text('%s')"),
+			*GraphName,
+			*PythonContent));
+	}
+#endif
 
 	return NodeNames;
 }
@@ -5073,17 +4999,14 @@ URigVMCollapseNode* URigVMController::CollapseNodes(const TArray<URigVMNode*>& I
 		}
 	}
 
-	for (const FName& NodeToRemove : NodeNames)
-	{
-		RemoveNodeByName(NodeToRemove, false, true);
-	}
+	RemoveNodesByName(NodeNames, false, false);
 
 	if (!InCollapseNodeName.IsEmpty() && CollapseNodeName != InCollapseNodeName)
 	{		
 		FString ValidName = GetValidNodeName(InCollapseNodeName);
 		if (ValidName == InCollapseNodeName)
 		{
-			RenameNode(CollapseNode, *ValidName, bSetupUndoRedo);
+			RenameNode(CollapseNode, *ValidName, false);
 		}
 	}
 
@@ -5969,13 +5892,13 @@ URigVMFunctionReferenceNode* URigVMController::PromoteCollapseNodeToFunctionRefe
 
 			for (const URigVMLink* InnerLink : InCollapseNode->GetContainedGraph()->GetLinks())
 			{
-				URigVMPin* SourcePin = InCollapseNode->GetGraph()->FindPin(InnerLink->SourcePinPath);
-				URigVMPin* TargetPin = InCollapseNode->GetGraph()->FindPin(InnerLink->TargetPinPath);
+				URigVMPin* SourcePin = InCollapseNode->GetGraph()->FindPin(InnerLink->GetSourcePinPath());
+				URigVMPin* TargetPin = InCollapseNode->GetGraph()->FindPin(InnerLink->GetTargetPinPath());
 				if (SourcePin && TargetPin)
 				{
 					if (!SourcePin->IsLinkedTo(TargetPin))
 					{
-						AddLink(InnerLink->SourcePinPath, InnerLink->TargetPinPath, false);	
+						AddLink(InnerLink->GetSourcePinPath(), InnerLink->GetTargetPinPath(), false);	
 					}
 				}				
 			}
@@ -5989,13 +5912,9 @@ URigVMFunctionReferenceNode* URigVMController::PromoteCollapseNodeToFunctionRefe
 		FVector2D NodePosition = InCollapseNode->GetPosition();
 		TMap<FString, FPinState> PinStates = GetPinStates(InCollapseNode);
 
-		TArray<URigVMLink*> Links = InCollapseNode->GetLinks();
-		TArray< TPair< FString, FString > > LinkPaths;
-		for (URigVMLink* Link : Links)
-		{
-			LinkPaths.Add(TPair< FString, FString >(Link->GetSourcePin()->GetPinPath(), Link->GetTargetPin()->GetPinPath()));
-		}
+		const TArray<FLinkedPath> LinkedPaths = GetLinkedPaths(InCollapseNode, true);
 
+		FastBreakLinkedPaths(LinkedPaths);
 		RemoveNode(InCollapseNode, false, true);
 
 		FunctionRefNode = AddFunctionReferenceNode(FunctionDefinition, NodePosition, NodeName, false);
@@ -6003,10 +5922,7 @@ URigVMFunctionReferenceNode* URigVMController::PromoteCollapseNodeToFunctionRefe
 		if (FunctionRefNode)
 		{
 			ApplyPinStates(FunctionRefNode, PinStates);
-			for (const TPair<FString, FString>& LinkPath : LinkPaths)
-			{
-				AddLink(LinkPath.Key, LinkPath.Value, false);
-			}
+			RestoreLinkedPaths(LinkedPaths);
 		}
 
 		if (bSetupUndoRedo)
@@ -6091,18 +6007,14 @@ URigVMCollapseNode* URigVMController::PromoteFunctionReferenceNodeToCollapseNode
 	FVector2D NodePosition = InFunctionRefNode->GetPosition();
 	TMap<FString, FPinState> PinStates = GetPinStates(InFunctionRefNode);
 
-	TArray<URigVMLink*> Links = InFunctionRefNode->GetLinks();
-	TArray< TPair< FString, FString > > LinkPaths;
-	for (URigVMLink* Link : Links)
-	{
-		LinkPaths.Add(TPair< FString, FString >(Link->GetSourcePin()->GetPinPath(), Link->GetTargetPin()->GetPinPath()));
-	}
+	const TArray<FLinkedPath> LinkedPaths = GetLinkedPaths(InFunctionRefNode, true);
 
 	if (bSetupUndoRedo)
 	{
 		ActionStack->AddAction(FRigVMPromoteNodeAction(InFunctionRefNode, NodeName, FunctionDefinition->GetPathName()));
 	}
 
+	FastBreakLinkedPaths(LinkedPaths);
 	RemoveNode(InFunctionRefNode, false, true);
 
 	if (RequestNewExternalVariableDelegate.IsBound())
@@ -6116,16 +6028,20 @@ URigVMCollapseNode* URigVMController::PromoteFunctionReferenceNodeToCollapseNode
 	URigVMCollapseNode* CollapseNode = DuplicateObject<URigVMCollapseNode>(FunctionDefinition, Graph, *NodeName);
 	if(CollapseNode)
 	{
+		FRestoreLinkedPathSettings RestoreSettings;
+
 		{
 			FRigVMControllerGraphGuard Guard(this, CollapseNode->GetContainedGraph(), false);
 			TGuardValue<bool> SuspendNotifications(bSuspendNotifications, true);
-			ReattachLinksToPinObjects();
 
 			for (URigVMNode* Node : CollapseNode->GetContainedGraph()->GetNodes())
 			{
 				if (URigVMVariableNode* VariableNode = Cast<URigVMVariableNode>(Node))
 				{
+					const TArray<FLinkedPath> VariableLinkedPaths = GetLinkedPaths(VariableNode);
+					FastBreakLinkedPaths(VariableLinkedPaths);
 					RepopulatePinsOnNode(VariableNode, true, false, true);
+					RestoreLinkedPaths(VariableLinkedPaths, RestoreSettings);
 				}
 			}
 
@@ -6138,10 +6054,7 @@ URigVMCollapseNode* URigVMController::PromoteFunctionReferenceNodeToCollapseNode
 		Notify(ERigVMGraphNotifType::NodeAdded, CollapseNode);
 
 		ApplyPinStates(CollapseNode, PinStates);
-		for (const TPair<FString, FString>& LinkPath : LinkPaths)
-		{
-			AddLink(LinkPath.Key, LinkPath.Value, false);
-		}
+		RestoreLinkedPaths(LinkedPaths, RestoreSettings);
 	}
 
 	if(bRemoveFunctionDefinition)
@@ -6194,11 +6107,15 @@ void URigVMController::RefreshFunctionPins(URigVMNode* InNode)
 
 	if (EntryNode || ReturnNode)
 	{
+		const TArray<FLinkedPath> LinkedPaths = GetLinkedPaths(InNode);
+		FastBreakLinkedPaths(LinkedPaths);
 		RepopulatePinsOnNode(InNode, false, false, true);
+		FRestoreLinkedPathSettings Settings;
+		RestoreLinkedPaths(LinkedPaths, Settings);
 	}
 }
 
-void URigVMController::ReportRemovedLink(const FString& InSourcePinPath, const FString& InTargetPinPath)
+void URigVMController::ReportRemovedLink(const FString& InSourcePinPath, const FString& InTargetPinPath, const FString& Reason)
 {
 	if(bSuspendNotifications)
 	{
@@ -6216,21 +6133,21 @@ void URigVMController::ReportRemovedLink(const FString& InSourcePinPath, const F
 		TargetSegmentPath = InTargetPinPath;
 	}
 	
-	ReportWarningf(TEXT("Link '%s' -> '%s' was removed."), *InSourcePinPath, *InTargetPinPath);
+	ReportWarningf(TEXT("Link '%s' -> '%s' was removed.%s"), *InSourcePinPath, *InTargetPinPath, *(Reason.IsEmpty() ? Reason : TEXT(" ") + Reason));
 	SendUserFacingNotification(
 		FString::Printf(TEXT("Link to target pin '%s' was removed."), *TargetSegmentPath),
 		0.f, TargetPin, TEXT("MessageLog.Note")
 	);
 }
 
-bool URigVMController::RemoveNode(URigVMNode* InNode, bool bSetupUndoRedo, bool bRecursive, bool bPrintPythonCommand, bool bRelinkPins)
+bool URigVMController::RemoveNodes(TArray<URigVMNode*> InNodes, bool bSetupUndoRedo, bool bPrintPythonCommand)
 {
 	if (!bIsTransacting && !IsGraphEditable())
 	{
 		return false;
 	}
-	
-	if (!IsValidNodeForGraph(InNode))
+
+	if(InNodes.Num() == 0)
 	{
 		return false;
 	}
@@ -6238,78 +6155,91 @@ bool URigVMController::RemoveNode(URigVMNode* InNode, bool bSetupUndoRedo, bool
 	URigVMGraph* Graph = GetGraph();
 	check(Graph);
 
-	if (InNode->IsInjected())
+	bool bBulkEditDialogShown = false;
+	for(const URigVMNode* InNode : InNodes)
 	{
-		URigVMInjectionInfo* InjectionInfo = InNode->GetInjectionInfo();
-		if (InjectionInfo->GetPin()->GetInjectedNodes().Last() != InjectionInfo)
+		if (!IsValidNodeForGraph(InNode))
 		{
-			ReportErrorf(TEXT("Cannot remove injected node %s as it is not the last injection on the pin"), *InNode->GetNodePath());
 			return false;
 		}
-	}
 
-	if (bSetupUndoRedo)
-	{
-		// don't allow deletion of function entry / return nodes
-		if ((Cast<URigVMFunctionEntryNode>(InNode) != nullptr && InNode->GetName() == TEXT("Entry")) ||
-			(Cast<URigVMFunctionReturnNode>(InNode) != nullptr && InNode->GetName() == TEXT("Return")))
+		if (InNode->IsInjected())
 		{
-			// due to earlier bugs in the copy & paste code entry and return nodes could end up in
-			// root graphs - in those cases we allow deletion
-			if(!Graph->IsRootGraph())
+			URigVMInjectionInfo* InjectionInfo = InNode->GetInjectionInfo();
+			if (InjectionInfo->GetPin()->GetInjectedNodes().Last() != InjectionInfo)
 			{
+				ReportErrorf(TEXT("Cannot remove injected node %s as it is not the last injection on the pin"), *InNode->GetNodePath());
 				return false;
 			}
 		}
 
-		// check if the operation will cause to dirty assets
-		if(URigVMVariableNode* VariableNode = Cast<URigVMVariableNode>(InNode))
+		if (bSetupUndoRedo)
 		{
-			if(URigVMFunctionLibrary* OuterLibrary = Graph->GetTypedOuter<URigVMFunctionLibrary>())
+			// don't allow deletion of function entry / return nodes
+			if ((Cast<URigVMFunctionEntryNode>(InNode) != nullptr && InNode->GetName() == TEXT("Entry")) ||
+				(Cast<URigVMFunctionReturnNode>(InNode) != nullptr && InNode->GetName() == TEXT("Return")))
 			{
-				if(URigVMLibraryNode* OuterFunction = OuterLibrary->FindFunctionForNode(Graph->GetTypedOuter<URigVMCollapseNode>()))
+				// due to earlier bugs in the copy & paste code entry and return nodes could end up in
+				// root graphs - in those cases we allow deletion
+				if(!Graph->IsRootGraph())
 				{
-					const FName VariableToRemove = VariableNode->GetVariableName();
-					bool bIsLocalVariable = false;
-					for (FRigVMGraphVariableDescription VariableDescription : OuterFunction->GetContainedGraph()->LocalVariables)
-					{
-						if (VariableDescription.Name == VariableToRemove)
-						{
-							bIsLocalVariable = true;
-							break;
-						}
-					}
+					return false;
+				}
+			}
 
-					if (!bIsLocalVariable)
+			// check if the operation will cause to dirty assets
+			if(!bBulkEditDialogShown)
+			{
+				if(const URigVMVariableNode* VariableNode = Cast<URigVMVariableNode>(InNode))
+				{
+					if(URigVMFunctionLibrary* OuterLibrary = Graph->GetTypedOuter<URigVMFunctionLibrary>())
 					{
-						TArray<FRigVMExternalVariable> ExternalVariablesWithoutVariableNode;
-						{
-							URigVMGraph* EditedGraph = InNode->GetGraph();
-							TGuardValue<TArray<URigVMNode*>> TemporaryRemoveNodes(EditedGraph->Nodes, TArray<URigVMNode*>());
-							ExternalVariablesWithoutVariableNode = EditedGraph->GetExternalVariables();
-						}
-
-						bool bFoundExternalVariable = false;
-						for(const FRigVMExternalVariable& ExternalVariable : ExternalVariablesWithoutVariableNode)
+						if(URigVMLibraryNode* OuterFunction = OuterLibrary->FindFunctionForNode(Graph->GetTypedOuter<URigVMCollapseNode>()))
 						{
-							if(ExternalVariable.Name == VariableToRemove)
+							const FName VariableToRemove = VariableNode->GetVariableName();
+							bool bIsLocalVariable = false;
+							for (FRigVMGraphVariableDescription VariableDescription : OuterFunction->GetContainedGraph()->LocalVariables)
 							{
-								bFoundExternalVariable = true;
-								break;
+								if (VariableDescription.Name == VariableToRemove)
+								{
+									bIsLocalVariable = true;
+									break;
+								}
 							}
-						}
 
-						if(!bFoundExternalVariable)
-						{
-							FRigVMControllerGraphGuard Guard(this, OuterFunction->GetContainedGraph(), false);
-							if(RequestBulkEditDialogDelegate.IsBound())
+							if (!bIsLocalVariable)
 							{
-								FRigVMController_BulkEditResult Result = RequestBulkEditDialogDelegate.Execute(OuterFunction, ERigVMControllerBulkEditType::RemoveVariable);
-								if(Result.bCanceled)
+								TArray<FRigVMExternalVariable> ExternalVariablesWithoutVariableNode;
+								{
+									URigVMGraph* EditedGraph = InNode->GetGraph();
+									TGuardValue<TArray<URigVMNode*>> TemporaryRemoveNodes(EditedGraph->Nodes, TArray<URigVMNode*>());
+									ExternalVariablesWithoutVariableNode = EditedGraph->GetExternalVariables();
+								}
+
+								bool bFoundExternalVariable = false;
+								for(const FRigVMExternalVariable& ExternalVariable : ExternalVariablesWithoutVariableNode)
+								{
+									if(ExternalVariable.Name == VariableToRemove)
+									{
+										bFoundExternalVariable = true;
+										break;
+									}
+								}
+
+								if(!bFoundExternalVariable)
 								{
-									return false;
+									FRigVMControllerGraphGuard Guard(this, OuterFunction->GetContainedGraph(), false);
+									if(RequestBulkEditDialogDelegate.IsBound())
+									{
+										FRigVMController_BulkEditResult Result = RequestBulkEditDialogDelegate.Execute(OuterFunction, ERigVMControllerBulkEditType::RemoveVariable);
+										if(Result.bCanceled)
+										{
+											return false;
+										}
+										bSetupUndoRedo = Result.bSetupUndoRedo;
+										bBulkEditDialogShown = true;
+									}
 								}
-								bSetupUndoRedo = Result.bSetupUndoRedo;
 							}
 						}
 					}
@@ -6317,7 +6247,7 @@ bool URigVMController::RemoveNode(URigVMNode* InNode, bool bSetupUndoRedo, bool
 			}
 		}
 	}
-
+	
 	TGuardValue<bool> GuardCompactness(bIgnoreRerouteCompactnessChanges, true);
 
 	FRigVMControllerCompileBracketScope CompileScope(this);
@@ -6325,197 +6255,198 @@ bool URigVMController::RemoveNode(URigVMNode* InNode, bool bSetupUndoRedo, bool
 	if (bSetupUndoRedo)
 	{
 		Action = FRigVMBaseAction();
-		Action.Title = FString::Printf(TEXT("Remove %s Node"), *InNode->GetNodeTitle());
+		if(InNodes.Num() == 1)
+		{
+			Action.Title = FString::Printf(TEXT("Remove %s Node"), *InNodes[0]->GetNodeTitle());
+		}
+		else
+		{
+			static const FString RemoveNodesTitle = TEXT("Remove nodes");
+			Action.Title = RemoveNodesTitle;
+		}
 		ActionStack->BeginAction(Action);
+		ActionStack->AddAction(FRigVMRemoveNodesAction(InNodes, this));
 	}
 
-	if (URigVMInjectionInfo* InjectionInfo = InNode->GetInjectionInfo())
+	for(URigVMNode* InNode : InNodes)
 	{
-		URigVMPin* Pin = InjectionInfo->GetPin();
-		check(Pin);
-
-		if (!EjectNodeFromPin(Pin->GetPinPath(), bSetupUndoRedo))
-		{
-			ActionStack->CancelAction(Action, this);
-			return false;
-		}
-		
-		if (InjectionInfo->bInjectedAsInput)
+		if (URigVMInjectionInfo* InjectionInfo = InNode->GetInjectionInfo())
 		{
-			if (InjectionInfo->InputPin)
+			URigVMPin* Pin = InjectionInfo->GetPin();
+			check(Pin);
+
+			if (!EjectNodeFromPin(Pin->GetPinPath(), bSetupUndoRedo))
+			{
+				ActionStack->CancelAction(Action, this);
+				return false;
+			}
+			
+			if (InjectionInfo->bInjectedAsInput)
 			{
-				URigVMPin* LastInputPin = Pin->GetPinForLink();
-				RewireLinks(InjectionInfo->InputPin, LastInputPin, true, bSetupUndoRedo);
+				if (InjectionInfo->InputPin)
+				{
+					URigVMPin* LastInputPin = Pin->GetPinForLink();
+					RewireLinks(InjectionInfo->InputPin, LastInputPin, true, bSetupUndoRedo);
+				}
 			}
-		}
-		else
-		{
-			if (InjectionInfo->OutputPin)
+			else
 			{
-				URigVMPin* LastOutputPin = Pin->GetPinForLink();
-				RewireLinks(InjectionInfo->OutputPin, LastOutputPin, false, bSetupUndoRedo);
+				if (InjectionInfo->OutputPin)
+				{
+					URigVMPin* LastOutputPin = Pin->GetPinForLink();
+					RewireLinks(InjectionInfo->OutputPin, LastOutputPin, false, bSetupUndoRedo);
+				}
 			}
 		}
-	}
 
-	
-	if (URigVMLibraryNode* LibraryNode = Cast<URigVMLibraryNode>(InNode))
-	{
-		// If we are removing a reference, remove the function references to this node in the function library
-		if(URigVMFunctionReferenceNode* FunctionReferenceNode = Cast<URigVMFunctionReferenceNode>(LibraryNode))
+		
+		if (URigVMLibraryNode* LibraryNode = Cast<URigVMLibraryNode>(InNode))
 		{
-			if(URigVMBuildData* BuildData = URigVMBuildData::Get())
+			// If we are removing a reference, remove the function references to this node in the function library
+			if(URigVMFunctionReferenceNode* FunctionReferenceNode = Cast<URigVMFunctionReferenceNode>(LibraryNode))
 			{
-				BuildData->UnregisterFunctionReference(FunctionReferenceNode->GetReferencedFunctionHeader().LibraryPointer, FunctionReferenceNode);
+				if(URigVMBuildData* BuildData = URigVMBuildData::Get())
+				{
+					BuildData->UnregisterFunctionReference(FunctionReferenceNode->GetReferencedFunctionHeader().LibraryPointer, FunctionReferenceNode);
+				}
 			}
-		}
-		// If we are removing a function, remove all the references first
-		else if (URigVMFunctionLibrary* FunctionLibrary = Cast<URigVMFunctionLibrary>(LibraryNode->GetGraph()))
-		{
-			if(URigVMBuildData* BuildData = URigVMBuildData::Get())
+			// If we are removing a function, remove all the references first
+			else if (URigVMFunctionLibrary* FunctionLibrary = Cast<URigVMFunctionLibrary>(LibraryNode->GetGraph()))
 			{
-				FRigVMGraphFunctionIdentifier Identifier = LibraryNode->GetFunctionIdentifier();
-				if (const FRigVMFunctionReferenceArray* ReferencesEntry = BuildData->FindFunctionReferences(Identifier))
+				if(URigVMBuildData* BuildData = URigVMBuildData::Get())
 				{
-					// make a copy since we'll be modifying the array
-					TArray< TSoftObjectPtr<URigVMFunctionReferenceNode> > FunctionReferences = ReferencesEntry->FunctionReferences;
-					for (const TSoftObjectPtr<URigVMFunctionReferenceNode>& FunctionReferencePtr : FunctionReferences)
+					FRigVMGraphFunctionIdentifier Identifier = LibraryNode->GetFunctionIdentifier();
+					if (const FRigVMFunctionReferenceArray* ReferencesEntry = BuildData->FindFunctionReferences(Identifier))
 					{
-						if (!ReferencesEntry->FunctionReferences.Contains(FunctionReferencePtr))
-						{
-							continue;
-						}
-						
-						if (FunctionReferencePtr.IsValid())
+						// make a copy since we'll be modifying the array
+						TArray< TSoftObjectPtr<URigVMFunctionReferenceNode> > FunctionReferences = ReferencesEntry->FunctionReferences;
+						for (const TSoftObjectPtr<URigVMFunctionReferenceNode>& FunctionReferencePtr : FunctionReferences)
 						{
-							FRigVMControllerGraphGuard GraphGuard(this, FunctionReferencePtr->GetGraph(), bSetupUndoRedo);
-							RemoveNode(FunctionReferencePtr.Get());
+							if (!ReferencesEntry->FunctionReferences.Contains(FunctionReferencePtr))
+							{
+								continue;
+							}
+							
+							if (FunctionReferencePtr.IsValid())
+							{
+								FRigVMControllerGraphGuard GraphGuard(this, FunctionReferencePtr->GetGraph(), bSetupUndoRedo);
+								RemoveNode(FunctionReferencePtr.Get(), bSetupUndoRedo, false);
+							}
 						}
 					}
-				}
 
-				BuildData->GraphFunctionReferences.Remove(Identifier);
-			}
-			
-			for(const auto& Pair : FunctionLibrary->LocalizedFunctions)
-			{
-				if(Pair.Value == LibraryNode)
+					BuildData->GraphFunctionReferences.Remove(Identifier);
+				}
+				
+				for(const auto& Pair : FunctionLibrary->LocalizedFunctions)
 				{
-					FunctionLibrary->LocalizedFunctions.Remove(Pair.Key);
-					break;
+					if(Pair.Value == LibraryNode)
+					{
+						FunctionLibrary->LocalizedFunctions.Remove(Pair.Key);
+						break;
+					}
 				}
-			}
-
-			if (FunctionLibrary->PublicFunctionNames.Contains(LibraryNode->GetFName()))
-			{
-				FunctionLibrary->PublicFunctionNames.Remove(LibraryNode->GetFName());
 
-				if (bSetupUndoRedo)
+				if (FunctionLibrary->PublicFunctionNames.Contains(LibraryNode->GetFName()))
 				{
-					ActionStack->AddAction(FRigVMMarkFunctionPublicAction(LibraryNode->GetFName(), true));
+					FunctionLibrary->PublicFunctionNames.Remove(LibraryNode->GetFName());
+
+					if (bSetupUndoRedo)
+					{
+						ActionStack->AddAction(FRigVMMarkFunctionPublicAction(LibraryNode->GetFName(), true));
+					}
 				}
 			}
 		}
-	}
 
-	// try to reconnect source and target nodes based on the current links
-	if (bRelinkPins)
-	{
-		RelinkSourceAndTargetPins(InNode, bSetupUndoRedo);
-	}
-	
-	if (bSetupUndoRedo || bRecursive)
-	{
-		SelectNode(InNode, false, bSetupUndoRedo);
+		SelectNode(InNode, false, false);
 
 		for (URigVMPin* Pin : InNode->GetPins())
 		{
 			// Remove injected nodes
 			while (Pin->HasInjectedNodes())
 			{
-				RemoveInjectedNode(Pin->GetPinPath(), Pin->GetDirection() != ERigVMPinDirection::Output, bSetupUndoRedo);
+				RemoveInjectedNode(Pin->GetPinPath(), Pin->GetDirection() != ERigVMPinDirection::Output, false);
 			}
 		
 			// breaking links also removes injected nodes 
-			BreakAllLinks(Pin, true, bSetupUndoRedo);
-			BreakAllLinks(Pin, false, bSetupUndoRedo);
-			BreakAllLinksRecursive(Pin, true, false, bSetupUndoRedo);
-			BreakAllLinksRecursive(Pin, false, false, bSetupUndoRedo);
+			BreakAllLinks(Pin, true, false);
+			BreakAllLinks(Pin, false, false);
+			BreakAllLinksRecursive(Pin, true, false, false);
+			BreakAllLinksRecursive(Pin, false, false, false);
 		}
 
 		if (URigVMCollapseNode* CollapseNode = Cast<URigVMCollapseNode>(InNode))
 		{
 			URigVMGraph* SubGraph = CollapseNode->GetContainedGraph();
-			FRigVMControllerGraphGuard GraphGuard(this, SubGraph, bSetupUndoRedo);
+			FRigVMControllerGraphGuard GraphGuard(this, SubGraph, false);
 		
 			TArray<URigVMNode*> ContainedNodes = SubGraph->GetNodes();
-			for (URigVMNode* ContainedNode : ContainedNodes)
-			{
-				if(Cast<URigVMFunctionEntryNode>(ContainedNode) != nullptr ||
-					Cast<URigVMFunctionReturnNode>(ContainedNode) != nullptr)
-				{
-					continue;
-				}
-				RemoveNode(ContainedNode, bSetupUndoRedo, bRecursive);
-			}
+			TGuardValue<bool> SuspendTemplateComputation(bSuspendTemplateComputation, true);
+			RemoveNodes(ContainedNodes, false, false);
 		}
 		
-		if (bSetupUndoRedo)
-		{
-			ActionStack->AddAction(FRigVMRemoveNodeAction(InNode, this));
-		}
+		Graph->Nodes.Remove(InNode);
+		Notify(ERigVMGraphNotifType::NodeRemoved, InNode);
 	}
 
-	Graph->Nodes.Remove(InNode);
 	if (!bSuspendNotifications)
 	{
 		Graph->MarkPackageDirty();
 	}
 
-	Notify(ERigVMGraphNotifType::NodeRemoved, InNode);
-
-	
-
 	if (bPrintPythonCommand)
 	{
 		const FString GraphName = GetSanitizedGraphName(GetGraph()->GetGraphName());
 
 		if (Graph->IsA<URigVMFunctionLibrary>())
 		{
-			const FString NodeName = GetSanitizedNodeName(InNode->GetName());
-			
-			RigVMPythonUtils::Print(GetGraphOuterName(), 
-								FString::Printf(TEXT("library_controller.remove_function_from_library('%s')"),
-												*NodeName));
+			for(const URigVMNode* InNode : InNodes)
+			{
+				const FString NodeName = GetSanitizedNodeName(InNode->GetName());
+				
+				RigVMPythonUtils::Print(GetGraphOuterName(), 
+									FString::Printf(TEXT("library_controller.remove_function_from_library('%s')"),
+													*NodeName));
+			}
 		}
 		else
 		{
-			const FString NodePath = GetSanitizedPinPath(InNode->GetNodePath());
+			TArray<FString> NodePaths;
+			Algo::Transform(InNodes, NodePaths, [](const URigVMNode* Node)
+			{
+				static constexpr TCHAR Format[] = TEXT("'%s'");
+				return FString::Printf(Format, *GetSanitizedPinPath(Node->GetNodePath()));
+			});
 
+			const FString NodePathsJoined = FString::Join(NodePaths, TEXT(", "));
+			
 			FString PythonCmd = FString::Printf(TEXT("blueprint.get_controller_by_name('%s')."), *GraphName );
-			PythonCmd += bRelinkPins ? FString::Printf(TEXT("remove_node_by_name('%s', relink_pins=True)"), *NodePath) :
-									   FString::Printf(TEXT("remove_node_by_name('%s')"), *NodePath);
+			PythonCmd += FString::Printf(TEXT("remove_node_by_name([%s])"), *NodePathsJoined);
 			
 			RigVMPythonUtils::Print(GetGraphOuterName(), PythonCmd );
 		}
 	}
 
-	if (URigVMVariableNode* VariableNode = Cast<URigVMVariableNode>(InNode))
-	{
-		Notify(ERigVMGraphNotifType::VariableRemoved, VariableNode);
-	}
-	
-	if (URigVMInjectionInfo* InjectionInfo = InNode->GetInjectionInfo())
+	for(URigVMNode* InNode : InNodes)
 	{
-		DestroyObject(InjectionInfo);
-	}
+		if (URigVMVariableNode* VariableNode = Cast<URigVMVariableNode>(InNode))
+		{
+			Notify(ERigVMGraphNotifType::VariableRemoved, VariableNode);
+		}
+		
+		if (URigVMInjectionInfo* InjectionInfo = InNode->GetInjectionInfo())
+		{
+			DestroyObject(InjectionInfo);
+		}
 
-	if (URigVMCollapseNode* CollapseNode = Cast<URigVMCollapseNode>(InNode))
-	{
-		DestroyObject(CollapseNode->GetContainedGraph());
-	}
+		if (URigVMCollapseNode* CollapseNode = Cast<URigVMCollapseNode>(InNode))
+		{
+			DestroyObject(CollapseNode->GetContainedGraph());
+		}
 
-	DestroyObject(InNode);
+		DestroyObject(InNode);
+	}
 
 	if (bSetupUndoRedo)
 	{
@@ -6525,7 +6456,7 @@ bool URigVMController::RemoveNode(URigVMNode* InNode, bool bSetupUndoRedo, bool
 	return true;
 }
 
-bool URigVMController::RemoveNodeByName(const FName& InNodeName, bool bSetupUndoRedo, bool bRecursive, bool bPrintPythonCommand, bool bRelinkPins)
+bool URigVMController::RemoveNodesByName(const TArray<FName>& InNodeNames, bool bSetupUndoRedo, bool bPrintPythonCommand)
 {
 	if (!IsValidGraph())
 	{
@@ -6540,7 +6471,33 @@ bool URigVMController::RemoveNodeByName(const FName& InNodeName, bool bSetupUndo
 	URigVMGraph* Graph = GetGraph();
 	check(Graph);
 
-	return RemoveNode(Graph->FindNodeByName(InNodeName), bSetupUndoRedo, bRecursive, bPrintPythonCommand, bRelinkPins);
+	TArray<URigVMNode*> Nodes;
+	Algo::Transform(InNodeNames, Nodes, [this, Graph](const FName& NodeName) -> URigVMNode*
+	{
+		URigVMNode* Node = Graph->FindNodeByName(NodeName);
+		if(Node == nullptr)
+		{
+			ReportErrorf(TEXT("Cannot find node '%s'."), *NodeName.ToString());
+		}
+		return Node;
+	});
+
+	if(Nodes.Contains(nullptr))
+	{
+		return false;
+	}
+	
+	return RemoveNodes(Nodes, bSetupUndoRedo, bPrintPythonCommand);
+}
+
+bool URigVMController::RemoveNode(URigVMNode* InNode, bool bSetupUndoRedo, bool bPrintPythonCommand)
+{
+	return RemoveNodes({InNode}, bSetupUndoRedo, bPrintPythonCommand);
+}
+
+bool URigVMController::RemoveNodeByName(const FName& InNodeName, bool bSetupUndoRedo, bool bPrintPythonCommand)
+{
+	return RemoveNodesByName({InNodeName}, bSetupUndoRedo, bPrintPythonCommand);
 }
 
 bool URigVMController::RenameNode(URigVMNode* InNode, const FName& InNewName, bool bSetupUndoRedo, bool bPrintPythonCommand)
@@ -6570,12 +6527,8 @@ bool URigVMController::RenameNode(URigVMNode* InNode, const FName& InNewName, bo
 	}
 
 	// loop over all links and remove them
-	TArray<URigVMLink*> Links = InNode->GetLinks();
-	for (URigVMLink* Link : Links)
-	{
-		Link->PrepareForCopy();
-		Notify(ERigVMGraphNotifType::LinkRemoved, Link);
-	}
+	const TArray<FLinkedPath> LinkedPaths = GetLinkedPaths(InNode);
+	FastBreakLinkedPaths(LinkedPaths);
 
 	const FSoftObjectPath PreviousObjectPath(InNode);
 	InNode->PreviousName = InNode->GetFName();
@@ -6588,11 +6541,9 @@ bool URigVMController::RenameNode(URigVMNode* InNode, const FName& InNewName, bo
 	Notify(ERigVMGraphNotifType::NodeRenamed, InNode);
 
 	// update the links once more
-	for (URigVMLink* Link : Links)
-	{
-		Link->PrepareForCopy();
-		Notify(ERigVMGraphNotifType::LinkAdded, Link);
-	}
+	FRestoreLinkedPathSettings Settings;
+	Settings.NodeNameMap.Add(InNode->PreviousName.ToString(), InNode->GetName());
+	RestoreLinkedPaths(LinkedPaths, Settings);
 
 	if(URigVMLibraryNode* LibraryNode = Cast<URigVMLibraryNode>(InNode))
 	{
@@ -6713,13 +6664,7 @@ bool URigVMController::SetNodeSelection(const TArray<FName>& InNodeNames, bool b
 	URigVMGraph* Graph = GetGraph();
 	check(Graph);
 
-	FRigVMSetNodeSelectionAction Action;
-	if (bSetupUndoRedo)
-	{
-		Action = FRigVMSetNodeSelectionAction(Graph, InNodeNames);
-		ActionStack->BeginAction(Action);
-	}
-
+	FRigVMSetNodeSelectionAction Action(Graph, InNodeNames);
 	bool bSelectionChanged = false;
 
 	TArray<FName> PreviousSelection = Graph->GetSelectNodes();
@@ -6767,11 +6712,7 @@ bool URigVMController::SetNodeSelection(const TArray<FName>& InNodeNames, bool b
 					Action.Title = TEXT("Selected multiple nodes.");
 				}
 			}
-			ActionStack->EndAction(Action);
-		}
-		else
-		{
-			ActionStack->CancelAction(Action, this);
+			ActionStack->AddAction(Action);
 		}
 	}
 
@@ -8092,11 +8033,13 @@ FString URigVMController::AddAggregatePin(URigVMNode* InNode, const FString& InP
 	}
 
 	FRigVMControllerCompileBracketScope CompileScope(this);
-	FRigVMBaseAction Action;
-	if (bSetupUndoRedo)
-	{
-		Action.Title = FString::Printf(TEXT("Add Aggregate Pin"));
-		ActionStack->BeginAction(Action);
+	const TMap<FString, FPinState> PinStates = GetPinStates(InNode);
+
+	FRigVMReplaceNodesAction Action;
+	if(bSetupUndoRedo)
+	{
+		Action = FRigVMReplaceNodesAction({InNode}, this);
+		Action.Title = TEXT("Add aggregate pin");
 	}
 
 	if (AggregateNode == nullptr)
@@ -8125,29 +8068,17 @@ FString URigVMController::AddAggregatePin(URigVMNode* InNode, const FString& InP
 		}
 		else
 		{
-			if (bSetupUndoRedo)
-			{
-				ActionStack->CancelAction(Action, this);
-			}
 			return FString();
 		}
 
 		if (!Arg1 || !Arg2 || !ArgOpposite)
 		{
-			if (bSetupUndoRedo)
-			{
-				ActionStack->CancelAction(Action, this);
-			}
 			return FString();
 		}
 
 		if (Arg1->GetCPPType() != Arg2->GetCPPType() || Arg1->GetCPPTypeObject() != Arg2->GetCPPTypeObject() ||
 			Arg1->GetCPPType() != ArgOpposite->GetCPPType() || Arg1->GetCPPTypeObject() != ArgOpposite->GetCPPTypeObject())
 		{
-			if (bSetupUndoRedo)
-			{
-				ActionStack->CancelAction(Action, this);
-			}
 			return FString();
 		}
 
@@ -8155,13 +8086,9 @@ FString URigVMController::AddAggregatePin(URigVMNode* InNode, const FString& InP
 		const FString AggregateArg2 = Arg2->GetName();
 		const FString AggregateOppositeArg = ArgOpposite->GetName();
 
-		TArray<TPair<FString, FString>> LinkedPaths = GetLinkedPinPaths(InNode);
-		if(!BreakLinkedPaths(LinkedPaths, bSetupUndoRedo))
+		TArray<FLinkedPath> LinkedPaths = GetLinkedPaths(InNode);
+		if(!FastBreakLinkedPaths(LinkedPaths, false))
 		{
-			if(bSetupUndoRedo)
-			{
-				ActionStack->CancelAction(Action, this);
-			}
 			return FString();
 		}
 
@@ -8186,28 +8113,20 @@ FString URigVMController::AddAggregatePin(URigVMNode* InNode, const FString& InP
 
 			if (AnswerType != INDEX_NONE)
 			{
-				ResolveWildCardPin(Arg1, AnswerType, bSetupUndoRedo);
+				ResolveWildCardPin(Arg1, AnswerType, false);
 			}
 
 		}
 
 		if (Arg1->IsWildCard() || Arg2->IsWildCard() || ArgOpposite->IsWildCard())
 		{
-			if(bSetupUndoRedo)
-			{
-				ActionStack->CancelAction(Action, this);
-			}
 			return FString();
 		}
 
 		const FName PreviousNodeName = InNode->GetFName();
-		URigVMCollapseNode* CollapseNode = CollapseNodes({InNode}, InNode->GetName(), bSetupUndoRedo, true);
+		URigVMCollapseNode* CollapseNode = CollapseNodes({InNode}, InNode->GetName(), false, true);
 		if (!CollapseNode)
 		{
-			if (bSetupUndoRedo)
-			{
-				ActionStack->CancelAction(Action, this);
-			}
 			return FString();
 		}
 
@@ -8216,7 +8135,7 @@ FString URigVMController::AddAggregatePin(URigVMNode* InNode, const FString& InP
 		AggregateNode = Cast<URigVMAggregateNode>(CollapseNode);
 		if (AggregateNode)
 		{
-			FRigVMControllerGraphGuard GraphGuard(this, AggregateNode->GetContainedGraph(), bSetupUndoRedo);
+			FRigVMControllerGraphGuard GraphGuard(this, AggregateNode->GetContainedGraph(), false);
 			TGuardValue<bool> GuardEditGraph(GetGraph()->bEditable, true);
 
 			for(int32 Index = 0; Index < InNode->GetPins().Num(); Index++)
@@ -8226,11 +8145,11 @@ FString URigVMController::AddAggregatePin(URigVMNode* InNode, const FString& InP
 				
 				if (URigVMPin* AggregatePin = AggregateNode->FindPin(PinName.ToString()))
 				{
-					SetExposedPinIndex(PinName, Index, bSetupUndoRedo);
+					SetExposedPinIndex(PinName, Index, false);
 					continue;
 				}
 
-				const FName ExposedPinName = AddExposedPin(PinName, Pin->GetDirection(), Pin->GetCPPType(), *Pin->GetCPPTypeObject()->GetPathName(), Pin->GetDefaultValue());
+				const FName ExposedPinName = AddExposedPin(PinName, Pin->GetDirection(), Pin->GetCPPType(), *Pin->GetCPPTypeObject()->GetPathName(), Pin->GetDefaultValue(), false, false);
 
 				const FString PinNameStr = PinName.ToString();
 				const FString ExposedPinNameStr = ExposedPinName.ToString();
@@ -8252,51 +8171,41 @@ FString URigVMController::AddAggregatePin(URigVMNode* InNode, const FString& InP
 
 				if (Pin->GetDirection() == ERigVMPinDirection::Input)
 				{
-					AddLink(FString::Printf(TEXT("Entry.%s"), *ExposedPinNameStr), FString::Printf(TEXT("%s.%s"), *InNode->GetName(), *PinNameStr), bSetupUndoRedo);
+					AddLink(FString::Printf(TEXT("Entry.%s"), *ExposedPinNameStr), FString::Printf(TEXT("%s.%s"), *InNode->GetName(), *PinNameStr), false);
 				}
 				else
 				{
-					AddLink(FString::Printf(TEXT("%s.%s"), *InNode->GetName(), *PinNameStr), FString::Printf(TEXT("Return.%s"), *ExposedPinNameStr), bSetupUndoRedo);
+					AddLink(FString::Printf(TEXT("%s.%s"), *InNode->GetName(), *PinNameStr), FString::Printf(TEXT("Return.%s"), *ExposedPinNameStr), false);
 				}
 			}
 		}
 		else
 		{
-			if (bSetupUndoRedo)
-			{
-				ActionStack->CancelAction(Action, this);
-			}
 			return FString();
 		}
 
-		RestoreLinkedPaths(LinkedPaths, {{PreviousNodeName.ToString(), AggregateNode->GetName()}}, {}, bSetupUndoRedo);
+		FRestoreLinkedPathSettings Settings;
+		Settings.NodeNameMap = {{PreviousNodeName.ToString(), AggregateNode->GetName()}};
+		RestoreLinkedPaths(LinkedPaths, Settings, false);
 	}
 	
 	if (!AggregateNode)
 	{
-		if (bSetupUndoRedo)
-		{
-			ActionStack->CancelAction(Action, this);
-		}
 		return FString();
 	}
 
 	URigVMPin* NewPin = nullptr;	
 	{
-		FRigVMControllerGraphGuard GraphGuard(this, AggregateNode->GetContainedGraph(), bSetupUndoRedo);
+		FRigVMControllerGraphGuard GraphGuard(this, AggregateNode->GetContainedGraph(), false);
 		TGuardValue<bool> GuardEditGraph(GetGraph()->bEditable, true);
 
 		URigVMNode* InnerNode = (AggregateNode == nullptr) ? InNode : AggregateNode->GetFirstInnerNode();
 
 		const FString InnerNodeContent = ExportNodesToText({InnerNode->GetFName()});
-		const TArray<FName> NewNodeNames = ImportNodesFromText(InnerNodeContent);
+		const TArray<FName> NewNodeNames = ImportNodesFromText(InnerNodeContent, false);
 		
 		if(NewNodeNames.IsEmpty())
 		{
-			if (bSetupUndoRedo)
-			{
-				ActionStack->CancelAction(Action, this);
-			}
 			return FString();
 		}
 
@@ -8334,7 +8243,7 @@ FString URigVMController::AddAggregatePin(URigVMNode* InNode, const FString& InP
 		}
 		
 		const URigVMPin* Arg1 = AggregateNode->GetFirstAggregatePin();
-		FName NewExposedPinName = AddExposedPin(NewPinName, Arg1->GetDirection(), Arg1->GetCPPType(), *Arg1->GetCPPTypeObject()->GetPathName(), InDefaultValue, bSetupUndoRedo);
+		FName NewExposedPinName = AddExposedPin(NewPinName, Arg1->GetDirection(), Arg1->GetCPPType(), *Arg1->GetCPPTypeObject()->GetPathName(), InDefaultValue, false);
 		NewPin = AggregateNode->FindPin(NewExposedPinName.ToString());
 		URigVMPin* NewUnitPinArg1 = NewNode->GetFirstAggregatePin();
 		URigVMPin* NewUnitPinArg2 = NewNode->GetSecondAggregatePin();
@@ -8348,10 +8257,10 @@ FString URigVMController::AddAggregatePin(URigVMNode* InNode, const FString& InP
 			URigVMPin* PreviousReturnPin = ReturnPin->GetLinkedSourcePins()[0];
 			PreviousNode = PreviousReturnPin->GetNode();
 		
-			BreakAllLinks(ReturnPin, true, bSetupUndoRedo);
-			AddLink(PreviousReturnPin, NewUnitPinArg1, bSetupUndoRedo);						
-			AddLink(EntryPin, NewUnitPinArg2, bSetupUndoRedo);
-			AddLink(NewUnitPinOppositeArg, ReturnPin, bSetupUndoRedo);
+			BreakAllLinks(ReturnPin, true, false);
+			AddLink(PreviousReturnPin, NewUnitPinArg1, false);						
+			AddLink(EntryPin, NewUnitPinArg2, false);
+			AddLink(NewUnitPinOppositeArg, ReturnPin, false);
 		}
 		else
 		{
@@ -8361,17 +8270,17 @@ FString URigVMController::AddAggregatePin(URigVMNode* InNode, const FString& InP
 			URigVMPin* PreviousReturnPin = OldReturnPin->GetLinkedSourcePins()[0];
 			PreviousNode = PreviousReturnPin->GetNode();
 
-			BreakAllLinks(OldReturnPin, true, bSetupUndoRedo);
-			AddLink(PreviousReturnPin, NewUnitPinOppositeArg, bSetupUndoRedo);						
-			AddLink(NewUnitPinArg1, OldReturnPin, bSetupUndoRedo);
-			AddLink(NewUnitPinArg2, NewReturnPin, bSetupUndoRedo);
+			BreakAllLinks(OldReturnPin, true, false);
+			AddLink(PreviousReturnPin, NewUnitPinOppositeArg, false);						
+			AddLink(NewUnitPinArg1, OldReturnPin, false);
+			AddLink(NewUnitPinArg2, NewReturnPin, false);
 		}
 
 		// Rearrange the graph nodes
 		URigVMFunctionReturnNode* ReturnNode = AggregateNode->GetReturnNode();
 		FVector2D NodeDimensions(200, 150);
-		SetNodePosition(NewNode, PreviousNode->GetPosition() + NodeDimensions, bSetupUndoRedo);
-		SetNodePosition(ReturnNode, NewNode->GetPosition() + NodeDimensions, bSetupUndoRedo);
+		SetNodePosition(NewNode, PreviousNode->GetPosition() + NodeDimensions, false);
+		SetNodePosition(ReturnNode, NewNode->GetPosition() + NodeDimensions, false);
 
 		// Connect other input pins
 		for (URigVMPin* OtherInputPin : AggregateNode->GetFirstInnerNode()->GetPins())
@@ -8381,7 +8290,7 @@ FString URigVMController::AddAggregatePin(URigVMNode* InNode, const FString& InP
 				OtherInputPin->GetName() != NewUnitPinOppositeArg->GetName())
 			{
 				URigVMPin* OtherEntryPin = AggregateNode->GetEntryNode()->FindPin(OtherInputPin->GetName());
-				AddLink(OtherEntryPin, NewNode->FindPin(OtherEntryPin->GetName()), bSetupUndoRedo);
+				AddLink(OtherEntryPin, NewNode->FindPin(OtherEntryPin->GetName()), false);
 			}
 		}
 
@@ -8390,16 +8299,19 @@ FString URigVMController::AddAggregatePin(URigVMNode* InNode, const FString& InP
 
 	if (!NewPin)
 	{
-		if (bSetupUndoRedo)
-		{
-			ActionStack->CancelAction(Action, this);
-		}
 		return FString();
 	}
 
+	if(!PinStates.IsEmpty())
+	{
+		FRigVMControllerGraphGuard GraphGuard(this, AggregateNode->GetGraph(), false);
+		ApplyPinStates(AggregateNode, PinStates);
+	}
+
 	if (bSetupUndoRedo)
 	{
-		ActionStack->EndAction(Action);
+		Action.StoreNode(AggregateNode, this, false);
+		ActionStack->AddAction(Action);
 	}
 
 	if (bPrintPythonCommand)
@@ -8465,16 +8377,18 @@ bool URigVMController::RemoveAggregatePin(URigVMPin* InPin, bool bSetupUndoRedo,
 		return false;
 	}
 
-	FRigVMControllerCompileBracketScope CompileScope(this);
-	FRigVMBaseAction Action;
-	if (bSetupUndoRedo)
+	FRigVMReplaceNodesAction Action;
+	if(bSetupUndoRedo)
 	{
-		Action.Title = FString::Printf(TEXT("Remove Aggregate Pin"));
-		ActionStack->BeginAction(Action);
+		Action = FRigVMReplaceNodesAction({InPin->GetNode()}, this);
+		Action.Title = TEXT("Remove aggregate pin");
 	}
 
+	URigVMNode* Node = InPin->GetNode();
+	FRigVMControllerCompileBracketScope CompileScope(this);
+
 	bool bSuccess = false;
-	if (URigVMAggregateNode* AggregateNode = Cast<URigVMAggregateNode>(InPin->GetNode()))
+	if (URigVMAggregateNode* AggregateNode = Cast<URigVMAggregateNode>(Node))
 	{
 		URigVMGraph* Graph = AggregateNode->GetContainedGraph();
 		if (AggregateNode->IsInputAggregate())
@@ -8485,7 +8399,7 @@ bool URigVMController::RemoveAggregatePin(URigVMPin* InPin, bool bSetupUndoRedo,
 				{
 					if (EntryPin->GetLinkedTargetPins().Num() > 0)
 					{
-						FRigVMControllerGraphGuard GraphGuard(this, AggregateNode->GetContainedGraph(), bSetupUndoRedo);
+						FRigVMControllerGraphGuard GraphGuard(this, AggregateNode->GetContainedGraph(), false);
 						TGuardValue<bool> GuardEditGraph(GetGraph()->bEditable, true);
 						
 						URigVMPin* TargetPin = EntryPin->GetLinkedTargetPins()[0];
@@ -8502,10 +8416,10 @@ bool URigVMController::RemoveAggregatePin(URigVMPin* InPin, bool bSetupUndoRedo,
 						const FString FirstAggregatePin = AggregateNode->GetFirstAggregatePin()->GetName();
 						const FString SecondAggregatePin = AggregateNode->GetSecondAggregatePin()->GetName();
 						FString OtherArg = TargetPin->GetName() == FirstAggregatePin ? SecondAggregatePin : FirstAggregatePin;
-						BreakAllLinks(NextNodePin, true, bSetupUndoRedo);
-						RewireLinks(NodeToRemove->FindPin(OtherArg), NextNodePin, true, bSetupUndoRedo);
-						RemoveNode(NodeToRemove, bSetupUndoRedo);
-						RemoveExposedPin(*InPin->GetName(), bSetupUndoRedo);
+						BreakAllLinks(NextNodePin, true, false);
+						RewireLinks(NodeToRemove->FindPin(OtherArg), NextNodePin, true, false);
+						RemoveNode(NodeToRemove, false);
+						RemoveExposedPin(*InPin->GetName(), false);
 						bSuccess = true;
 					}
 				}
@@ -8519,7 +8433,7 @@ bool URigVMController::RemoveAggregatePin(URigVMPin* InPin, bool bSetupUndoRedo,
 				{
 					if (ReturnPin->GetLinkedSourcePins().Num() > 0)
 					{
-						FRigVMControllerGraphGuard GraphGuard(this, AggregateNode->GetContainedGraph(), bSetupUndoRedo);
+						FRigVMControllerGraphGuard GraphGuard(this, AggregateNode->GetContainedGraph(), false);
 						TGuardValue<bool> GuardEditGraph(GetGraph()->bEditable, true);
 						
 						URigVMPin* SourcePin = ReturnPin->GetLinkedSourcePins()[0];
@@ -8537,10 +8451,10 @@ bool URigVMController::RemoveAggregatePin(URigVMPin* InPin, bool bSetupUndoRedo,
 						const FString FirstAggregatePin = AggregateNode->GetFirstAggregatePin()->GetName();
 						const FString SecondAggregatePin = AggregateNode->GetSecondAggregatePin()->GetName();
 						FString OtherArg = SourcePin->GetName() == FirstAggregatePin ? SecondAggregatePin : FirstAggregatePin;
-						BreakAllLinks(NextNodePin, false, bSetupUndoRedo);
-						RewireLinks(NodeToRemove->FindPin(OtherArg), NextNodePin, false, bSetupUndoRedo);
-						RemoveNode(NodeToRemove, bSetupUndoRedo);
-						RemoveExposedPin(*InPin->GetName(), bSetupUndoRedo);
+						BreakAllLinks(NextNodePin, false, false);
+						RewireLinks(NodeToRemove->FindPin(OtherArg), NextNodePin, false, false);
+						RemoveNode(NodeToRemove, false);
+						RemoveExposedPin(*InPin->GetName(), false);
 						bSuccess = true;
 					}
 				}
@@ -8549,13 +8463,9 @@ bool URigVMController::RemoveAggregatePin(URigVMPin* InPin, bool bSetupUndoRedo,
 
 		if (bSuccess && AggregateNode->GetContainedNodes().Num() == 3)
 		{
-			TArray<TPair<FString, FString>> LinkedPaths = GetLinkedPinPaths(AggregateNode);
-			if(!BreakLinkedPaths(LinkedPaths, bSetupUndoRedo))
+			TArray<FLinkedPath> LinkedPaths = GetLinkedPaths(AggregateNode);
+			if(!FastBreakLinkedPaths(LinkedPaths, false))
 			{
-				if(bSetupUndoRedo)
-				{
-					ActionStack->CancelAction(Action, this);
-				}
 				return false;
 			}
 
@@ -8580,16 +8490,19 @@ bool URigVMController::RemoveAggregatePin(URigVMPin* InPin, bool bSetupUndoRedo,
 				}
 			}
 
-			const FString PreviousNodeName = AggregateNode->GetName();
-			TArray<URigVMNode*> NodesEjected = ExpandLibraryNode(AggregateNode, bSetupUndoRedo);
+			const FName PreviousNodeName = AggregateNode->GetFName();
+			TArray<URigVMNode*> NodesEjected = ExpandLibraryNode(AggregateNode, false);
 			bSuccess = NodesEjected.Num() == 1;
 
 			if(bSuccess)
 			{
 				URigVMNode* EjectedNode = NodesEjected[0];
-				RestoreLinkedPaths(LinkedPaths, {}, {{
-					PreviousNodeName, FRigVMController_PinPathRemapDelegate::CreateLambda([
-						PreviousNodeName,
+				RenameNode(EjectedNode, PreviousNodeName, false, false);
+				Node = EjectedNode;
+
+				FRestoreLinkedPathSettings Settings;
+				Settings.RemapDelegates = {{
+					PreviousNodeName.ToString(), FRigVMController_PinPathRemapDelegate::CreateLambda([
 						EjectedNode,
 						PinNameMap
 					](const FString& InPinPath, bool bIsInput) -> FString
@@ -8604,10 +8517,14 @@ bool URigVMController::RemoveAggregatePin(URigVMPin* InPin, bool bSetupUndoRedo,
 						{
 							Segments[1] = *RemappedPin;
 						}
-						
 						return URigVMPin::JoinPinPath(Segments);
 					})
-				}}, bSetupUndoRedo);
+				}};
+				RestoreLinkedPaths(LinkedPaths, Settings, false);
+			}
+			else
+			{
+				Node = nullptr;
 			}
 		}		
 	}
@@ -8616,11 +8533,8 @@ bool URigVMController::RemoveAggregatePin(URigVMPin* InPin, bool bSetupUndoRedo,
 	{
 		if (bSuccess)
 		{
-			ActionStack->EndAction(Action);
-		}
-		else
-		{
-			ActionStack->CancelAction(Action, this);			
+			Action.StoreNode(Node, this, false);
+			ActionStack->AddAction(Action);
 		}
 	}
 
@@ -8904,7 +8818,7 @@ bool URigVMController::RemoveArrayPin(const FString& InArrayElementPinPath, bool
 	return true;
 }
 
-bool URigVMController::RemovePin(URigVMPin* InPinToRemove, bool bSetupUndoRedo)
+bool URigVMController::RemovePin(URigVMPin* InPinToRemove, bool bSetupUndoRedo, bool bForceBreakLinks)
 {
 	if (!bIsTransacting && !IsGraphEditable())
 	{
@@ -8912,7 +8826,7 @@ bool URigVMController::RemovePin(URigVMPin* InPinToRemove, bool bSetupUndoRedo)
 	}
 	
 	FRigVMControllerCompileBracketScope CompileScope(this);
-	if (bSetupUndoRedo)
+	if (bSetupUndoRedo || bForceBreakLinks)
 	{
 		BreakAllLinks(InPinToRemove, true, bSetupUndoRedo);
 		BreakAllLinks(InPinToRemove, false, bSetupUndoRedo);
@@ -8920,25 +8834,26 @@ bool URigVMController::RemovePin(URigVMPin* InPinToRemove, bool bSetupUndoRedo)
 		BreakAllLinksRecursive(InPinToRemove, false, false, bSetupUndoRedo);
 	}
 
-	TArray<URigVMPin*> SubPins = InPinToRemove->GetSubPins();
-	for (URigVMPin* SubPin : SubPins)
-	{
-		if (!RemovePin(SubPin, bSetupUndoRedo))
-		{
-			return false;
-		}
-	}
-
 	if (URigVMPin* ParentPin = InPinToRemove->GetParentPin())
 	{
 		ParentPin->SubPins.Remove(InPinToRemove);
 	}
+	
 	else if(URigVMNode* Node = InPinToRemove->GetNode())
 	{
 		Node->Pins.Remove(InPinToRemove);
 		Node->OrphanedPins.Remove(InPinToRemove);
 	}
 
+	TArray<URigVMPin*> SubPins = InPinToRemove->GetSubPins();
+	for (URigVMPin* SubPin : SubPins)
+	{
+		if (!RemovePin(SubPin, bSetupUndoRedo, bForceBreakLinks))
+		{
+			return false;
+		}
+	}
+
 	if (!bSuspendNotifications)
 	{
 		Notify(ERigVMGraphNotifType::PinRemoved, InPinToRemove);
@@ -9611,7 +9526,7 @@ bool URigVMController::AddLink(const FString& InOutputPinPath, const FString& In
 	return bSuccess;
 }
 
-bool URigVMController::AddLink(URigVMPin* OutputPin, URigVMPin* InputPin, bool bSetupUndoRedo, ERigVMPinDirection InUserDirection, bool bCreateCastNode)
+bool URigVMController::AddLink(URigVMPin* OutputPin, URigVMPin* InputPin, bool bSetupUndoRedo, ERigVMPinDirection InUserDirection, bool bCreateCastNode, bool bIsRestoringLinks, FString* OutFailureReason)
 {
 	if (!bIsTransacting && !IsGraphEditable())
 	{
@@ -9640,7 +9555,12 @@ bool URigVMController::AddLink(URigVMPin* OutputPin, URigVMPin* InputPin, bool b
 
 	if (Graph->IsA<URigVMFunctionLibrary>())
 	{
-		ReportError(TEXT("Cannot add links in function library graphs."));
+		static const FString FailureReason = TEXT("Cannot add links in function library graphs."); 
+		if(OutFailureReason)
+		{
+			*OutFailureReason = FailureReason;
+		}
+		ReportError(FailureReason);
 		return false;
 	}
 
@@ -9650,8 +9570,24 @@ bool URigVMController::AddLink(URigVMPin* OutputPin, URigVMPin* InputPin, bool b
 	if (!bIsTransacting)
 	{
 		FString FailureReason;
-		if (!Graph->CanLink(OutputPin, InputPin, &FailureReason, GetCurrentByteCode(), UserLinkDirection, bCreateCastNode))
+
+		bool bCanLink = true;
+		if(bIsRestoringLinks)
+		{
+			bCanLink = URigVMPin::CanLink(OutputPin, InputPin, &FailureReason, nullptr, UserLinkDirection, bCreateCastNode);
+		}
+		else
 		{
+			bCanLink = Graph->CanLink(OutputPin, InputPin, &FailureReason, GetCurrentByteCode(), UserLinkDirection, bCreateCastNode);
+		}
+		
+		if (!bCanLink)
+		{
+			if(OutFailureReason)
+			{
+				*OutFailureReason = FailureReason;
+			}
+			
 			if(OutputPin->IsExecuteContext() && InputPin->IsExecuteContext())
 			{
 				if(OutputPin->GetNode()->IsA<URigVMFunctionEntryNode>() &&
@@ -9811,16 +9747,30 @@ bool URigVMController::AddLink(URigVMPin* OutputPin, URigVMPin* InputPin, bool b
 		
 		if(bOutputPinCanChangeType)
 		{
-			Notify(ERigVMGraphNotifType::InteractionBracketOpened, nullptr);
+			bool bInteractionBracketOpened = false;
 			if(OutputPin->GetNode()->IsA<URigVMRerouteNode>())
 			{
+				if(!bInteractionBracketOpened)
+				{
+					Notify(ERigVMGraphNotifType::InteractionBracketOpened, nullptr);
+					bInteractionBracketOpened = true;
+				}
 				SetPinDefaultValue(OutputPin, InputPin->GetDefaultValue(), true, bSetupUndoRedo, false);
 			}
 			if(InputPin->GetNode()->IsA<URigVMRerouteNode>())
 			{
+				if(!bInteractionBracketOpened)
+				{
+					Notify(ERigVMGraphNotifType::InteractionBracketOpened, nullptr);
+					bInteractionBracketOpened = true;
+				}
 				SetPinDefaultValue(OutputPin, OutputPin->GetDefaultValue(), true, bSetupUndoRedo, false);
 			}
-			Notify(ERigVMGraphNotifType::InteractionBracketClosed, nullptr);
+
+			if(bInteractionBracketOpened)
+			{
+				Notify(ERigVMGraphNotifType::InteractionBracketClosed, nullptr);
+			}
 		}
 	}
 
@@ -9852,11 +9802,29 @@ bool URigVMController::AddLink(URigVMPin* OutputPin, URigVMPin* InputPin, bool b
 		ActionStack->AddAction(FRigVMAddLinkAction(OutputPin, InputPin));
 	}
 
-	URigVMLink* Link = NewObject<URigVMLink>(Graph);
-	Link->SourcePin = OutputPin;
-	Link->TargetPin = InputPin;
-	Link->SourcePinPath = OutputPin->GetPinPath();
-	Link->TargetPinPath = InputPin->GetPinPath();
+	URigVMLink* Link = nullptr;
+	
+	// try to reuse UObjects that may still be sitting in the detached links array
+	const FString OutputPinPath = OutputPin->GetPinPath();
+	const FString InputPinPath = InputPin->GetPinPath();
+	TObjectPtr<URigVMLink>* ExistingLink = Graph->DetachedLinks.FindByPredicate(
+		[OutputPinPath, InputPinPath](const URigVMLink* Link) -> bool
+		{
+			return Link->GetSourcePinPath().Equals(OutputPinPath, ESearchCase::CaseSensitive) &&
+				Link->GetTargetPinPath().Equals(InputPinPath, ESearchCase::CaseSensitive);
+		});
+
+	if(ExistingLink)
+	{
+		Link = *ExistingLink;
+		Graph->DetachedLinks.Remove(Link);
+	}
+	else
+	{
+		Link = NewObject<URigVMLink>(Graph);
+	}
+	
+	Link->SetSourceAndTargetPinPaths(OutputPin->GetPinPath(), InputPin->GetPinPath());
 	Graph->Links.Add(Link);
 	OutputPin->Links.Add(Link);
 	InputPin->Links.Add(Link);
@@ -10074,7 +10042,7 @@ bool URigVMController::BreakLink(URigVMPin* OutputPin, URigVMPin* InputPin, bool
 
 	for (URigVMLink* Link : InputPin->Links)
 	{
-		if (Link->SourcePin == OutputPin && Link->TargetPin == InputPin)
+		if (Link->GetSourcePin() == OutputPin && Link->GetTargetPin() == InputPin)
 		{
 			FRigVMControllerCompileBracketScope CompileScope(this);
 			FRigVMBreakLinkAction Action;
@@ -10406,6 +10374,7 @@ FName URigVMController::AddExposedPin(const FName& InPinName, ERigVMPinDirection
 		Notify(ERigVMGraphNotifType::NodeAdded, ReturnNode);
 	}
 
+	
 	RefreshFunctionPins(EntryNode);
 	RefreshFunctionPins(ReturnNode);
 
@@ -10496,7 +10465,7 @@ bool URigVMController::RemoveExposedPin(const FName& InPinName, bool bSetupUndoR
 	{
 		{
 			FRigVMControllerGraphGuard GraphGuard(this, LibraryNode->GetGraph(), bSetupUndoRedo);
-			bSuccessfullyRemovedPin = RemovePin(Pin, bSetupUndoRedo);
+			bSuccessfullyRemovedPin = RemovePin(Pin, bSetupUndoRedo, true);
 		}
 
 		TArray<URigVMVariableNode*> NodesToRemove;
@@ -10625,17 +10594,14 @@ bool URigVMController::RenameExposedPin(const FName& InOldPinName, const FName&
 		{
 			FRigVMControllerGraphGuard GraphGuard(InController, InPin->GetGraph(), false);
 
-			TArray<URigVMLink*> Links;
-			Links.Append(InPin->GetSourceLinks(true));
-			Links.Append(InPin->GetTargetLinks(true));
+			TArray<FLinkedPath> LinkedPaths = GetLinkedPaths(InPin, true, true);
 
-			// store both the ptr + pin path
-			for (URigVMLink* Link : Links)
-			{
-				Link->PrepareForCopy();
-				InController->Notify(ERigVMGraphNotifType::LinkRemoved, Link);
-			}
+			InController->FastBreakLinkedPaths(LinkedPaths);
 
+			const FString OldPinPath = InPin->GetPinPath();
+			static constexpr TCHAR PinPrefixFormat[] = TEXT("%s.");
+			const FString OldPinPathPrefix = FString::Printf(PinPrefixFormat, *OldPinPath);
+			
 			if (!InController->RenameObject(InPin, *InNewName.ToString()))
 			{
 				return false;
@@ -10643,18 +10609,30 @@ bool URigVMController::RenameExposedPin(const FName& InOldPinName, const FName&
 			
 			InPin->SetDisplayName(InNewName);
 			
-			// update the eventually stored pin path to the new name
-			for (URigVMLink* Link : Links)
-			{
-				Link->PrepareForCopy();
-			}
-
 			InController->Notify(ERigVMGraphNotifType::PinRenamed, InPin);
 
-			for (URigVMLink* Link : Links)
-			{
-				InController->Notify(ERigVMGraphNotifType::LinkAdded, Link);
-			}
+			FRestoreLinkedPathSettings Settings;
+			Settings.RemapDelegates = {{
+				InPin->GetNode()->GetName(), FRigVMController_PinPathRemapDelegate::CreateLambda([
+					OldPinPath,
+					OldPinPathPrefix,
+					InPin
+				](const FString& InPinPath, bool bIsInput) -> FString
+				{
+					if(InPinPath == OldPinPath)
+					{
+						return InPin->GetPinPath();
+					}
+
+					if(InPinPath.StartsWith(OldPinPathPrefix, ESearchCase::CaseSensitive))
+					{
+						return InPin->GetPinPath() + InPinPath.Mid(OldPinPath.Len());
+					}
+
+					return InPinPath;
+				})
+			}};
+			InController->RestoreLinkedPaths(LinkedPaths, Settings);
 
 			return true;
 		}
@@ -10901,7 +10879,11 @@ bool URigVMController::ChangeExposedPinType(const FName& InPinName, const FStrin
 	{
 		if (InterfaceNode)
 		{
+			const TArray<FLinkedPath> LinkedPaths = GetLinkedPaths(InterfaceNode);
+			FastBreakLinkedPaths(LinkedPaths);
 			RepopulatePinsOnNode(InterfaceNode, true, bSetupOrphanPins, true);
+			FRestoreLinkedPathSettings Settings;
+			RestoreLinkedPaths(LinkedPaths, Settings);
 			RemoveUnusedOrphanedPins(InterfaceNode);
 		}
 	}
@@ -11116,12 +11098,9 @@ URigVMFunctionReferenceNode* URigVMController::AddFunctionReferenceNodeFromDescr
 
 	if (bSetupUndoRedo)
 	{
-		FRigVMInverseAction InverseAction;
-		InverseAction.Title = TEXT("Add function node");
-
-		ActionStack->BeginAction(InverseAction);
-		ActionStack->AddAction(FRigVMRemoveNodeAction(FunctionRefNode, this));
-		ActionStack->EndAction(InverseAction);
+		FRigVMImportFromTextAction Action(FunctionRefNode, this);
+		Action.Title = TEXT("Add function node");
+		ActionStack->AddAction(Action);
 	}
 
 	if (bPrintPythonCommand)
@@ -11389,12 +11368,9 @@ URigVMLibraryNode* URigVMController::AddFunctionToLibrary(const FName& InFunctio
 
 	if (bSetupUndoRedo)
 	{
-		FRigVMInverseAction InverseAction;
-		InverseAction.Title = TEXT("Add function to library");
-
-		ActionStack->BeginAction(InverseAction);
-		ActionStack->AddAction(FRigVMRemoveNodeAction(CollapseNode, this));
-		ActionStack->EndAction(InverseAction);
+		FRigVMImportFromTextAction Action(CollapseNode, this);
+		Action.Title = TEXT("Add function to library");
+		ActionStack->AddAction(Action);
 	}
 
 	if (bPrintPythonCommand)
@@ -11644,12 +11620,9 @@ FRigVMGraphVariableDescription URigVMController::AddLocalVariable(const FName& I
 
 	if (bSetupUndoRedo)
 	{
-		FRigVMInverseAction InverseAction;
-		InverseAction.Title = FString::Printf(TEXT("Add Local Variable %s"), *InVariableName.ToString());
-
-		ActionStack->BeginAction(InverseAction);
-		ActionStack->AddAction(FRigVMRemoveLocalVariableAction(NewVariable));
-		ActionStack->EndAction(InverseAction);
+		FRigVMAddLocalVariableAction Action(NewVariable);
+		Action.Title = FString::Printf(TEXT("Add Local Variable %s"), *InVariableName.ToString());
+		ActionStack->AddAction(Action);
 	}
 
 	if (!bSuspendNotifications)
@@ -12090,13 +12063,10 @@ bool URigVMController::SetLocalVariableDefaultValue(const FName& InVariableName,
 	FRigVMControllerCompileBracketScope CompileScope(this);
 	if (bSetupUndoRedo)
 	{
-		FRigVMInverseAction InverseAction;
-		InverseAction.Title = FString::Printf(TEXT("Change Local Variable %s default value"), *InVariableName.ToString());
-
-		ActionStack->BeginAction(InverseAction);
-		ActionStack->AddAction(FRigVMChangeLocalVariableDefaultValueAction(LocalVariables[FoundIndex], InDefaultValue));
-		ActionStack->EndAction(InverseAction);
-	}	
+		FRigVMChangeLocalVariableDefaultValueAction Action(LocalVariables[FoundIndex], InDefaultValue);
+		Action.Title = FString::Printf(TEXT("Change Local Variable %s default value"), *InVariableName.ToString());
+		ActionStack->AddAction(Action);
+	}
 
 	FRigVMGraphVariableDescription& VariableDescription = LocalVariables[FoundIndex];
 	VariableDescription.DefaultValue = InDefaultValue;
@@ -12502,7 +12472,7 @@ URigVMRerouteNode* URigVMController::AddFreeRerouteNode(bool bShowAsFullNode, co
 
 	if (bSetupUndoRedo)
 	{
-		ActionStack->AddAction(FRigVMAddRerouteNodeAction(Node));
+		ActionStack->AddAction(FRigVMImportFromTextAction(Node, this));
 	}
 
 	Notify(ERigVMGraphNotifType::NodeAdded, Node);
@@ -12741,12 +12711,12 @@ URigVMTemplateNode* URigVMController::AddTemplateNode(const FName& InNotation, c
 
 	Notify(ERigVMGraphNotifType::NodeAdded, Node);
 
-	FRigVMAddTemplateNodeAction Action;
+	FRigVMBaseAction Action;
 	FRigVMControllerCompileBracketScope CompileScope(this);
 	if (bSetupUndoRedo)
 	{
-		Action = FRigVMAddTemplateNodeAction(Node);
 		ActionStack->BeginAction(Action);
+		ActionStack->AddAction(FRigVMImportFromTextAction(Node, this));
 	}
 	
 	ResolveTemplateNodeMetaData(Node, bSetupUndoRedo);
@@ -12893,7 +12863,7 @@ URigVMEnumNode* URigVMController::AddEnumNode(const FName& InCPPTypeObjectPath,
 	FRigVMControllerCompileBracketScope CompileScope(this);
 	if (bSetupUndoRedo)
 	{
-		ActionStack->AddAction(FRigVMAddEnumNodeAction(Node));
+		ActionStack->AddAction(FRigVMImportFromTextAction(Node, this));
 	}
 
 	if (bPrintPythonCommand)
@@ -13102,7 +13072,7 @@ URigVMInvokeEntryNode* URigVMController::AddInvokeEntryNode(const FName& InEntry
 
 	if (bSetupUndoRedo)
 	{
-		FRigVMAddInvokeEntryNodeAction Action(Node);
+		FRigVMImportFromTextAction Action(Node, this);
 		Action.Title = FString::Printf(TEXT("Add Invoke %s Entry"), *InEntryName.ToString());
 		ActionStack->AddAction(Action);
 	}
@@ -13177,7 +13147,7 @@ bool URigVMController::IsGraphEditable() const
 	return Graph->bEditable;
 }
 
-bool URigVMController::IsValidNodeForGraph(URigVMNode* InNode)
+bool URigVMController::IsValidNodeForGraph(const URigVMNode* InNode)
 {
 	if(!IsValidGraph())
 	{
@@ -13201,15 +13171,16 @@ bool URigVMController::IsValidNodeForGraph(URigVMNode* InNode)
 		return false;
 	}
 
-	if (InNode->GetNodeIndex() == INDEX_NONE)
+	if (InNode->GetNodeIndex() == INDEX_NONE && !InNode->GetOuter()->IsA<URigVMInjectionInfo>())
 	{
 		ReportErrorf(TEXT("InNode '%s' is transient (not yet nested to a graph)."), *InNode->GetName());
+		return false;
 	}
 
 	return true;
 }
 
-bool URigVMController::IsValidPinForGraph(URigVMPin* InPin)
+bool URigVMController::IsValidPinForGraph(const URigVMPin* InPin)
 {
 	if(!IsValidGraph())
 	{
@@ -13240,7 +13211,7 @@ bool URigVMController::IsValidPinForGraph(URigVMPin* InPin)
 	return true;
 }
 
-bool URigVMController::IsValidLinkForGraph(URigVMLink* InLink)
+bool URigVMController::IsValidLinkForGraph(const URigVMLink* InLink)
 {
 	if(!IsValidGraph())
 	{
@@ -13366,17 +13337,23 @@ bool URigVMController::CanAddNode(URigVMNode* InNode, bool bReportErrors, bool b
 		// only allow one function entry node
 		if(InNode->IsA<URigVMFunctionEntryNode>())
 		{
-			if(Graph->GetEntryNode() != nullptr)
+			if(const URigVMFunctionEntryNode* ExistingEntryNode = Graph->GetEntryNode())
 			{
-				return false;
+				if(ExistingEntryNode != InNode)
+				{
+					return false;
+				}
 			}
 		}
 		// only allow one function return node
 		else if(InNode->IsA<URigVMFunctionReturnNode>())
 		{
-			if(Graph->GetReturnNode() != nullptr)
+			if(const URigVMFunctionReturnNode* ExistingReturnNode = Graph->GetReturnNode())
 			{
-				return false;
+				if(ExistingReturnNode != InNode)
+				{
+					return false;
+				}
 			}
 		}
 	}
@@ -14137,365 +14114,114 @@ FProperty* URigVMController::FindPropertyForPin(const FString& InPinPath)
 	return nullptr;
 }
 
-int32 URigVMController::DetachLinksFromPinObjects(const TArray<URigVMLink*>* InLinks)
+void URigVMController::RemoveStaleNodes()
 {
+	if (!IsValidGraph())
+	{
+		return;
+	}
+
 	URigVMGraph* Graph = GetGraph();
 	check(Graph);
 
-	TArray<URigVMLink*> Links;
-	if (InLinks)
+	Graph->Nodes.Remove(nullptr);
+}
+
+void URigVMController::AddPinRedirector(bool bInput, bool bOutput, const FString& OldPinPath, const FString& NewPinPath)
+{
+	if (OldPinPath.IsEmpty() || NewPinPath.IsEmpty() || OldPinPath == NewPinPath)
 	{
-		Links = *InLinks;
+		return;
 	}
-	else
+
+	if (bInput)
 	{
-		Links = Graph->Links;
+		InputPinRedirectors.FindOrAdd(OldPinPath) = NewPinPath;
 	}
-
-	for (URigVMLink* Link : Links)
+	if (bOutput)
 	{
-		Notify(ERigVMGraphNotifType::LinkRemoved, Link);
-
-		URigVMPin* SourcePin = Link->GetSourcePin();
-		URigVMPin* TargetPin = Link->GetTargetPin();
+		OutputPinRedirectors.FindOrAdd(OldPinPath) = NewPinPath;
+	}
+}
 
-		if (SourcePin)
-		{
-			Link->SourcePinPath = SourcePin->GetPinPath();
-			SourcePin->Links.Remove(Link);
-		}
+void URigVMController::ClearPinRedirectors()
+{
+	InputPinRedirectors.Reset();
+	OutputPinRedirectors.Reset();
+}
 
-		if (TargetPin)
-		{
-			Link->TargetPinPath = TargetPin->GetPinPath();
-			TargetPin->Links.Remove(Link);
-		}
+#if WITH_EDITOR
 
-		Link->SourcePin = nullptr;
-		Link->TargetPin = nullptr;
+bool URigVMController::ShouldRedirectPin(UScriptStruct* InOwningStruct, const FString& InOldRelativePinPath, FString& InOutNewRelativePinPath) const
+{
+	if(InOwningStruct == nullptr) // potentially a template node
+	{
+		return false;
 	}
-
-	if (InLinks == nullptr)
+	
+	FControlRigStructPinRedirectorKey RedirectorKey(InOwningStruct, InOldRelativePinPath);
+	if (const FString* RedirectedPinPath = PinPathCoreRedirectors.Find(RedirectorKey))
 	{
-		for (URigVMNode* Node : Graph->Nodes)
-		{
-			if (URigVMCollapseNode* CollapseNode = Cast<URigVMCollapseNode>(Node))
-			{
-				FRigVMControllerGraphGuard GraphGuard(this, CollapseNode->GetContainedGraph(), false);
-				TGuardValue<bool> GuardEditGraph(CollapseNode->ContainedGraph->bEditable, true);
-				DetachLinksFromPinObjects(InLinks);
-			}
-		}
+		InOutNewRelativePinPath = *RedirectedPinPath;
+		return InOutNewRelativePinPath != InOldRelativePinPath;
 	}
 
-	return Links.Num();
-}
-
-int32 URigVMController::ReattachLinksToPinObjects(bool bFollowCoreRedirectors, const TArray<URigVMLink*>* InLinks, bool bSetupOrphanedPins, bool bAllowNonArgumentLinks, bool bRecursive)
-{
-	URigVMGraph* Graph = GetGraph();
-	check(Graph);
-	FScopeLock Lock(&PinPathCoreRedirectorsLock);
+	FString RelativePinPath = InOldRelativePinPath;
+	FString PinName, SubPinPath;
+	if (!URigVMPin::SplitPinPathAtStart(RelativePinPath, PinName, SubPinPath))
+	{
+		PinName = RelativePinPath;
+		SubPinPath.Empty();
+	}
 
-	bool bReplacingAllLinks = false;
-	TArray<URigVMLink*> Links;
-	if (InLinks)
+	bool bShouldRedirect = false;
+	FCoreRedirectObjectName OldObjectName(*PinName, InOwningStruct->GetFName(), *InOwningStruct->GetOutermost()->GetPathName());
+	FCoreRedirectObjectName NewObjectName = FCoreRedirects::GetRedirectedName(ECoreRedirectFlags::Type_Property, OldObjectName);
+	if (OldObjectName != NewObjectName)
 	{
-		Links = *InLinks;
+		PinName = NewObjectName.ObjectName.ToString();
+		bShouldRedirect = true;
 	}
-	else
+
+	FProperty* Property = InOwningStruct->FindPropertyByName(*PinName);
+	if (Property == nullptr)
 	{
-		Links = Graph->Links;
-		bReplacingAllLinks = true;
+		return false;
 	}
 
-	TMap<FString, FString> RedirectedPinPaths;
-	if (bFollowCoreRedirectors)
+	if (!SubPinPath.IsEmpty())
 	{
-		for (URigVMLink* Link : Links)
+		if (FStructProperty* StructProperty = CastField<FStructProperty>(Property))
 		{
-			FString RedirectedSourcePinPath;
-			if (ShouldRedirectPin(Link->SourcePinPath, RedirectedSourcePinPath))
+			FString NewSubPinPath;
+			if (ShouldRedirectPin(StructProperty->Struct, SubPinPath, NewSubPinPath))
 			{
-				OutputPinRedirectors.FindOrAdd(Link->SourcePinPath, RedirectedSourcePinPath);
+				SubPinPath = NewSubPinPath;
+				bShouldRedirect = true;
 			}
-
-			FString RedirectedTargetPinPath;
-			if (ShouldRedirectPin(Link->TargetPinPath, RedirectedTargetPinPath))
+		}
+		else if (FArrayProperty* ArrayProperty = CastField<FArrayProperty>(Property))
+		{
+			FString SubPinName, SubSubPinPath;
+			if (URigVMPin::SplitPinPathAtStart(SubPinPath, SubPinName, SubSubPinPath))
 			{
-				InputPinRedirectors.FindOrAdd(Link->TargetPinPath, RedirectedTargetPinPath);
+				if (FStructProperty* InnerStructProperty = CastField<FStructProperty>(ArrayProperty->Inner))
+				{
+					FString NewSubSubPinPath;
+					if (ShouldRedirectPin(InnerStructProperty->Struct, SubSubPinPath, NewSubSubPinPath))
+					{
+						SubSubPinPath = NewSubSubPinPath;
+						SubPinPath = URigVMPin::JoinPinPath(SubPinName, SubSubPinPath);
+						bShouldRedirect = true;
+					}
+				}
 			}
 		}
 	}
 
-	// fix up the pin links based on the persisted data
-	TArray<URigVMLink*> NewLinks;
-	for (URigVMLink* Link : Links)
+	if (bShouldRedirect)
 	{
-		if (FString* RedirectedSourcePinPath = OutputPinRedirectors.Find(Link->SourcePinPath))
-		{
-			ensure(Link->SourcePin == nullptr);
-			Link->SourcePinPath = *RedirectedSourcePinPath;
-		}
-
-		if (FString* RedirectedTargetPinPath = InputPinRedirectors.Find(Link->TargetPinPath))
-		{
-			ensure(Link->TargetPin == nullptr);
-			Link->TargetPinPath = *RedirectedTargetPinPath;
-		}
-
-		URigVMPin* SourcePin = Link->GetSourcePin();
-		URigVMPin* TargetPin = Link->GetTargetPin();
-
-		if(bSetupOrphanedPins && (SourcePin != nullptr) && (TargetPin != nullptr))
-		{
-			// ignore duplicated links that have been processed
-			if (SourcePin->IsLinkedTo(TargetPin))
-			{
-				NewLinks.Add(Link);
-				continue;
-			}
-
-			if (!URigVMPin::CanLink(SourcePin, TargetPin, nullptr, nullptr, ERigVMPinDirection::IO, bAllowNonArgumentLinks))
-			{
-				if(SourcePin->GetNode()->HasOrphanedPins() && bSetupOrphanedPins)
-				{
-					SourcePin = nullptr;
-				}
-				else if(TargetPin->GetNode()->HasOrphanedPins() && bSetupOrphanedPins)
-				{
-					TargetPin = nullptr;
-				}
-				else
-				{
-					ReportWarningf(TEXT("Unable to re-create link %s"), *URigVMLink::GetPinPathRepresentation(Link->SourcePinPath, Link->TargetPinPath));
-					TargetPin->Links.Remove(Link);
-					SourcePin->Links.Remove(Link);
-					continue;
-				}
-			}
-		}
-
-		if(bSetupOrphanedPins)
-		{
-			bool bRelayedBackToOrphanPin = false;
-			for(int32 PinIndex=0; PinIndex<2; PinIndex++)
-			{
-				URigVMPin*& PinToFind = PinIndex == 0 ? SourcePin : TargetPin;
-				
-				if(PinToFind == nullptr)
-				{
-					const FString& PinPathToFind = PinIndex == 0 ? Link->SourcePinPath : Link->TargetPinPath;
-					FString NodeName, RemainingPinPath;
-					URigVMPin::SplitPinPathAtStart(PinPathToFind, NodeName, RemainingPinPath);
-					check(!NodeName.IsEmpty() && !RemainingPinPath.IsEmpty());
-
-					URigVMNode* Node = Graph->FindNode(NodeName);
-					if(Node == nullptr)
-					{
-						continue;
-					}
-
-					RemainingPinPath = FString::Printf(TEXT("%s%s"), *URigVMPin::OrphanPinPrefix, *RemainingPinPath);
-					PinToFind = Node->FindPin(RemainingPinPath);
-
-					if(PinToFind != nullptr)
-					{
-						if(PinIndex == 0)
-						{
-							Link->SourcePinPath = PinToFind->GetPinPath();
-							Link->SourcePin = nullptr;
-							SourcePin = Link->GetSourcePin();
-						}
-						else
-						{
-							Link->TargetPinPath = PinToFind->GetPinPath();
-							Link->TargetPin = nullptr;
-							TargetPin = Link->GetTargetPin();
-						}
-						bRelayedBackToOrphanPin = true;
-					}
-				}
-			}
-		}
-		
-		if (SourcePin == nullptr)
-		{
-			ReportWarningf(TEXT("Unable to re-create link %s"), *URigVMLink::GetPinPathRepresentation(Link->SourcePinPath, Link->TargetPinPath));
-			if (TargetPin != nullptr)
-			{
-				TargetPin->Links.Remove(Link);
-			}
-			continue;
-		}
-		if (TargetPin == nullptr)
-		{
-			ReportWarningf(TEXT("Unable to re-create link %s"), *URigVMLink::GetPinPathRepresentation(Link->SourcePinPath, Link->TargetPinPath));
-			if (SourcePin != nullptr)
-			{
-				SourcePin->Links.Remove(Link);
-			}
-			continue;
-		}
-
-		SourcePin->Links.AddUnique(Link);
-		TargetPin->Links.AddUnique(Link);
-		NewLinks.Add(Link);
-	}
-
-	if (bReplacingAllLinks)
-	{
-		if(Graph->Links.Num() != NewLinks.Num())
-		{
-			ReportWarningf(TEXT("Number of links changed during ReattachLinksToPinObjects in graph %s in project %s"), *Graph->GetPathName(), *GetPackage()->GetPathName());
-		}
-		Graph->Links = NewLinks;
-
-		for (URigVMLink* Link : Graph->Links)
-		{
-			Notify(ERigVMGraphNotifType::LinkAdded, Link);
-		}
-	}
-	else
-	{
-		// if we are running of a subset of links
-		// find the ones we weren't able to connect
-		// again and remove them.
-		for (URigVMLink* Link : Links)
-		{
-			if (!NewLinks.Contains(Link))
-			{
-				Graph->Links.Remove(Link);
-				Notify(ERigVMGraphNotifType::LinkRemoved, Link);
-			}
-			else
-			{
-				Notify(ERigVMGraphNotifType::LinkAdded, Link);
-			}
-		}
-	}
-
-	if (bRecursive && InLinks == nullptr)
-	{
-		for (URigVMNode* Node : Graph->Nodes)
-		{
-			if (URigVMCollapseNode* CollapseNode = Cast<URigVMCollapseNode>(Node))
-			{
-				FRigVMControllerGraphGuard GraphGuard(this, CollapseNode->GetContainedGraph(), false);
-				TGuardValue<bool> GuardEditGraph(CollapseNode->ContainedGraph->bEditable, true);
-				ReattachLinksToPinObjects(bFollowCoreRedirectors, nullptr, bSetupOrphanedPins, bAllowNonArgumentLinks, bRecursive);
-			}
-		}
-	}
-
-	InputPinRedirectors.Reset();
-	OutputPinRedirectors.Reset();
-
-	return NewLinks.Num();
-}
-
-void URigVMController::RemoveStaleNodes()
-{
-	if (!IsValidGraph())
-	{
-		return;
-	}
-
-	URigVMGraph* Graph = GetGraph();
-	check(Graph);
-
-	Graph->Nodes.Remove(nullptr);
-}
-
-void URigVMController::AddPinRedirector(bool bInput, bool bOutput, const FString& OldPinPath, const FString& NewPinPath)
-{
-	if (OldPinPath.IsEmpty() || NewPinPath.IsEmpty() || OldPinPath == NewPinPath)
-	{
-		return;
-	}
-
-	if (bInput)
-	{
-		InputPinRedirectors.FindOrAdd(OldPinPath) = NewPinPath;
-	}
-	if (bOutput)
-	{
-		OutputPinRedirectors.FindOrAdd(OldPinPath) = NewPinPath;
-	}
-}
-
-#if WITH_EDITOR
-
-bool URigVMController::ShouldRedirectPin(UScriptStruct* InOwningStruct, const FString& InOldRelativePinPath, FString& InOutNewRelativePinPath) const
-{
-	if(InOwningStruct == nullptr) // potentially a template node
-	{
-		return false;
-	}
-	
-	FControlRigStructPinRedirectorKey RedirectorKey(InOwningStruct, InOldRelativePinPath);
-	if (const FString* RedirectedPinPath = PinPathCoreRedirectors.Find(RedirectorKey))
-	{
-		InOutNewRelativePinPath = *RedirectedPinPath;
-		return InOutNewRelativePinPath != InOldRelativePinPath;
-	}
-
-	FString RelativePinPath = InOldRelativePinPath;
-	FString PinName, SubPinPath;
-	if (!URigVMPin::SplitPinPathAtStart(RelativePinPath, PinName, SubPinPath))
-	{
-		PinName = RelativePinPath;
-		SubPinPath.Empty();
-	}
-
-	bool bShouldRedirect = false;
-	FCoreRedirectObjectName OldObjectName(*PinName, InOwningStruct->GetFName(), *InOwningStruct->GetOutermost()->GetPathName());
-	FCoreRedirectObjectName NewObjectName = FCoreRedirects::GetRedirectedName(ECoreRedirectFlags::Type_Property, OldObjectName);
-	if (OldObjectName != NewObjectName)
-	{
-		PinName = NewObjectName.ObjectName.ToString();
-		bShouldRedirect = true;
-	}
-
-	FProperty* Property = InOwningStruct->FindPropertyByName(*PinName);
-	if (Property == nullptr)
-	{
-		return false;
-	}
-
-	if (!SubPinPath.IsEmpty())
-	{
-		if (FStructProperty* StructProperty = CastField<FStructProperty>(Property))
-		{
-			FString NewSubPinPath;
-			if (ShouldRedirectPin(StructProperty->Struct, SubPinPath, NewSubPinPath))
-			{
-				SubPinPath = NewSubPinPath;
-				bShouldRedirect = true;
-			}
-		}
-		else if (FArrayProperty* ArrayProperty = CastField<FArrayProperty>(Property))
-		{
-			FString SubPinName, SubSubPinPath;
-			if (URigVMPin::SplitPinPathAtStart(SubPinPath, SubPinName, SubSubPinPath))
-			{
-				if (FStructProperty* InnerStructProperty = CastField<FStructProperty>(ArrayProperty->Inner))
-				{
-					FString NewSubSubPinPath;
-					if (ShouldRedirectPin(InnerStructProperty->Struct, SubSubPinPath, NewSubSubPinPath))
-					{
-						SubSubPinPath = NewSubSubPinPath;
-						SubPinPath = URigVMPin::JoinPinPath(SubPinName, SubSubPinPath);
-						bShouldRedirect = true;
-					}
-				}
-			}
-		}
-	}
-
-	if (bShouldRedirect)
-	{
-		if (SubPinPath.IsEmpty())
+		if (SubPinPath.IsEmpty())
 		{
 			InOutNewRelativePinPath = PinName;
 			PinPathCoreRedirectors.Add(RedirectorKey, InOutNewRelativePinPath);
@@ -14583,7 +14309,7 @@ bool URigVMController::ShouldRedirectPin(const FString& InOldPinPath, FString& I
 	return false;
 }
 
-void URigVMController::RepopulatePinsOnNode(URigVMNode* InNode, bool bFollowCoreRedirectors, bool bSetupOrphanedPins, bool bDetachAndReattachLinks)
+void URigVMController::RepopulatePinsOnNode(URigVMNode* InNode, bool bFollowCoreRedirectors, bool bSetupOrphanedPins, bool bRecreateLinks)
 {
 	if (InNode == nullptr)
 	{
@@ -14941,7 +14667,7 @@ void URigVMController::RepopulatePinsOnNode(URigVMNode* InNode, bool bFollowCore
 	UE_LOG(LogRigVMDeveloper, Display, TEXT("Repopulating pins on node %s"), *InNode->GetPathName());
 #endif
 	
-	bool bRequireDetachLinks = false;
+	bool bRequireRecreateLinks = false;
 	bool bRequirePinStates = false;
 
 	TArray<int32> NewPinsToAdd, PreviousPinsToRemove, PreviousPinsToOrphan, PreviousPinsToUpdate;
@@ -14959,7 +14685,7 @@ void URigVMController::RepopulatePinsOnNode(URigVMNode* InNode, bool bFollowCore
 				{
 					if(!Pin->GetLinks().IsEmpty())
 					{
-						bRequireDetachLinks = true;
+						bRequireRecreateLinks = true;
 						bRequirePinStates = true;
 					}
 
@@ -14974,7 +14700,7 @@ void URigVMController::RepopulatePinsOnNode(URigVMNode* InNode, bool bFollowCore
 							{
 								PreviousPinsToOrphan.Add(RootIndex);
 							
-								bRequireDetachLinks = true;
+								bRequireRecreateLinks = true;
 								bRequirePinStates = true;
 #if UE_RIGVMCONTROLLER_VERBOSE_REPOPULATE
 								UE_LOG(LogRigVMDeveloper, Display, TEXT("Previously existing pin '%s' needs to be orphaned."), *RootPin->GetPinPath());
@@ -14998,11 +14724,11 @@ void URigVMController::RepopulatePinsOnNode(URigVMNode* InNode, bool bFollowCore
 			const bool bTypesDiffer = !Registry.CanMatchTypes(PreviousPinInfos[Index].TypeIndex, NewPinInfos[NewIndex].TypeIndex, true);
 			if(PreviousPinInfos[Index].Direction != NewPinInfos[NewIndex].Direction)
 			{
-				bRequireDetachLinks = true;
+				bRequireRecreateLinks = true;
 			}
 			else if(PreviousPinInfos[Index].Direction != ERigVMPinDirection::Hidden)
 			{
-				bRequireDetachLinks |= bTypesDiffer;
+				bRequireRecreateLinks |= bTypesDiffer;
 			}
 
 			if(Registry.CanMatchTypes(PreviousPinInfos[Index].TypeIndex, NewPinInfos[NewIndex].TypeIndex, true))
@@ -15138,20 +14864,20 @@ void URigVMController::RepopulatePinsOnNode(URigVMNode* InNode, bool bFollowCore
 	}
 
 	TMap<FString, FPinState> PinStates;
-	TArray<URigVMLink*> DetachedLinks;
-
+	TArray<FLinkedPath> LinkedPaths;
+	
 	if(bRequirePinStates)
 	{
 		PinStates = GetPinStates(InNode);
 	}
 
-	if(bDetachAndReattachLinks && bRequireDetachLinks)
+	if(bRecreateLinks && bRequireRecreateLinks)
 	{
 #if UE_RIGVMCONTROLLER_VERBOSE_REPOPULATE
 		UE_LOG(LogRigVMDeveloper, Display, TEXT("Detaching links of node %s."), *InNode->GetPathName());
 #endif
-		DetachedLinks = InNode->GetLinks();
-		DetachLinksFromPinObjects(&DetachedLinks);
+		LinkedPaths = GetLinkedPaths(InNode);
+		FastBreakLinkedPaths(LinkedPaths);
 	}
 
 	// we can do a simpler version of the update and simply
@@ -15359,12 +15085,16 @@ void URigVMController::RepopulatePinsOnNode(URigVMNode* InNode, bool bFollowCore
 		ApplyPinStates(InNode, PinStates, RedirectedPinPaths);
 	}
 
-	if(!DetachedLinks.IsEmpty())
+	if(!LinkedPaths.IsEmpty())
 	{
 #if UE_RIGVMCONTROLLER_VERBOSE_REPOPULATE
 		UE_LOG(LogRigVMDeveloper, Display, TEXT("Reattaching links of node %s..."), *InNode->GetPathName());
 #endif
-		ReattachLinksToPinObjects(bFollowCoreRedirectors, &DetachedLinks, bSetupOrphanPinsForThisNode);
+
+		FRestoreLinkedPathSettings Settings;
+		Settings.bFollowCoreRedirectors = true;
+		Settings.bRelayToOrphanPins = true;
+		RestoreLinkedPaths(LinkedPaths, Settings);
 	}
 
 	if (URigVMInjectionInfo* InjectionInfo = InNode->GetInjectionInfo())
@@ -15403,6 +15133,9 @@ void URigVMController::RemovePinsDuringRepopulate(URigVMNode* InNode, TArray<URi
 			{
 				if(Pin->IsRootPin()) // if we are passing root pins we can reparent them directly
 				{
+					check(RootPin->GetLinkedSourcePins(true).IsEmpty());
+					check(RootPin->GetLinkedTargetPins(true).IsEmpty());
+					
 					RootPin->DisplayName = RootPin->GetFName();
 					RenameObject(RootPin, *OrphanedName, nullptr);
 					InNode->Pins.Remove(RootPin);
@@ -15454,7 +15187,7 @@ void URigVMController::RemovePinsDuringRepopulate(URigVMNode* InNode, TArray<URi
 	InPins.Reset();
 }
 
-bool URigVMController::RemoveUnusedOrphanedPins(URigVMNode* InNode)
+bool URigVMController::RemoveUnusedOrphanedPins(URigVMNode* InNode, bool bRelayLinks)
 {
 	if(!InNode->HasOrphanedPins())
 	{
@@ -15467,8 +15200,26 @@ bool URigVMController::RemoveUnusedOrphanedPins(URigVMNode* InNode)
 	{
 		URigVMPin* OrphanedPin = OrphanedPins[PinIndex];
 
-		const int32 NumSourceLinks = OrphanedPin->GetSourceLinks(true).Num(); 
-		const int32 NumTargetLinks = OrphanedPin->GetTargetLinks(true).Num();
+		int32 NumSourceLinks = OrphanedPin->GetSourceLinks(true).Num(); 
+		int32 NumTargetLinks = OrphanedPin->GetTargetLinks(true).Num();
+
+		if(bRelayLinks && (NumSourceLinks + NumTargetLinks > 0))
+		{
+			TArray<URigVMLink*> Links = OrphanedPin->GetSourceLinks(true); 
+			Links.Append(OrphanedPin->GetTargetLinks(true));
+
+			TArray<FLinkedPath> LinkedPaths = GetLinkedPaths(Links);
+			FastBreakLinkedPaths(LinkedPaths);
+
+			FRestoreLinkedPathSettings Settings;
+			Settings.bFollowCoreRedirectors = true;
+			Settings.bRelayToOrphanPins = true;
+
+			RestoreLinkedPaths(LinkedPaths, Settings);
+
+			NumSourceLinks = OrphanedPin->GetSourceLinks(true).Num();
+			NumTargetLinks = OrphanedPin->GetTargetLinks(true).Num();
+		}
 
 		if(NumSourceLinks + NumTargetLinks == 0)
 		{
@@ -15616,7 +15367,7 @@ URigVMController::FPinState URigVMController::GetPinState(URigVMPin* InPin, bool
 		}
 		State.InjectionInfos.Reset();
 	}
-	
+
 	return State;
 }
 
@@ -15828,7 +15579,7 @@ void URigVMController::ReportPinTypeChange(URigVMPin* InPin, const FString& InNe
 		*InPin->GetCPPType(),
 		*InNewCPPType
 	);
-	*/
+	//*/
 }
 
 void URigVMController::SendUserFacingNotification(const FString& InMessage, float InDuration, const UObject* InSubject, const FName& InBrushName) const
@@ -16694,20 +16445,19 @@ bool URigVMController::UpdateTemplateNodePinTypes(URigVMTemplateNode* InNode, bo
 				}
 
 				// execute pins are no longer part of the template, avoid changing the type
-				if(Argument == nullptr && Pin->IsExecuteContext())
+				if(Pin->IsExecuteContext() && !Pin->GetNode()->IsA<URigVMRerouteNode>())
 				{
-					CPPType = Pin->GetCPPType();
-					CPPObjectType = Pin->GetCPPTypeObject();
+					check(Argument == nullptr);
+					MakeExecutePin(Pin);
+					continue;
 				}
 
 				if (Pin->GetCPPType() != CPPType || Pin->GetCPPTypeObject() != CPPObjectType)
 				{
-					bAnyTypeChanged = !Pin->IsExecuteContext();
-					if(bAnyTypeChanged)
-					{
-						ReportPinTypeChange(Pin, CPPType);
-					}
+					check(!Pin->IsExecuteContext() || Pin->GetNode()->IsA<URigVMRerouteNode>());
+					ReportPinTypeChange(Pin, CPPType);
 					ChangePinType(Pin, CPPType, CPPObjectType, bSetupUndoRedo, false, false, false, bInitializeDefaultValue);
+					bAnyTypeChanged = true;
 				}
 			}
 			else 
@@ -16993,7 +16743,6 @@ bool URigVMController::ChangePinType(URigVMPin* InPin, TRigVMTypeIndex InTypeInd
 	{
 		return false;
 	}
-	
 
 	// only allow valid pin cpp types on template nodes
 	TRigVMTypeIndex TypeIndex = InTypeIndex;
@@ -17039,7 +16788,7 @@ bool URigVMController::ChangePinType(URigVMPin* InPin, TRigVMTypeIndex InTypeInd
 		ActionStack->BeginAction(Action);
 	}
 
-	TArray<URigVMLink*> Links;
+	TArray<FLinkedPath> LinkedPaths;
 
 	if (bSetupUndoRedo)
 	{
@@ -17054,11 +16803,10 @@ bool URigVMController::ChangePinType(URigVMPin* InPin, TRigVMTypeIndex InTypeInd
 	
 	if(bSetupOrphanPins)
 	{
-		Links.Append(InPin->GetSourceLinks(true));
-		Links.Append(InPin->GetTargetLinks(true));
-		if (Links.Num() > 0)
+		LinkedPaths = GetLinkedPaths(InPin, true, true);
+		if(!LinkedPaths.IsEmpty())
 		{
-			DetachLinksFromPinObjects(&Links);
+			FastBreakLinkedPaths(LinkedPaths);
 
 			const FString OrphanedName = FString::Printf(TEXT("%s%s"), *URigVMPin::OrphanPinPrefix, *InPin->GetName());
 			if(InPin->GetNode()->FindPin(OrphanedName) == nullptr)
@@ -17319,9 +17067,11 @@ bool URigVMController::ChangePinType(URigVMPin* InPin, TRigVMTypeIndex InTypeInd
 		ActionStack->EndAction(Action);
 	}
 
-	if(Links.Num() > 0)
+	if(!LinkedPaths.IsEmpty())
 	{
-		ReattachLinksToPinObjects(false, &Links, true);
+		FRestoreLinkedPathSettings Settings;
+		Settings.bRelayToOrphanPins = true;
+		RestoreLinkedPaths(LinkedPaths, Settings);
 		RemoveUnusedOrphanedPins(InPin->GetNode());
 	}
 
@@ -17552,7 +17302,12 @@ void URigVMController::RefreshFunctionReferences(URigVMLibraryNode* InFunctionDe
 		FunctionLibrary->ForEachReference(InFunctionDefinition->GetFName(), [this, bSetupUndoRedo](URigVMFunctionReferenceNode* ReferenceNode)
 		{
 			FRigVMControllerGraphGuard GraphGuard(this, ReferenceNode->GetGraph(), bSetupUndoRedo);
+
+			const TArray<FLinkedPath> LinkedPaths = GetLinkedPaths(ReferenceNode->GetLinks());
+			FastBreakLinkedPaths(LinkedPaths);
 			RepopulatePinsOnNode(ReferenceNode, false, false, true);
+			TGuardValue<bool> ReportGuard(bReportWarningsAndErrors, false);
+			RestoreLinkedPaths(LinkedPaths);
 		});
 	}
 }
@@ -17621,16 +17376,99 @@ void URigVMController::SanitizeName(FString& InOutName, bool bAllowPeriod, bool
 	}
 }
 
-TArray<TPair<FString, FString>> URigVMController::GetLinkedPinPaths(URigVMNode* InNode, bool bIncludeInjectionNodes)
+URigVMController::FLinkedPath::FLinkedPath(URigVMLink* InLink)
+	: GraphPtr(InLink->GetGraph())
+	, SourcePinPath(InLink->GetSourcePinPath())
+	, TargetPinPath(InLink->GetTargetPinPath())
+	, bSourceNodeIsInjected(false)
+	, bTargetNodeIsInjected(false)
+{
+	if(const URigVMNode* SourceNode = InLink->GetSourceNode())
+	{
+		bSourceNodeIsInjected = SourceNode->IsInjected();
+	}
+	if(const URigVMNode* TargetNode = InLink->GetTargetNode())
+	{
+		bTargetNodeIsInjected = TargetNode->IsInjected();
+	}
+}
+
+URigVMGraph* URigVMController::FLinkedPath::GetGraph(URigVMGraph* InGraph) const
+{
+	if(InGraph)
+	{
+		return InGraph;
+	}
+	if(GraphPtr.IsValid())
+	{
+		return GraphPtr.Get();
+	}
+	return nullptr;
+}
+
+FString URigVMController::FLinkedPath::GetPinPathRepresentation() const
+{
+	return URigVMLink::GetPinPathRepresentation(SourcePinPath, TargetPinPath);
+}
+
+URigVMPin* URigVMController::FLinkedPath::GetSourcePin(URigVMGraph* InGraph) const
+{
+	const URigVMGraph* Graph = GetGraph(InGraph);
+	if(Graph == nullptr)
+	{
+		return nullptr;
+	}
+	return Graph->FindPin(SourcePinPath);
+}
+
+URigVMPin* URigVMController::FLinkedPath::GetTargetPin(URigVMGraph* InGraph) const
+{
+	const URigVMGraph* Graph = GetGraph(InGraph);
+	if(Graph == nullptr)
+	{
+		return nullptr;
+	}
+	return Graph->FindPin(TargetPinPath);
+}
+
+uint32 GetTypeHash(const URigVMController::FLinkedPath& InPath)
+{
+	uint32 Hash = GetTypeHash(InPath.GraphPtr.ToSoftObjectPath().ToString());
+	Hash = HashCombine(Hash, GetTypeHash(InPath.SourcePinPath));
+	Hash = HashCombine(Hash, GetTypeHash(InPath.TargetPinPath));
+	return Hash;
+}
+
+TArray<URigVMController::FLinkedPath> URigVMController::GetLinkedPaths() const
+{
+	if(const URigVMGraph* Graph = GetGraph())
+	{
+		return GetLinkedPaths(Graph->GetLinks());
+	}
+	return TArray<FLinkedPath>();
+}
+
+TArray<URigVMController::FLinkedPath> URigVMController::GetLinkedPaths(const TArray<URigVMLink*>& InLinks)
+{
+	TArray<FLinkedPath> LinkedPaths;
+	LinkedPaths.Reserve(InLinks.Num());
+	for(URigVMLink* Link : InLinks)
+	{
+		LinkedPaths.Emplace(Link);
+	}
+	return LinkedPaths;
+}
+
+TArray<URigVMController::FLinkedPath> URigVMController::GetLinkedPaths(URigVMNode* InNode, bool bIncludeInjectionNodes)
 {
 	const TArray<URigVMNode*> Nodes = {InNode};
-	return GetLinkedPinPaths(Nodes, bIncludeInjectionNodes);
+	return GetLinkedPaths(Nodes, bIncludeInjectionNodes);
 }
 
-TArray<TPair<FString, FString>> URigVMController::GetLinkedPinPaths(const TArray<URigVMNode*>& InNodes, bool bIncludeInjectionNodes)
+TArray<URigVMController::FLinkedPath> URigVMController::GetLinkedPaths(const TArray<URigVMNode*>& InNodes, bool bIncludeInjectionNodes)
 {
-	TArray<TPair<FString, FString>> LinkedPaths;
-	for(URigVMNode* Node : InNodes)
+	TArray<FLinkedPath> LinkedPaths;
+	for(const URigVMNode* Node : InNodes)
 	{
 		TArray<URigVMLink*> Links = Node->GetLinks();
 		for(URigVMLink* Link : Links)
@@ -17643,38 +17481,112 @@ TArray<TPair<FString, FString>> URigVMController::GetLinkedPinPaths(const TArray
 					continue;
 				}
 			}
-			TPair<FString, FString> LinkedPath(Link->GetSourcePin()->GetPinPath(), Link->GetTargetPin()->GetPinPath());
+			const FLinkedPath LinkedPath(Link);
 			LinkedPaths.AddUnique(LinkedPath);
 		}
 	}
 	return LinkedPaths;
 }
 
-bool URigVMController::BreakLinkedPaths(const TArray<TPair<FString, FString>>& InLinkedPaths, bool bSetupUndoRedo)
+TArray<URigVMController::FLinkedPath> URigVMController::GetLinkedPaths(const URigVMPin* InPin, bool bSourceLinksRecursive, bool bTargetLinksRecursive)
+{
+	TArray<URigVMLink*> Links;
+	Links.Append(InPin->GetSourceLinks(bSourceLinksRecursive));
+	Links.Append(InPin->GetTargetLinks(bTargetLinksRecursive));
+	return GetLinkedPaths(Links);
+}
+
+bool URigVMController::BreakLinkedPaths(const TArray<FLinkedPath>& InLinkedPaths, bool bSetupUndoRedo, bool bRelyOnBreakLink)
 {
-	for(const TPair<FString, FString>& LinkedPath : InLinkedPaths)
+	TArray<uint32> ProcessedLinks;
+
+	for(const FLinkedPath& LinkedPath : InLinkedPaths)
 	{
-		if(!BreakLink(LinkedPath.Key, LinkedPath.Value, bSetupUndoRedo))
+		// avoid duplicate links
+		const uint32 Hash = GetTypeHash(LinkedPath);
+		if(ProcessedLinks.Contains(Hash))
 		{
-			ReportErrorf(TEXT("Couldn't remove link '%s' -> '%s'"), *LinkedPath.Key, *LinkedPath.Value);
-			return false;
+			continue;
+		}
+		ProcessedLinks.Add(Hash);
+
+		if(bRelyOnBreakLink)
+		{
+			if(!BreakLink(LinkedPath.SourcePinPath, LinkedPath.TargetPinPath, bSetupUndoRedo))
+			{
+				ReportErrorf(TEXT("Couldn't remove link '%s'"), *LinkedPath.GetPinPathRepresentation());
+				return false;
+			}
+		}
+		// if we are trying to perform this as quickly as possible - let's move the links from
+		// the graph's main Links storage to a temporary location to avoid UObject creation.
+		else
+		{
+			const FString PinPathRepresentation = LinkedPath.GetPinPathRepresentation();
+			URigVMLink* Link = FindLinkFromPinPathRepresentation(PinPathRepresentation, false);
+			if(Link)
+			{
+				URigVMPin* SourcePin = Link->GetSourcePin();
+				URigVMPin* TargetPin = Link->GetTargetPin();				
+
+				check((SourcePin == nullptr) == (TargetPin == nullptr))
+
+				if(SourcePin)
+				{
+					SourcePin->Links.Remove(Link);
+				}
+				if(TargetPin)
+				{
+					TargetPin->Links.Remove(Link);
+				}
+				
+				Link->Detach();
+				GetGraph()->Links.Remove(Link);
+				GetGraph()->DetachedLinks.Add(Link);
+
+				if(bSetupUndoRedo && SourcePin && TargetPin)
+				{
+					ActionStack->AddAction(FRigVMBreakLinkAction(SourcePin, TargetPin));
+				}
+			}
+			else
+			{
+				ReportErrorf(TEXT("Couldn't remove link '%s'"), *LinkedPath.GetPinPathRepresentation());
+				return false;
+			}
 		}
 	}
 
 	return true;
 }
 
-bool URigVMController::RestoreLinkedPaths(
-	const TArray<TPair<FString, FString>>& InLinkedPaths,
-	const TMap<FString, FString>& InNodeNameMap,
-	const TMap<FString,FRigVMController_PinPathRemapDelegate>& InRemapDelegates,
-	FRigVMController_CheckPinComatibilityDelegate InCompatibilityDelegate,
-	bool bSetupUndoRedo,
-	ERigVMPinDirection InUserDirection)
+bool URigVMController::FastBreakLinkedPaths(const TArray<FLinkedPath>& InLinkedPaths, bool bSetupUndoRedo)
 {
-	bool bSuccess = true;
+	return BreakLinkedPaths(InLinkedPaths, bSetupUndoRedo, false);
+}
+
+URigVMLink* URigVMController::FindLinkFromPinPathRepresentation(const FString& InPinPathRepresentation, bool bLookForDetachedLink) const
+{
+	TArray<TObjectPtr<URigVMLink>>& LinksToSearch = bLookForDetachedLink ? GetGraph()->DetachedLinks : GetGraph()->Links;
+	const TObjectPtr<URigVMLink>* LinkPtr = LinksToSearch.FindByPredicate([InPinPathRepresentation](const URigVMLink* Link) -> bool
+	{
+		return Link->GetPinPathRepresentation().Equals(InPinPathRepresentation, ESearchCase::CaseSensitive);
+	});
 
-	auto RemapNodeName = [InNodeNameMap, InRemapDelegates](const FString& InPinPath, bool bAsInput) -> FString
+	if(LinkPtr)
+	{
+		return LinkPtr->Get();
+	}
+	
+	return nullptr;
+}
+
+TArray<URigVMController::FLinkedPath> URigVMController::RemapLinkedPaths(
+	const TArray<FLinkedPath>& InLinkedPaths,
+	const FRestoreLinkedPathSettings& InSettings,
+	bool bSetupUndoRedo)
+{
+	auto RemapPinPath = [this, InSettings, bSetupUndoRedo](const FString& InPinPath, bool bAsInput) -> FString
 	{
 		FString NodeName, SegmentPath;
 		if(!URigVMPin::SplitPinPathAtStart(InPinPath, NodeName, SegmentPath))
@@ -17683,54 +17595,308 @@ bool URigVMController::RestoreLinkedPaths(
 		}
 
 		FString PinPath = InPinPath;
+		FString RemappedPinPath;
 
-		if(const FRigVMController_PinPathRemapDelegate* RemapDelegate = InRemapDelegates.Find(NodeName))
+		if (InSettings.bFollowCoreRedirectors)
 		{
-			PinPath = RemapDelegate->Execute(PinPath, bAsInput);
+			FString RedirectedSourcePinPath;
+			if (ShouldRedirectPin(PinPath, RedirectedSourcePinPath))
+			{
+				RemappedPinPath = RedirectedSourcePinPath;
+			}
+		}
+
+		if(const FRigVMController_PinPathRemapDelegate* RemapDelegate = InSettings.RemapDelegates.Find(NodeName))
+		{
+			RemappedPinPath = RemapDelegate->Execute(PinPath, bAsInput);
+		}
+		else if(const FString* RemappedNodeName = InSettings.NodeNameMap.Find(NodeName))
+		{
+			RemappedPinPath = URigVMPin::JoinPinPath(*RemappedNodeName, SegmentPath);
+		}
+
+		// if the pin cannot be found based on this path the pin may be a sub pin on something that
+		// hasn't resolved its type yet.
+		if(InSettings.bIsImportingFromText && !RemappedPinPath.IsEmpty() && GetGraph()->FindPin(RemappedPinPath) == nullptr)
+		{
+			if (const URigVMPin* OriginalPin = GetGraph()->FindPin(InPinPath))
+			{
+				if (OriginalPin->IsStructMember())
+				{
+					const URigVMPin* OldRootPin = OriginalPin->GetRootPin();
+					
+					FString RemappedNodeName, RemappedSegmentPath, RemappedRootPinName, RemappedRemainingSegmentPath;
+					if(URigVMPin::SplitPinPathAtStart(RemappedPinPath, RemappedNodeName, RemappedSegmentPath))
+					{
+						if(URigVMPin::SplitPinPathAtStart(RemappedSegmentPath, RemappedRootPinName, RemappedRemainingSegmentPath))
+						{
+							const FString RemappedRootPinPath = URigVMPin::JoinPinPath(RemappedNodeName, RemappedRootPinName);
+							if (URigVMPin* RemappedRootPin = GetGraph()->FindPin(RemappedRootPinPath))
+							{
+								// todoooo
+								//PrepareTemplatePinForType(RemappedRootPin, {OldRootPin->GetTypeIndex()}, bSetupUndoRedo);
+							}
+						}
+					}
+				}
+			}
 		}
-		else if(const FString* RemappedNodeName = InNodeNameMap.Find(NodeName))
+
+		if(!RemappedPinPath.IsEmpty() && GetGraph()->FindPin(RemappedPinPath))
 		{
-			PinPath = URigVMPin::JoinPinPath(*RemappedNodeName, SegmentPath);
+			PinPath = RemappedPinPath;
 		}
 
 		return PinPath;
 	};
+
+	TArray<FLinkedPath> RemappedLinkedPaths;
+	RemappedLinkedPaths.Reserve(InLinkedPaths.Num());
+	
+	for(const FLinkedPath& LinkedPath : InLinkedPaths)
+	{
+		RemappedLinkedPaths.Add(LinkedPath);
+		FLinkedPath& RemappedLinkedPath = RemappedLinkedPaths.Last();
+		RemappedLinkedPath.SourcePinPath = RemapPinPath(LinkedPath.SourcePinPath, false);
+		RemappedLinkedPath.TargetPinPath = RemapPinPath(LinkedPath.TargetPinPath, true);
+
+		URigVMPin* SourcePin = RemappedLinkedPath.GetSourcePin();
+		URigVMPin* TargetPin = RemappedLinkedPath.GetTargetPin();
+		
+		if(InSettings.bRelayToOrphanPins && (SourcePin != nullptr) && (TargetPin != nullptr))
+		{
+			check(SourcePin->IsLinkedTo(TargetPin) == TargetPin->IsLinkedTo(SourcePin));
+			if (!SourcePin->IsLinkedTo(TargetPin))
+			{
+				if (!URigVMPin::CanLink(SourcePin, TargetPin, nullptr, nullptr, ERigVMPinDirection::IO, true))
+				{
+					if(SourcePin->GetNode()->HasOrphanedPins() && InSettings.bRelayToOrphanPins)
+					{
+						SourcePin = nullptr;
+					}
+					else if(TargetPin->GetNode()->HasOrphanedPins() && InSettings.bRelayToOrphanPins)
+					{
+						TargetPin = nullptr;
+					}
+					else
+					{
+						ReportWarningf(TEXT("Unable to re-create link %s"), *LinkedPath.GetPinPathRepresentation());
+						continue;
+					}
+				}
+			}
+		}
+
+		if(InSettings.bRelayToOrphanPins)
+		{
+			for(int32 PinIndex=0; PinIndex<2; PinIndex++)
+			{
+				URigVMPin*& PinToFind = PinIndex == 0 ? SourcePin : TargetPin;
+				
+				if(PinToFind == nullptr)
+				{
+					const FString& PinPathToFind = PinIndex == 0 ? RemappedLinkedPath.SourcePinPath : RemappedLinkedPath.TargetPinPath;
+					FString NodeName, RemainingPinPath;
+					URigVMPin::SplitPinPathAtStart(PinPathToFind, NodeName, RemainingPinPath);
+					check(!NodeName.IsEmpty() && !RemainingPinPath.IsEmpty());
+
+					const URigVMNode* Node = GetGraph()->FindNode(NodeName);
+					if(Node == nullptr)
+					{
+						continue;
+					}
+
+					RemainingPinPath = FString::Printf(TEXT("%s%s"), *URigVMPin::OrphanPinPrefix, *RemainingPinPath);
+					PinToFind = Node->FindPin(RemainingPinPath);
+
+					if(PinToFind != nullptr)
+					{
+						if(PinIndex == 0)
+						{
+							RemappedLinkedPath.SourcePinPath = PinToFind->GetPinPath();
+							SourcePin = PinToFind;
+						}
+						else
+						{
+							RemappedLinkedPath.TargetPinPath = PinToFind->GetPinPath();
+							TargetPin = PinToFind;
+						}
+					}
+				}
+			}
+		}
+	}
+
+	return RemappedLinkedPaths; 
+}
+
+bool URigVMController::RestoreLinkedPaths(
+	const TArray<FLinkedPath>& InLinkedPaths,
+	const FRestoreLinkedPathSettings& InSettings,
+	bool bSetupUndoRedo)
+{
+	const TArray<FLinkedPath> RemappedLinkedPaths = RemapLinkedPaths(InLinkedPaths, InSettings, bSetupUndoRedo);
+	TArray<uint32> ProcessedLinks;
+
+	bool bSuccess = true;
+	bool bAffectedAnyTemplateNode = false;
 	
-	for(const TPair<FString, FString>& LinkedPath : InLinkedPaths)
+	for(const FLinkedPath& LinkedPath : RemappedLinkedPaths)
 	{
-		const FString SourcePath = RemapNodeName(LinkedPath.Key, false);
-		const FString TargetPath = RemapNodeName(LinkedPath.Value, true);
+		const FString OriginalPinPathRepresentation = LinkedPath.GetPinPathRepresentation();
 
+		const uint32 Hash = GetTypeHash(OriginalPinPathRepresentation);
+		if(ProcessedLinks.Contains(Hash))
+		{
+			continue;
+		}
+		ProcessedLinks.Add(Hash);
+		
+		const FString SourcePath = LinkedPath.SourcePinPath;
+		const FString TargetPath = LinkedPath.TargetPinPath;
 		URigVMPin* SourcePin = GetGraph()->FindPin(SourcePath);
 		URigVMPin* TargetPin = GetGraph()->FindPin(TargetPath);
 
-		if(SourcePin == nullptr || TargetPin == nullptr)
+		URigVMLink* ExistingLink = FindLinkFromPinPathRepresentation(OriginalPinPathRepresentation, true);
+		auto OnFailedToRestoreLink = [this, ExistingLink, LinkedPath, &bSuccess](const FString& InFailureReason)
 		{
-			ReportRemovedLink(SourcePath, TargetPath);
+			const URigVMPin* SourcePin = LinkedPath.GetSourcePin(GetGraph());
+			const URigVMPin* TargetPin = LinkedPath.GetTargetPin(GetGraph());
+
+			bool bNotify = true;
+
+			// treat links on injected nodes special. if a pin is within another pin
+			// it means it is on an injected node
+			if(SourcePin && TargetPin)
+			{
+				if(SourcePin->IsLinkedTo(TargetPin))
+				{
+					if(SourcePin->IsInOuter(TargetPin) ||
+						TargetPin->IsInOuter(SourcePin))
+					{
+						bNotify = false;
+					}
+				}
+			}
+
+			if(bNotify)
+			{
+				ReportRemovedLink(LinkedPath.SourcePinPath, LinkedPath.TargetPinPath, InFailureReason);
+			}
+
+			if(ExistingLink)
+			{
+				if(bNotify)
+				{
+					Notify(ERigVMGraphNotifType::LinkRemoved, ExistingLink);
+				}
+				GetGraph()->DetachedLinks.Remove(ExistingLink);
+				DestroyObject(ExistingLink);
+			}
+
 			bSuccess = false;
+		};
+
+		if(SourcePin == nullptr || TargetPin == nullptr)
+		{
+			static const FString BothPinsMissing = TEXT("Source and target pin cannot be found.");
+			static const FString SourcePinMissing = TEXT("Source pin cannot be found.");
+			static const FString TargetPinMissing = TEXT("Target pin cannot be found.");
+
+			const FString* Reason = nullptr;
+			if(SourcePin == nullptr && TargetPin == nullptr)
+			{
+				Reason = &BothPinsMissing;
+			}
+			else if(SourcePin == nullptr)
+			{
+				Reason = &SourcePinMissing;
+			}
+			else
+			{
+				Reason = &TargetPinMissing;
+			}
+
+			OnFailedToRestoreLink(*Reason);
 			continue;
 		}
 
-		if(InCompatibilityDelegate.IsBound())
+		if(InSettings.CompatibilityDelegate.IsBound())
 		{
-			if(!InCompatibilityDelegate.Execute(SourcePin, TargetPin))
+			if(!InSettings.CompatibilityDelegate.Execute(SourcePin, TargetPin))
 			{
-				bSuccess = false;
+				static const FString IncompatibleReason = TEXT("Pins are not compatible.");
+				OnFailedToRestoreLink(IncompatibleReason);
 				continue;
 			}
 		}
 
-		// it's ok if this fails - we want to maintain the minimum set of links
-		if(!AddLink(SourcePin, TargetPin, bSetupUndoRedo, InUserDirection))
+		bool bSingleLinkSuccess = true;
+
+		if(!SourcePin->IsLinkedTo(TargetPin))
 		{
-			ReportRemovedLink(SourcePath, TargetPath);
-			bSuccess = false;
+			// make sure the existing link has the right pin path representation
+			// for AddLink to find it and reuse it
+			if(ExistingLink)
+			{
+				ExistingLink->SetSourceAndTargetPinPaths(SourcePin->GetPinPath(), TargetPin->GetPinPath());
+			}
+			
+			// it's ok if this fails - we want to maintain the minimum set of links
+			FString FailureReason;
+			if(!AddLink(SourcePin, TargetPin, bSetupUndoRedo, InSettings.UserDirection, false, true, &FailureReason))
+			{
+				OnFailedToRestoreLink(FailureReason);
+				bSingleLinkSuccess = false;
+			}
+		}
+
+		if(bSingleLinkSuccess)
+		{
+			if(SourcePin->GetNode() && SourcePin->GetNode()->IsA<URigVMTemplateNode>())
+			{
+				bAffectedAnyTemplateNode = true;
+			}
+			if(TargetPin->GetNode() && TargetPin->GetNode()->IsA<URigVMTemplateNode>())
+			{
+				bAffectedAnyTemplateNode = true;
+			}
 		}
 	}
-	
+
 	return bSuccess;
 }
 
+void URigVMController::ProcessDetachedLinks(const FRestoreLinkedPathSettings& InSettings)
+{
+	if(!IsValidGraph())
+	{
+		return;
+	}
+
+	// try to restore the links
+	if(!GetGraph()->DetachedLinks.IsEmpty())
+	{
+		TGuardValue<bool> SuspendNotifications(bSuspendNotifications, true);
+		TGuardValue<bool> ReportWarningsAndErrors(bReportWarningsAndErrors, true);
+		const TArray<FLinkedPath> LinkedPaths = GetLinkedPaths(GetGraph()->DetachedLinks);
+		RestoreLinkedPaths(LinkedPaths, InSettings);
+	}
+
+	// remove the links from the graph first
+	TArray<URigVMLink*> DetachedLinks;
+	Swap(GetGraph()->DetachedLinks, DetachedLinks);
+
+	// destroy the links finally
+	for(URigVMLink* DetachedLink : DetachedLinks)
+	{
+		ReportRemovedLink(DetachedLink->GetSourcePinPath(), DetachedLink->GetTargetPinPath());
+		Notify(ERigVMGraphNotifType::LinkRemoved, DetachedLink);
+		DestroyObject(DetachedLink);
+	}
+}
+
+
 #if WITH_EDITOR
 
 void URigVMController::RegisterUseOfTemplate(const URigVMTemplateNode* InNode)
@@ -17816,7 +17982,7 @@ FRigVMTemplate::FTypeMap URigVMController::GetCommonlyUsedTypesForTemplate(
 FRigVMClientPatchResult URigVMController::PatchUnitNodesOnLoad()
 {
 	FRigVMClientPatchResult Result;
-	
+
 	if (const URigVMGraph* Graph = GetGraph())
 	{
 		TArray<URigVMUnitNode*> UnitNodesToTurnToDispatches;
@@ -17839,29 +18005,19 @@ FRigVMClientPatchResult URigVMController::PatchUnitNodesOnLoad()
 		// convert unit nodes to dispatches
 		for(URigVMUnitNode* UnitNode : UnitNodesToTurnToDispatches)
 		{
-			TArray<TPair<FString, FString>> LinkedPaths = GetLinkedPinPaths(UnitNode);
+			TArray<FLinkedPath> LinkedPaths = GetLinkedPaths(UnitNode, true);
 			const FVector2D NodePosition = UnitNode->GetPosition();
 			const FString NodeName = UnitNode->GetName();
 			TMap<FString, FPinState> PinStates = GetPinStates(UnitNode, false);
-			
-			FRigVMTemplate::FTypeMap TypeMap;
-			for(URigVMPin* Pin : UnitNode->GetPins())
-			{
-				if(Pin->GetDirection() == ERigVMPinDirection::Input ||
-					Pin->GetDirection() == ERigVMPinDirection::Visible ||
-					Pin->GetDirection() == ERigVMPinDirection::IO ||
-					Pin->GetDirection() == ERigVMPinDirection::Output)
-				{
-					TypeMap.Add(Pin->GetFName(), Pin->GetTypeIndex());
-				}
-			}
-
+		
+			FRigVMTemplate::FTypeMap TypeMap = UnitNode->GetTemplatePinTypeMap();
 			const FRigVMTemplate* Template = UnitNode->GetTemplate();
 
 			Result.RemovedNodes.Add(UnitNode->GetPathName());
 			Result.bChangedContent = true;
 
-			RemoveNode(UnitNode, false, false, false, false);
+			FastBreakLinkedPaths(LinkedPaths);
+			RemoveNode(UnitNode, false, false);
 
 			URigVMTemplateNode* NewNode = AddTemplateNode(
 				Template->GetNotation(),
@@ -17890,7 +18046,9 @@ FRigVMClientPatchResult URigVMController::PatchUnitNodesOnLoad()
 			}
 
 			ApplyPinStates(NewNode, PinStates, {}, false);
-			RestoreLinkedPaths(LinkedPaths, {}, {}, false);
+
+			FRestoreLinkedPathSettings Settings;
+			RestoreLinkedPaths(LinkedPaths, Settings);
 		}
 	}
 
@@ -17920,6 +18078,9 @@ FRigVMClientPatchResult URigVMController::PatchDispatchNodesOnLoad()
 							FString FactoryName, ArgumentsString;
 							if(DispatchNode->ResolvedFunctionName.Split(TEXT("::"), &FactoryName, &ArgumentsString))
 							{
+								const FString PreviousResolvedFunctionName = DispatchNode->ResolvedFunctionName;
+								const int32 PreviousResolvedPermutation = DispatchNode->ResolvedPermutation;
+								
 								DispatchNode->ResolvedFunctionName.Reset();
 								DispatchNode->ResolvedPermutation = INDEX_NONE;
 								
@@ -17937,6 +18098,13 @@ FRigVMClientPatchResult URigVMController::PatchDispatchNodesOnLoad()
 										}
 									}
 								}
+
+								// fall back on the serialized information
+								if(DispatchNode->ResolvedFunctionName.IsEmpty())
+								{
+									DispatchNode->ResolvedFunctionName = PreviousResolvedFunctionName;
+									DispatchNode->ResolvedPermutation = PreviousResolvedPermutation;
+								}
 							}
 						}
 					}
@@ -17961,7 +18129,7 @@ FRigVMClientPatchResult URigVMController::PatchBranchNodesOnLoad()
 
 		for(URigVMNode* BranchNode : BranchNodes)
 		{
-			TArray<TPair<FString, FString>> LinkedPaths = GetLinkedPinPaths(BranchNode);
+			TArray<FLinkedPath> LinkedPaths = GetLinkedPaths(BranchNode, true);
 			const FVector2D NodePosition = BranchNode->GetPosition();
 			const FString NodeName = BranchNode->GetName();
 			const URigVMPin* OldConditionPin = BranchNode->FindPin(GET_MEMBER_NAME_CHECKED(FRigVMFunction_ControlFlowBranch, Condition).ToString());
@@ -17970,7 +18138,8 @@ FRigVMClientPatchResult URigVMController::PatchBranchNodesOnLoad()
 			Result.RemovedNodes.Add(BranchNode->GetPathName());
 			Result.bChangedContent = true;
 
-			RemoveNode(BranchNode, false, true, false, false);
+			FastBreakLinkedPaths(LinkedPaths);
+			RemoveNode(BranchNode, false, false);
 
 			const URigVMNode* NewNode = AddUnitNode(FRigVMFunction_ControlFlowBranch::StaticStruct(), FRigVMStruct::ExecuteName, NodePosition, NodeName, false, false);
 
@@ -17981,7 +18150,7 @@ FRigVMClientPatchResult URigVMController::PatchBranchNodesOnLoad()
 				const URigVMPin* ConditionPin = NewNode->FindPin(GET_MEMBER_NAME_CHECKED(FRigVMFunction_ControlFlowBranch, Condition).ToString());
 				SetPinDefaultValue(ConditionPin->GetPinPath(), ConditionDefault, false, false, false, false);
 			}
-			RestoreLinkedPaths(LinkedPaths, {}, {}, false);
+			RestoreLinkedPaths(LinkedPaths);
 		}
 	}
 	
@@ -17991,7 +18160,7 @@ FRigVMClientPatchResult URigVMController::PatchBranchNodesOnLoad()
 FRigVMClientPatchResult URigVMController::PatchIfSelectNodesOnLoad()
 {
 	FRigVMClientPatchResult Result;
-	
+
 	if (const URigVMGraph* Graph = GetGraph())
 	{
 		TArray<URigVMNode*> IfOrSelectNodes = Graph->GetNodes().FilterByPredicate([](URigVMNode* Node)
@@ -18003,7 +18172,7 @@ FRigVMClientPatchResult URigVMController::PatchIfSelectNodesOnLoad()
 		for(URigVMNode* IfOrSelectNode : IfOrSelectNodes)
 		{
 			const bool bIsIfNode = IfOrSelectNode->IsA<UDEPRECATED_RigVMIfNode>();
-			TArray<TPair<FString, FString>> LinkedPaths = GetLinkedPinPaths(IfOrSelectNode);
+			TArray<FLinkedPath> LinkedPaths = GetLinkedPaths(IfOrSelectNode, true);
 			const FVector2D NodePosition = IfOrSelectNode->GetPosition();
 			const FString NodeName = IfOrSelectNode->GetName();
 			const TRigVMTypeIndex TypeIndex = IfOrSelectNode->GetPins().Last()->GetTypeIndex();
@@ -18011,8 +18180,9 @@ FRigVMClientPatchResult URigVMController::PatchIfSelectNodesOnLoad()
 
 			Result.RemovedNodes.Add(IfOrSelectNode->GetPathName());
 			Result.bChangedContent = true;
-
-			RemoveNode(IfOrSelectNode, false, true, false, false);
+			
+			FastBreakLinkedPaths(LinkedPaths);
+			RemoveNode(IfOrSelectNode, false, false);
 
 			const FRigVMDispatchFactory* Factory = FRigVMRegistry::Get().FindOrAddDispatchFactory(
 				bIsIfNode ? FRigVMDispatch_If::StaticStruct() : FRigVMDispatch_SelectInt32::StaticStruct());
@@ -18050,7 +18220,9 @@ FRigVMClientPatchResult URigVMController::PatchIfSelectNodesOnLoad()
 			}
 
 			ApplyPinStates(NewNode, PinStates, {}, false);
-			RestoreLinkedPaths(LinkedPaths, {}, {}, false);
+
+			FRestoreLinkedPathSettings Settings;
+			RestoreLinkedPaths(LinkedPaths, Settings);
 		}
 	}
 
@@ -18060,7 +18232,7 @@ FRigVMClientPatchResult URigVMController::PatchIfSelectNodesOnLoad()
 FRigVMClientPatchResult URigVMController::PatchArrayNodesOnLoad()
 {
 	FRigVMClientPatchResult Result;
-	
+
 	if (const URigVMGraph* Graph = GetGraph())
 	{
 		TArray<URigVMNode*> ArrayNodes = Graph->GetNodes().FilterByPredicate([](URigVMNode* Node)
@@ -18071,7 +18243,7 @@ FRigVMClientPatchResult URigVMController::PatchArrayNodesOnLoad()
 		for(URigVMNode* ModelNode : ArrayNodes)
 		{
 			UDEPRECATED_RigVMArrayNode* ArrayNode = CastChecked<UDEPRECATED_RigVMArrayNode>(ModelNode);
-			TArray<TPair<FString, FString>> LinkedPaths = GetLinkedPinPaths(ArrayNode);
+			TArray<FLinkedPath> LinkedPaths = GetLinkedPaths(ArrayNode, true);
 			const FVector2D NodePosition = ArrayNode->GetPosition();
 			const FString NodeName = ArrayNode->GetName();
 			const FString CPPType = ArrayNode->GetCPPType();
@@ -18082,14 +18254,17 @@ FRigVMClientPatchResult URigVMController::PatchArrayNodesOnLoad()
 			Result.RemovedNodes.Add(ArrayNode->GetPathName());
 			Result.bChangedContent = true;
 
-			RemoveNode(ArrayNode, false, false, false, false);
+			FastBreakLinkedPaths(LinkedPaths);
+			RemoveNode(ArrayNode, false, false);
 
 			URigVMNode* NewNode = AddArrayNode(OpCode, CPPType, CPPTypeObject, NodePosition, NodeName, false, false, true);
 
 			Result.AddedNodes.Add(NewNode);
 
 			ApplyPinStates(NewNode, PinStates, {}, false);
-			RestoreLinkedPaths(LinkedPaths, {}, {}, false);
+
+			FRestoreLinkedPathSettings Settings;
+			RestoreLinkedPaths(LinkedPaths, Settings);
 		}
 	}
 
diff --git a/Engine/Source/Developer/RigVMDeveloper/Private/RigVMModel/RigVMControllerActions.cpp b/Engine/Source/Developer/RigVMDeveloper/Private/RigVMModel/RigVMControllerActions.cpp
--- a/Engine/Source/Developer/RigVMDeveloper/Private/RigVMModel/RigVMControllerActions.cpp
+++ b/Engine/Source/Developer/RigVMDeveloper/Private/RigVMModel/RigVMControllerActions.cpp
@@ -3,6 +3,7 @@
 #include "RigVMModel/RigVMControllerActions.h"
 #include "Serialization/MemoryReader.h"
 #include "Serialization/MemoryWriter.h"
+#include "Algo/Transform.h"
 #include "UObject/Package.h"
 
 #include UE_INLINE_GENERATED_CPP_BY_NAME(RigVMControllerActions)
@@ -226,12 +227,6 @@ void URigVMActionStack::LogAction(const UScriptStruct* InActionStruct, const FRi
 			const FRigVMActionWrapper Wrapper(SubAction);
 			LogAction(Wrapper.GetScriptStruct(), *Wrapper.GetAction(), InPrefix);
 		}
-		if(InActionStruct == FRigVMRemoveNodeAction::StaticStruct())
-		{
-			const FRigVMRemoveNodeAction* RemoveNodeAction = (const FRigVMRemoveNodeAction*)&InAction; 
-			const FRigVMActionWrapper Wrapper(RemoveNodeAction->InverseActionKey);
-			LogAction(Wrapper.GetScriptStruct(), *Wrapper.GetAction(), InPrefix);
-		}
 	}
 }
 
@@ -245,6 +240,11 @@ bool FRigVMBaseAction::Merge(const FRigVMBaseAction* Other)
 	return SubActions.Num() == 0 && Other->SubActions.Num() == 0;
 }
 
+bool FRigVMBaseAction::MakesObsolete(const FRigVMBaseAction* Other) const
+{
+	return false;
+}
+
 bool FRigVMBaseAction::Undo(URigVMController* InController)
 {
 #if RIGVM_ACTIONSTACK_VERBOSE_LOG		
@@ -295,266 +295,80 @@ bool FRigVMBaseAction::Redo(URigVMController* InController)
 	return Result;
 }
 
-#if RIGVM_ACTIONSTACK_VERBOSE_LOG
-
-void FRigVMBaseAction::LogAction(URigVMController* InController, const FString& InPrefix) const
+bool FRigVMBaseAction::StoreNode(const URigVMNode* InNode, URigVMController* InController, bool bIsPriorChange)
 {
+	check(InNode);
 	check(InController);
-	URigVMActionStack* Stack = InController->ActionStack;
-	check(Stack);
-	TGuardValue<int32> TabDepthGuard(Stack->LogActionDepth, Stack->LogActionDepth + 1);
-
-	Stack->LogAction(GetScriptStruct(), *this, InPrefix);
-}
-
-#endif
-
-bool FRigVMInverseAction::Undo(URigVMController* InController)
-{
-	return FRigVMBaseAction::Redo(InController);
-}
-
-bool FRigVMInverseAction::Redo(URigVMController* InController)
-{
-	return FRigVMBaseAction::Undo(InController);
-}
-
-FRigVMAddUnitNodeAction::FRigVMAddUnitNodeAction()
-	: ScriptStructPath()
-	, MethodName(NAME_None)
-	, Position(FVector2D::ZeroVector)
-	, NodePath()
-{
-}
-
-FRigVMAddUnitNodeAction::FRigVMAddUnitNodeAction(URigVMUnitNode* InNode)
-	: ScriptStructPath(InNode->GetScriptStruct()->GetPathName())
-	, MethodName(InNode->GetMethodName())
-	, Position(InNode->GetPosition())
-	, NodePath(InNode->GetNodePath())
-{
-}
-
-bool FRigVMAddUnitNodeAction::Undo(URigVMController* InController)
-{
-	if (!FRigVMBaseAction::Undo(InController))
-	{
-		return false;
-	}
-	return InController->RemoveNodeByName(*NodePath, false);
-}
 
-bool FRigVMAddUnitNodeAction::Redo(URigVMController* InController)
-{
-#if WITH_EDITOR
-	if (URigVMUnitNode* Node = InController->AddUnitNodeFromStructPath(ScriptStructPath, MethodName, Position, NodePath, false))
+	const FString& Content = InController->ExportNodesToText({InNode->GetFName()}, true);
+	if(!Content.IsEmpty())
 	{
-		return FRigVMBaseAction::Redo(InController);
-	}
-#endif
-	return false;
-}
-
-FRigVMAddVariableNodeAction::FRigVMAddVariableNodeAction()
-	: VariableName(NAME_None)
-	, CPPType()
-	, CPPTypeObjectPath()
-	, bIsGetter(false)
-	, DefaultValue()
-	, Position(FVector2D::ZeroVector)
-	, NodePath()
-{
-}
-
-FRigVMAddVariableNodeAction::FRigVMAddVariableNodeAction(URigVMVariableNode* InNode)
-	: VariableName(InNode->GetVariableName())
-	, CPPType(InNode->GetCPPType())
-	, CPPTypeObjectPath()
-	, bIsGetter(InNode->IsGetter())
-	, DefaultValue(InNode->GetDefaultValue())
-	, Position(InNode->GetPosition())
-	, NodePath(InNode->GetNodePath())
-{
-	if (InNode->GetCPPTypeObject())
-	{
-		CPPTypeObjectPath = InNode->GetCPPTypeObject()->GetPathName();
-	}
-}
-
-bool FRigVMAddVariableNodeAction::Undo(URigVMController* InController)
-{
-	if (!FRigVMBaseAction::Undo(InController))
-	{
-		return false;
-	}
-	return InController->RemoveNodeByName(*NodePath, false);
-}
-
-bool FRigVMAddVariableNodeAction::Redo(URigVMController* InController)
-{
-#if WITH_EDITOR
-	if (URigVMVariableNode* Node = InController->AddVariableNodeFromObjectPath(VariableName, CPPType, CPPTypeObjectPath, bIsGetter, DefaultValue, Position, NodePath, false))
-	{
-		return FRigVMBaseAction::Redo(InController);
-	}
-#endif
-	return false;
-}
-
-FRigVMAddCommentNodeAction::FRigVMAddCommentNodeAction()
-	: CommentText()
-	, Position(FVector2D::ZeroVector)
-	, Size(FVector2D::ZeroVector)
-	, Color(FLinearColor::Black)
-	, NodePath()
-{
-}
-
-FRigVMAddCommentNodeAction::FRigVMAddCommentNodeAction(URigVMCommentNode* InNode)
-	: CommentText(InNode->GetCommentText())
-	, Position(InNode->GetPosition())
-	, Size(InNode->GetSize())
-	, Color(InNode->GetNodeColor())
-	, NodePath(InNode->GetNodePath())
-{
-}
-
-bool FRigVMAddCommentNodeAction::Undo(URigVMController* InController)
-{
-	if (!FRigVMBaseAction::Undo(InController))
-	{
-		return false;
-	}
-	return InController->RemoveNodeByName(*NodePath, false);
-}
-
-bool FRigVMAddCommentNodeAction::Redo(URigVMController* InController)
-{
-#if WITH_EDITOR
-	if (URigVMCommentNode* Node = InController->AddCommentNode(CommentText, Position, Size, Color, NodePath, false))
-	{
-		return FRigVMBaseAction::Redo(InController);
+		FRigVMActionNodeContent& StoredContent = ExportedNodes.FindOrAdd(InNode->GetFName());
+		if(bIsPriorChange)
+		{
+			StoredContent.Old = Content;
+		}
+		else
+		{
+			StoredContent.New = Content;
+		}
 	}
-#endif
-	return false;
-}
-
-FRigVMAddRerouteNodeAction::FRigVMAddRerouteNodeAction()
-	: bShowAsFullNode(false)
-	, CPPType()
-	, CPPTypeObjectPath(NAME_None)
-	, DefaultValue()
-	, bIsConstant(false)
-	, CustomWidgetName(NAME_None)
-	, Position(FVector2D::ZeroVector)
-	, NodePath()
-{
+	return !Content.IsEmpty();
 }
 
-FRigVMAddRerouteNodeAction::FRigVMAddRerouteNodeAction(URigVMRerouteNode* InNode)
-	: bShowAsFullNode(InNode->GetShowsAsFullNode())
-	, CPPType(InNode->GetPins()[0]->GetCPPType())
-	, CPPTypeObjectPath(NAME_None)
-	, DefaultValue(InNode->GetPins()[0]->GetDefaultValue())
-	, bIsConstant(InNode->GetPins()[0]->IsDefinedAsConstant())
-	, CustomWidgetName(InNode->GetPins()[0]->GetCustomWidgetName())
-	, Position(InNode->GetPosition())
-	, NodePath(InNode->GetNodePath())
+bool FRigVMBaseAction::RestoreNode(const FName& InNodeName, URigVMController* InController, bool bIsUndoing)
 {
-	if (InNode->GetPins()[0]->GetCPPTypeObject())
-	{
-		CPPTypeObjectPath = *InNode->GetPins()[0]->GetCPPTypeObject()->GetPathName();
-	}
-}
+	URigVMGraph* Graph = InController->GetGraph();
+	check(Graph);
 
-bool FRigVMAddRerouteNodeAction::Undo(URigVMController* InController)
-{
-	if (!FRigVMBaseAction::Undo(InController))
-	{
-		return false;
-	}
-	return InController->RemoveNodeByName(*NodePath, false);
-}
+	FRigVMActionNodeContent& Content = ExportedNodes.FindChecked(InNodeName);
 
-bool FRigVMAddRerouteNodeAction::Redo(URigVMController* InController)
-{
-#if WITH_EDITOR
-	if (URigVMRerouteNode* Node = InController->AddFreeRerouteNode(bShowAsFullNode, CPPType, CPPTypeObjectPath, bIsConstant, CustomWidgetName, DefaultValue, Position, NodePath, false))
+	//TArray<URigVMController::FLinkedPath> LinkedPaths;
+	if(URigVMNode* ExistingNode = Graph->FindNodeByName(InNodeName))
 	{
-		return FRigVMBaseAction::Redo(InController);
+		//LinkedPaths = InController->GetLinkedPaths(ExistingNode);
+		if(!InController->RemoveNode(ExistingNode, false, false))
+		{
+			return false;
+		}
 	}
-#endif
-	return false;
-}
-
-FRigVMAddEnumNodeAction::FRigVMAddEnumNodeAction()
-	: CPPTypeObjectPath(NAME_None)
-	, Position(FVector2D::ZeroVector)
-	, NodePath()
-{
-}
 
-FRigVMAddEnumNodeAction::FRigVMAddEnumNodeAction(URigVMEnumNode* InNode)
-	: CPPTypeObjectPath(NAME_None)
-	, Position(InNode->GetPosition())
-	, NodePath(InNode->GetNodePath())
-{
-	CPPTypeObjectPath = *InNode->GetCPPTypeObject()->GetPathName();
-}
-
-bool FRigVMAddEnumNodeAction::Undo(URigVMController* InController)
-{
-	if (!FRigVMBaseAction::Undo(InController))
+	const FString& ContentToImport = bIsUndoing ? Content.Old : Content.New;
+	const TArray<FName> ImportedNodeNames = InController->ImportNodesFromText(ContentToImport, false, false);
+	if(ImportedNodeNames.Num() == 1)
 	{
-		return false;
+		if(!ImportedNodeNames[0].IsEqual(InNodeName, ENameCase::CaseSensitive))
+		{
+			return false;
+		}
 	}
-	return InController->RemoveNodeByName(*NodePath, false);
-}
 
-bool FRigVMAddEnumNodeAction::Redo(URigVMController* InController)
-{
-#if WITH_EDITOR
-	if (URigVMEnumNode* Node = InController->AddEnumNode(CPPTypeObjectPath, Position, NodePath, false))
+	/*
+	if(!LinkedPaths.IsEmpty())
 	{
-		return FRigVMBaseAction::Redo(InController);
+		if(!InController->RestoreLinkedPaths(LinkedPaths))
+		{
+			return false;
+		}
 	}
-#endif
-	return false;
+	*/
+	
+	return true;
 }
 
-FRigVMAddTemplateNodeAction::FRigVMAddTemplateNodeAction()
-	: TemplateNotation(NAME_None)
-	, Position(FVector2D::ZeroVector)
-	, NodePath()
-{
-}
+#if RIGVM_ACTIONSTACK_VERBOSE_LOG
 
-FRigVMAddTemplateNodeAction::FRigVMAddTemplateNodeAction(URigVMTemplateNode* InNode)
-	: TemplateNotation(InNode->GetNotation())
-	, Position(InNode->GetPosition())
-	, NodePath(InNode->GetNodePath())
+void FRigVMBaseAction::LogAction(URigVMController* InController, const FString& InPrefix) const
 {
-}
+	check(InController);
+	URigVMActionStack* Stack = InController->ActionStack;
+	check(Stack);
+	TGuardValue<int32> TabDepthGuard(Stack->LogActionDepth, Stack->LogActionDepth + 1);
 
-bool FRigVMAddTemplateNodeAction::Undo(URigVMController* InController)
-{
-	if (!FRigVMBaseAction::Undo(InController))
-	{
-		return false;
-	}
-	return InController->RemoveNodeByName(*NodePath, false);
+	Stack->LogAction(GetScriptStruct(), *this, InPrefix);
 }
 
-bool FRigVMAddTemplateNodeAction::Redo(URigVMController* InController)
-{
-#if WITH_EDITOR
-	if (URigVMTemplateNode* Node = InController->AddTemplateNode(TemplateNotation, Position, NodePath, false))
-	{
-		return FRigVMBaseAction::Redo(InController);
-	}
 #endif
-	return false;
-}
 
 FRigVMInjectNodeIntoPinAction::FRigVMInjectNodeIntoPinAction()
 	: PinPath()
@@ -600,148 +414,69 @@ bool FRigVMInjectNodeIntoPinAction::Redo(URigVMController* InController)
 	return false;
 }
 
-FRigVMRemoveNodeAction::FRigVMRemoveNodeAction(URigVMNode* InNode, URigVMController* InController)
+FRigVMEjectNodeFromPinAction::FRigVMEjectNodeFromPinAction()
+	: FRigVMInjectNodeIntoPinAction()
 {
-	FRigVMInverseAction InverseAction;
+}
 
-	if (URigVMVariableNode* VariableNode = Cast<URigVMVariableNode>(InNode))
-	{
-		InverseAction.AddAction(FRigVMAddVariableNodeAction(VariableNode), InController);
-		URigVMPin* ValuePin = VariableNode->GetValuePin();
-		InverseAction.AddAction(FRigVMSetPinDefaultValueAction(ValuePin, ValuePin->GetDefaultValue()), InController);
-	}
-	else if (URigVMCommentNode* CommentNode = Cast<URigVMCommentNode>(InNode))
-	{
-		InverseAction.AddAction(FRigVMAddCommentNodeAction(CommentNode), InController);
-	}
-	else if (URigVMRerouteNode* RerouteNode = Cast<URigVMRerouteNode>(InNode))
-	{
-		InverseAction.AddAction(FRigVMAddRerouteNodeAction(RerouteNode), InController);
-	}
-	else if (URigVMEnumNode* EnumNode = Cast<URigVMEnumNode>(InNode))
-	{
-		InverseAction.AddAction(FRigVMAddEnumNodeAction(EnumNode), InController);
-	}
-	else if (URigVMLibraryNode* LibraryNode = Cast<URigVMLibraryNode>(InNode))
-	{
-		InverseAction.AddAction(FRigVMImportNodeFromTextAction(LibraryNode, InController), InController);
-	}
-	else if (URigVMInvokeEntryNode* InvokeEntryNode = Cast<URigVMInvokeEntryNode>(InNode))
-	{
-		InverseAction.AddAction(FRigVMAddInvokeEntryNodeAction(InvokeEntryNode), InController);
-		URigVMPin* EntryNamePin = InvokeEntryNode->GetEntryNamePin();
-		InverseAction.AddAction(FRigVMSetPinDefaultValueAction(EntryNamePin, EntryNamePin->GetDefaultValue()), InController);
-	}
-	else if (InNode->IsA<URigVMFunctionEntryNode>() || InNode->IsA<URigVMFunctionReturnNode>())
-	{
-		// Do nothing
-	}
-	else if (URigVMTemplateNode* TemplateNode = Cast<URigVMTemplateNode>(InNode))
-	{
-		if (TemplateNode->IsSingleton())
-		{
-			if (URigVMUnitNode* UnitNode = Cast<URigVMUnitNode>(InNode))
-			{
-				InverseAction.AddAction(FRigVMAddUnitNodeAction(UnitNode), InController);
-			}
-		}
-		else
-		{
-			InverseAction.AddAction(FRigVMAddTemplateNodeAction(TemplateNode), InController);
-			
-			for (URigVMPin* Pin : TemplateNode->GetPins())
-			{
-				if (!Pin->IsWildCard())
-				{
-					if (const FRigVMTemplate* Template = TemplateNode->GetTemplate())
-					{
-						if (const FRigVMTemplateArgument* Argument = Template->FindArgument(Pin->GetFName()))
-						{
-							if (!Argument->IsSingleton())
-							{
-								InverseAction.AddAction(FRigVMChangePinTypeAction(Pin, Pin->GetTypeIndex(), false, false, false), InController);
-							}
-						}
-					}			
-				}
-			}
-		}
+FRigVMEjectNodeFromPinAction::FRigVMEjectNodeFromPinAction(URigVMInjectionInfo* InInjectionInfo)
+	: FRigVMInjectNodeIntoPinAction(InInjectionInfo)
+{
+}
 
-		for (URigVMPin* Pin : TemplateNode->GetPins())
-		{
-			if (!Pin->IsWildCard())
-			{
-				if (Pin->GetDirection() == ERigVMPinDirection::Input ||
-				    Pin->GetDirection() == ERigVMPinDirection::IO ||
-				    Pin->GetDirection() == ERigVMPinDirection::Visible)
-				{
-					InverseAction.AddAction(FRigVMSetPinDefaultValueAction(Pin, Pin->GetDefaultValue()), InController);
-				}
-			}
-		}
-	}
-	else
+bool FRigVMEjectNodeFromPinAction::Undo(URigVMController* InController)
+{
+	if (!FRigVMBaseAction::Undo(InController))
 	{
-		ensure(false);
+		return false;
 	}
+#if WITH_EDITOR
+	const URigVMInjectionInfo* InjectionInfo = InController->InjectNodeIntoPin(PinPath, bAsInput, InputPinName, OutputPinName, false);
+	return InjectionInfo != nullptr;
+#else
+	return false;
+#endif
+}
 
-	for (URigVMPin* Pin : InNode->GetPins())
+bool FRigVMEjectNodeFromPinAction::Redo(URigVMController* InController)
+{
+	if (InController->EjectNodeFromPin(*PinPath, false) != nullptr)
 	{
-		if(Pin->IsExpanded() && Pin->GetSubPins().Num() > 0)
-		{
-			FRigVMSetPinExpansionAction ExpansionAction(Pin, true);
-			ExpansionAction.OldIsExpanded = false;
-			InverseAction.AddAction(ExpansionAction, InController);
-		}
-
-		if (Pin->HasInjectedNodes())
-		{
-			if (URigVMVariableNode* VariableNode = Cast<URigVMVariableNode>(Pin->GetInjectedNodes()[0]->Node))
-			{
-				if (!InNode->IsA<URigVMLibraryNode>())
-				{
-					FRigVMAddVariableNodeAction AddVariableNodeAction(VariableNode);
-					FRigVMAddLinkAction AddLinkAction(VariableNode->GetValuePin(), Pin);
-					FRigVMInjectNodeIntoPinAction InjectAction(Pin->GetInjectedNodes()[0]);
-					InverseAction.AddAction(AddVariableNodeAction, InController);
-					InverseAction.AddAction(AddLinkAction, InController);
-					InverseAction.AddAction(InjectAction, InController);
-				}
-				else
-				{
-					FRigVMAddLinkAction AddLinkAction(VariableNode->GetValuePin(), Pin);
-					InverseAction.AddAction(AddLinkAction, InController);
-				}
-			}
-		}
+		return FRigVMBaseAction::Redo(InController);
 	}
+	return false;
+}
 
-	InverseActionKey.Set<FRigVMInverseAction>(InverseAction);
+FRigVMRemoveNodesAction::FRigVMRemoveNodesAction(TArray<URigVMNode*> InNodes, URigVMController* InController)
+{
+	Algo::Transform(InNodes, NodeNames, [](const URigVMNode* Node) -> FName
+	{
+		return Node->GetFName();
+	});
+	ExportedContent = InController->ExportNodesToText(NodeNames, true);
 }
 
-bool FRigVMRemoveNodeAction::Undo(URigVMController* InController)
+bool FRigVMRemoveNodesAction::Undo(URigVMController* InController)
 {
-	FRigVMActionWrapper InverseWrapper(InverseActionKey);
-	if (!InverseWrapper.GetAction()->Undo(InController))
+	if (!FRigVMBaseAction::Undo(InController))
 	{
 		return false;
 	}
-	return FRigVMBaseAction::Undo(InController);
+	const TArray<FName> NewNodeNames = InController->ImportNodesFromText(ExportedContent, false, false);
+	return NewNodeNames.Num() == NodeNames.Num();
 }
 
-bool FRigVMRemoveNodeAction::Redo(URigVMController* InController)
+bool FRigVMRemoveNodesAction::Redo(URigVMController* InController)
 {
-	if (!FRigVMBaseAction::Redo(InController))
+	if(InController->RemoveNodesByName(NodeNames, false, false))
 	{
-		return false;
+		return FRigVMBaseAction::Redo(InController);
 	}
-	FRigVMActionWrapper InverseWrapper(InverseActionKey);
-	return InverseWrapper.GetAction()->Redo(InController);
+	return false;
 }
 
 FRigVMSetNodeSelectionAction::FRigVMSetNodeSelectionAction()
 {
-
 }
 
 FRigVMSetNodeSelectionAction::FRigVMSetNodeSelectionAction(URigVMGraph* InGraph, TArray<FName> InNewSelection)
@@ -1390,43 +1125,6 @@ bool FRigVMChangePinTypeAction::Redo(URigVMController* InController)
 	return FRigVMBaseAction::Redo(InController);
 }
 
-FRigVMImportNodeFromTextAction::FRigVMImportNodeFromTextAction()
-	: Position(FVector2D::ZeroVector)
-	, NodePath()
-{
-}
-
-FRigVMImportNodeFromTextAction::FRigVMImportNodeFromTextAction(URigVMNode* InNode, URigVMController* InController)
-	: Position(InNode->GetPosition())
-	, NodePath(InNode->GetNodePath())
-	, ExportedText()
-{
-	TArray<FName> NodeNamesToExport;
-	NodeNamesToExport.Add(InNode->GetFName());
-	ExportedText = InController->ExportNodesToText(NodeNamesToExport);
-}
-
-bool FRigVMImportNodeFromTextAction::Undo(URigVMController* InController)
-{
-	if (!FRigVMBaseAction::Undo(InController))
-	{
-		return false;
-	}
-	return InController->RemoveNodeByName(*NodePath, false, true);
-}
-
-bool FRigVMImportNodeFromTextAction::Redo(URigVMController* InController)
-{
-#if WITH_EDITOR
-	TArray<FName> NodeNames = InController->ImportNodesFromText(ExportedText, false);
-	if (NodeNames.Num() > 0)
-	{
-		return FRigVMBaseAction::Redo(InController);
-	}
-#endif
-	return false;
-}
-
 FRigVMCollapseNodesAction::FRigVMCollapseNodesAction()
 	: LibraryNodePath(), bIsAggregate(false)
 {
@@ -1436,25 +1134,13 @@ FRigVMCollapseNodesAction::FRigVMCollapseNodesAction(URigVMController* InControl
 	: LibraryNodePath(InNodePath), bIsAggregate(bIsAggregate)
 {
 	TArray<FName> NodesToExport;
-	for (URigVMNode* InNode : InNodes)
+	for (const URigVMNode* InNode : InNodes)
 	{
-		NodesToExport.Add(InNode->GetFName());
 		CollapsedNodesPaths.Add(InNode->GetName());
-
-		// find the links external to the nodes to be collapsed
-		TArray<URigVMLink*> Links = InNode->GetLinks();
-		for(URigVMLink* Link : Links)
-		{
-			if(InNodes.Contains(Link->GetSourcePin()->GetNode()) &&
-				InNodes.Contains(Link->GetTargetPin()->GetNode()))
-			{
-				continue;
-			}
-			CollapsedNodesLinks.Add(Link->GetPinPathRepresentation());
-		}
+		NodesToExport.Add(InNode->GetFName());
 	}
 
-	CollapsedNodesContent = InController->ExportNodesToText(NodesToExport);
+	CollapsedNodesContent = InController->ExportNodesToText(NodesToExport, true);
 }
 
 bool FRigVMCollapseNodesAction::Undo(URigVMController* InController)
@@ -1465,9 +1151,11 @@ bool FRigVMCollapseNodesAction::Undo(URigVMController* InController)
 	}
 
 	// remove the library node
-	if(!InController->RemoveNodeByName(*LibraryNodePath, false, true, false))
 	{
-		return false;
+		if(!InController->RemoveNodeByName(*LibraryNodePath, false, false))
+		{
+			return false;
+		}
 	}
 
 	const TArray<FName> RecoveredNodes = InController->ImportNodesFromText(CollapsedNodesContent, false, false);
@@ -1476,15 +1164,6 @@ bool FRigVMCollapseNodesAction::Undo(URigVMController* InController)
 		return false;
 	}
 
-	for(const FString& CollapsedNodesLink : CollapsedNodesLinks)
-	{
-		FString Source, Target;
-		if(URigVMLink::SplitPinPathRepresentation(CollapsedNodesLink, Source, Target))
-		{
-			InController->AddLink(Source, Target, false, false);
-		}
-	}
-
 	return true;
 }
 
@@ -1516,13 +1195,7 @@ FRigVMExpandNodeAction::FRigVMExpandNodeAction(URigVMController* InController, U
 {
 	TArray<FName> NodesToExport;
 	NodesToExport.Add(InLibraryNode->GetFName());
-	LibraryNodeContent = InController->ExportNodesToText(NodesToExport);
-
-	TArray<URigVMLink*> Links = InLibraryNode->GetLinks();
-	for(URigVMLink* Link : Links)
-	{
-		LibraryNodeLinks.Add(Link->GetPinPathRepresentation());
-	}
+	LibraryNodeContent = InController->ExportNodesToText(NodesToExport, true);
 }
 
 bool FRigVMExpandNodeAction::Undo(URigVMController* InController)
@@ -1533,12 +1206,14 @@ bool FRigVMExpandNodeAction::Undo(URigVMController* InController)
 	}
 
 	// remove the expanded nodes
-	for (const FString& NodePath : ExpandedNodePaths)
+	TArray<FName> NodeNames;
+	Algo::Transform(ExpandedNodePaths, NodeNames, [](const FString& NodePath)
 	{
-		if(!InController->RemoveNodeByName(*NodePath, false, true, false))
-		{
-			return false;
-		}
+		return FName(*NodePath); 
+	});
+	if(!InController->RemoveNodesByName(NodeNames, false, false))
+	{
+		return false;
 	}
 
 	const TArray<FName> RecoveredNodes = InController->ImportNodesFromText(LibraryNodeContent, false, false);
@@ -1547,22 +1222,13 @@ bool FRigVMExpandNodeAction::Undo(URigVMController* InController)
 		return false;
 	}
 
-	for(const FString& LibraryNodeLink : LibraryNodeLinks)
-	{
-		FString Source, Target;
-		if(URigVMLink::SplitPinPathRepresentation(LibraryNodeLink, Source, Target))
-		{
-			InController->AddLink(Source, Target, false, false);
-		}
-	}
-
 	return true;
 }
 
 bool FRigVMExpandNodeAction::Redo(URigVMController* InController)
 {
 #if WITH_EDITOR
-	TArray<URigVMNode*> ExpandedNodes = InController->ExpandLibraryNode(*LibraryNodePath, false);
+	const TArray<URigVMNode*> ExpandedNodes = InController->ExpandLibraryNode(*LibraryNodePath, false);
 	if (ExpandedNodes.Num() == ExpandedNodePaths.Num())
 	{
 		return FRigVMBaseAction::Redo(InController);
@@ -1614,7 +1280,14 @@ bool FRigVMPushGraphAction::Undo(URigVMController* InController)
 	{
 		return false;
 	}
-	return InController->PopGraph(false) != nullptr;
+	
+	if(URigVMGraph* PoppedGraph = InController->PopGraph(false))
+	{
+		const FSoftObjectPath PoppedGraphPath = TSoftObjectPtr<URigVMGraph>(PoppedGraph).GetUniqueID();
+		return PoppedGraphPath == GraphPath;
+	}
+	
+	return false;
 }
 
 bool FRigVMPushGraphAction::Redo(URigVMController* InController)
@@ -1628,6 +1301,19 @@ bool FRigVMPushGraphAction::Redo(URigVMController* InController)
 	return false;
 }
 
+bool FRigVMPushGraphAction::MakesObsolete(const FRigVMBaseAction* Other) const
+{
+	if (Other->GetScriptStruct() == FRigVMPopGraphAction::StaticStruct())
+	{
+		const FRigVMPopGraphAction* PopGraphAction = (const FRigVMPopGraphAction*)Other;
+		if(PopGraphAction->GraphPath == GraphPath)
+		{
+			return true;
+		}
+	}
+	return false;
+}
+
 FRigVMPopGraphAction::FRigVMPopGraphAction(UObject* InGraph)
 {
 	GraphPath = TSoftObjectPtr<URigVMGraph>(Cast<URigVMGraph>(InGraph)).GetUniqueID();
@@ -1651,9 +1337,26 @@ bool FRigVMPopGraphAction::Undo(URigVMController* InController)
 
 bool FRigVMPopGraphAction::Redo(URigVMController* InController)
 {
-	if (InController->PopGraph(false) != nullptr)
+	if(URigVMGraph* PoppedGraph = InController->PopGraph(false))
 	{
-		return FRigVMBaseAction::Redo(InController);
+		const FSoftObjectPath PoppedGraphPath = TSoftObjectPtr<URigVMGraph>(PoppedGraph).GetUniqueID();
+		if(PoppedGraphPath == GraphPath)
+		{
+			return FRigVMBaseAction::Redo(InController);
+		}
+	}
+	return false;
+}
+
+bool FRigVMPopGraphAction::MakesObsolete(const FRigVMBaseAction* Other) const
+{
+	if (Other->GetScriptStruct() == FRigVMPushGraphAction::StaticStruct())
+	{
+		const FRigVMPushGraphAction* PushGraphAction = (const FRigVMPushGraphAction*)Other;
+		if(PushGraphAction->GraphPath == GraphPath)
+		{
+			return true;
+		}
 	}
 	return false;
 }
@@ -1993,33 +1696,31 @@ bool FRigVMPromoteNodeAction::Redo(URigVMController* InController)
 	return FunctionRefNodeName.ToString() == LibraryNodePath;
 }
 
-FRigVMAddInvokeEntryNodeAction::FRigVMAddInvokeEntryNodeAction()
-	: EntryName(NAME_None)
-	, Position(FVector2D::ZeroVector)
-	, NodePath()
+FRigVMMarkFunctionPublicAction::FRigVMMarkFunctionPublicAction()
+	: FunctionName(NAME_None)
+	, bIsPublic(false)
 {
 }
 
-FRigVMAddInvokeEntryNodeAction::FRigVMAddInvokeEntryNodeAction(URigVMInvokeEntryNode* InNode)
-	: EntryName(InNode->GetEntryName())
-	, Position(InNode->GetPosition())
-	, NodePath(InNode->GetNodePath())
+FRigVMMarkFunctionPublicAction::FRigVMMarkFunctionPublicAction(const FName& InFunctionName, bool bInIsPublic)
+	: FunctionName(InFunctionName)
+	, bIsPublic(bInIsPublic)
 {
 }
 
-bool FRigVMAddInvokeEntryNodeAction::Undo(URigVMController* InController)
+bool FRigVMMarkFunctionPublicAction::Undo(URigVMController* InController)
 {
 	if (!FRigVMBaseAction::Undo(InController))
 	{
 		return false;
 	}
-	return InController->RemoveNodeByName(*NodePath, false);
+	return InController->MarkFunctionAsPublic(FunctionName, !bIsPublic, false);
 }
 
-bool FRigVMAddInvokeEntryNodeAction::Redo(URigVMController* InController)
+bool FRigVMMarkFunctionPublicAction::Redo(URigVMController* InController)
 {
 #if WITH_EDITOR
-	if (URigVMInvokeEntryNode* Node = InController->AddInvokeEntryNode(EntryName, Position, NodePath, false))
+	if (InController->MarkFunctionAsPublic(FunctionName, bIsPublic, false))
 	{
 		return FRigVMBaseAction::Redo(InController);
 	}
@@ -2027,35 +1728,87 @@ bool FRigVMAddInvokeEntryNodeAction::Redo(URigVMController* InController)
 	return false;
 }
 
-FRigVMMarkFunctionPublicAction::FRigVMMarkFunctionPublicAction()
-	: FunctionName(NAME_None)
-	, bIsPublic(false)
+FRigVMImportFromTextAction::FRigVMImportFromTextAction(const FString& InContent, const TArray<FName>& InTopLevelNodeNames)
+	: Content(InContent)
+	, TopLevelNodeNames(InTopLevelNodeNames)
 {
 }
 
-FRigVMMarkFunctionPublicAction::FRigVMMarkFunctionPublicAction(const FName& InFunctionName, bool bInIsPublic)
-	: FunctionName(InFunctionName)
-	, bIsPublic(bInIsPublic)
+FRigVMImportFromTextAction::FRigVMImportFromTextAction(URigVMNode* InNode, URigVMController* InController, bool bIncludeExteriorLinks)
 {
+	SetContent({InNode}, InController, bIncludeExteriorLinks);
 }
 
-bool FRigVMMarkFunctionPublicAction::Undo(URigVMController* InController)
+FRigVMImportFromTextAction::FRigVMImportFromTextAction(const TArray<URigVMNode*>& InNodes, URigVMController* InController, bool bIncludeExteriorLinks)
 {
-	if (!FRigVMBaseAction::Undo(InController))
+	SetContent(InNodes, InController, bIncludeExteriorLinks);
+}
+
+void FRigVMImportFromTextAction::SetContent(const TArray<URigVMNode*>& InNodes, URigVMController* InController, bool bIncludeExteriorLinks)
+{
+	TopLevelNodeNames.Reset();
+	Algo::Transform(InNodes, TopLevelNodeNames, [](const URigVMNode* Node)
+	{
+		return Node->GetFName();
+	});
+	Content = InController->ExportNodesToText(TopLevelNodeNames, bIncludeExteriorLinks);
+}
+
+bool FRigVMImportFromTextAction::Undo(URigVMController* InController)
+{
+	if(!FRigVMBaseAction::Undo(InController))
 	{
 		return false;
 	}
-	return InController->MarkFunctionAsPublic(FunctionName, !bIsPublic, false);
+	return InController->RemoveNodesByName(TopLevelNodeNames, false, false);
 }
 
-bool FRigVMMarkFunctionPublicAction::Redo(URigVMController* InController)
+bool FRigVMImportFromTextAction::Redo(URigVMController* InController)
 {
-#if WITH_EDITOR
-	if (InController->MarkFunctionAsPublic(FunctionName, bIsPublic, false))
+	const TArray<FName> ImportedNames = InController->ImportNodesFromText(Content, false, false);
+	if(ImportedNames.Num() != TopLevelNodeNames.Num())
 	{
-		return FRigVMBaseAction::Redo(InController);
+		return false;
 	}
-#endif
-	return false;
+	return FRigVMBaseAction::Redo(InController);
+}
+
+FRigVMReplaceNodesAction::FRigVMReplaceNodesAction(const TArray<URigVMNode*>& InNodes, URigVMController* InController)
+{
+	for(const URigVMNode* Node : InNodes)
+	{
+		StoreNode(Node, InController);
+	}
+}
+
+bool FRigVMReplaceNodesAction::Undo(URigVMController* InController)
+{
+	if(!FRigVMBaseAction::Undo(InController))
+	{
+		return false;
+	}
+
+	for(const TPair<FName, FRigVMActionNodeContent>& Pair : ExportedNodes)
+	{
+		if(!RestoreNode(Pair.Key, InController, true))
+		{
+			return false;
+		}
+	}
+
+	return true;
+}
+
+bool FRigVMReplaceNodesAction::Redo(URigVMController* InController)
+{
+	for(const TPair<FName, FRigVMActionNodeContent>& Pair : ExportedNodes)
+	{
+		if(!RestoreNode(Pair.Key, InController, false))
+		{
+			return false;
+		}
+	}
+	
+	return FRigVMBaseAction::Redo(InController);
 }
 
diff --git a/Engine/Source/Developer/RigVMDeveloper/Private/RigVMModel/RigVMGraph.cpp b/Engine/Source/Developer/RigVMDeveloper/Private/RigVMModel/RigVMGraph.cpp
--- a/Engine/Source/Developer/RigVMDeveloper/Private/RigVMModel/RigVMGraph.cpp
+++ b/Engine/Source/Developer/RigVMDeveloper/Private/RigVMModel/RigVMGraph.cpp
@@ -20,6 +20,28 @@ URigVMGraph::URigVMGraph()
 	SetExecuteContextStruct(FRigVMExecuteContext::StaticStruct());
 }
 
+void URigVMGraph::PostLoad()
+{
+	Super::PostLoad();
+
+	// update the bookkeeping on the links
+	for(URigVMLink* Link : Links)
+	{
+		URigVMPin* SourcePin = Link->GetSourcePin();
+		URigVMPin* TargetPin = Link->GetTargetPin();
+		check((SourcePin == nullptr) == (TargetPin == nullptr));
+		
+		if(SourcePin)
+		{
+			SourcePin->Links.AddUnique(Link);
+		}
+		if(TargetPin)
+		{
+			TargetPin->Links.AddUnique(Link);
+		}
+	}
+}
+
 const TArray<URigVMNode*>& URigVMGraph::GetNodes() const
 {
 	return Nodes;
@@ -30,14 +52,19 @@ const TArray<URigVMLink*>& URigVMGraph::GetLinks() const
 	return Links;
 }
 
+bool URigVMGraph::ContainsLink(const FString& InPinPathRepresentation) const
+{
+	return FindLink(InPinPathRepresentation) != nullptr;
+}
+
 TArray<URigVMGraph*> URigVMGraph::GetContainedGraphs(bool bRecursive) const
 {
 	TArray<URigVMGraph*> Graphs;
 	for (URigVMNode* Node : GetNodes())
 	{
-		if (URigVMCollapseNode* CollapseNode = Cast<URigVMCollapseNode>(Node))
+		if (const URigVMCollapseNode* CollapseNode = Cast<URigVMCollapseNode>(Node))
 		{
-			Graphs.Add(CollapseNode->GetContainedGraph());
+			Graphs.AddUnique(CollapseNode->GetContainedGraph());
 			if (bRecursive)
 			{
 				Graphs.Append(CollapseNode->GetContainedGraph()->GetContainedGraphs(true));
@@ -65,6 +92,15 @@ URigVMGraph* URigVMGraph::GetRootGraph() const
 	return (URigVMGraph*)this;
 }
 
+int32 URigVMGraph::GetGraphDepth() const
+{
+	if(const URigVMGraph* ParentGraph = GetParentGraph())
+	{
+		return ParentGraph->GetGraphDepth() + 1;
+	}
+	return 0;
+}
+
 bool URigVMGraph::IsRootGraph() const
 {
 	return GetRootGraph() == this;
@@ -429,9 +465,9 @@ void URigVMGraph::PreSave(FObjectPreSaveContext SaveContext)
 	LastStructureHash = GetStructureHash();
 }
 
-bool URigVMGraph::IsNameAvailable(const FString& InName)
+bool URigVMGraph::IsNameAvailable(const FString& InName) const
 {
-	for (URigVMNode* Node : Nodes)
+	for (const URigVMNode* Node : Nodes)
 	{
 		if (Node->GetName() == InName)
 		{
diff --git a/Engine/Source/Developer/RigVMDeveloper/Private/RigVMModel/RigVMLink.cpp b/Engine/Source/Developer/RigVMDeveloper/Private/RigVMModel/RigVMLink.cpp
--- a/Engine/Source/Developer/RigVMDeveloper/Private/RigVMModel/RigVMLink.cpp
+++ b/Engine/Source/Developer/RigVMDeveloper/Private/RigVMModel/RigVMLink.cpp
@@ -10,6 +10,11 @@ URigVMGraph* URigVMLink::GetGraph() const
 	return Cast<URigVMGraph>(GetOuter());
 }
 
+int32 URigVMLink::GetGraphDepth() const
+{
+	return GetGraph()->GetGraphDepth();
+}
+
 void URigVMLink::Serialize(FArchive& Ar)
 {
 	if (Ar.IsLoading())
@@ -42,14 +47,11 @@ void URigVMLink::Serialize(FArchive& Ar)
 
 int32 URigVMLink::GetLinkIndex() const
 {
-	int32 Index = INDEX_NONE;
-	URigVMGraph* Graph = GetGraph();
-	if (Graph != nullptr)
+	if (const URigVMGraph* Graph = GetGraph())
 	{
-		Graph->GetLinks().Find((URigVMLink*)this, Index);
+		return Graph->GetLinks().Find((URigVMLink*)this);
 	}
-	return Index;
-
+	return INDEX_NONE;
 }
 
 URigVMPin* URigVMLink::GetSourcePin() const
@@ -66,7 +68,7 @@ URigVMPin* URigVMLink::GetSourcePin() const
 	}
 	if (SourcePin == nullptr)
 	{
-		SourcePin = GetGraph()->FindPin(SourcePinPath);
+		UpdatePinPointers();
 	}
 	return SourcePin;
 }
@@ -85,11 +87,143 @@ URigVMPin* URigVMLink::GetTargetPin() const
 	}
 	if (TargetPin == nullptr)
 	{
-		TargetPin = GetGraph()->FindPin(TargetPinPath);
+		UpdatePinPointers();
 	}
 	return TargetPin;
 }
 
+URigVMNode* URigVMLink::GetSourceNode() const
+{
+	if(const URigVMPin* CurrentSourcePin = GetSourcePin())
+	{
+		return CurrentSourcePin->GetNode();
+	}
+	return nullptr;
+}
+
+URigVMNode* URigVMLink::GetTargetNode() const
+{
+	if(const URigVMPin* CurrentTargetPin = GetTargetPin())
+	{
+		return CurrentTargetPin->GetNode();
+	}
+	return nullptr;
+}
+
+FString URigVMLink::GetSourcePinPath() const
+{
+	if(SourcePin)
+	{
+		return SourcePin->GetPinPath();
+	}
+	return SourcePinPath;
+}
+
+FString URigVMLink::GetTargetPinPath() const
+{
+	if(TargetPin)
+	{
+		return TargetPin->GetPinPath();
+	}
+	return TargetPinPath;
+}
+
+bool URigVMLink::SetSourcePinPath(const FString& InPinPath)
+{
+	URigVMLink* MutableThis = (URigVMLink*)this;
+	MutableThis->Detach();
+	SourcePinPath = InPinPath;
+	MutableThis->Attach();
+	return true;
+}
+
+bool URigVMLink::SetTargetPinPath(const FString& InPinPath)
+{
+	URigVMLink* MutableThis = (URigVMLink*)this;
+	MutableThis->Detach();
+	TargetPinPath = InPinPath;
+	MutableThis->Attach();
+	return true;
+}
+
+bool URigVMLink::SetSourceAndTargetPinPaths(const FString& InSourcePinPath, const FString& InTargetPinPath)
+{
+	URigVMLink* MutableThis = (URigVMLink*)this;
+	MutableThis->Detach();
+	SourcePinPath = InSourcePinPath;
+	TargetPinPath = InTargetPinPath;
+	return MutableThis->Attach();
+}
+
+bool URigVMLink::IsAttached() const
+{
+	return SourcePin != nullptr && TargetPin != nullptr;
+}
+
+bool URigVMLink::Detach()
+{
+	UpdatePinPaths();
+	SourcePin = TargetPin = nullptr;
+	return !SourcePinPath.IsEmpty() && !TargetPinPath.IsEmpty();
+}
+
+bool URigVMLink::Attach(FString* OutFailureReason)
+{
+	if(IsAttached())
+	{
+		return false;
+	}
+
+	UpdatePinPointers();
+
+	if(SourcePin == nullptr)
+	{
+		if(OutFailureReason)
+		{
+			static constexpr TCHAR Format[] = TEXT("Source pin '%s' cannot be found.");
+			*OutFailureReason = FString::Printf(Format, *SourcePinPath);
+		}
+
+		SourcePin = TargetPin = nullptr;
+		return false;
+	}
+
+	if(TargetPin == nullptr)
+	{
+		if(OutFailureReason)
+		{
+			static constexpr TCHAR Format[] = TEXT("Target pin '%s' cannot be found.");
+			*OutFailureReason = FString::Printf(Format, *TargetPinPath);
+		}
+
+		SourcePin = TargetPin = nullptr;
+		return false;
+	}
+
+	return true;
+}
+
+void URigVMLink::UpdatePinPaths()
+{
+	if(IsAttached())
+	{
+		if (const URigVMPin* CurrentSourcePin = GetSourcePin())
+		{
+			SourcePinPath = CurrentSourcePin->GetPinPath();
+		}
+		if (const URigVMPin* CurrenTargetPin = GetTargetPin())
+		{
+			TargetPinPath = CurrenTargetPin->GetPinPath();
+		}
+	}
+}
+
+void URigVMLink::UpdatePinPointers() const
+{
+	SourcePin = GetGraph()->FindPin(SourcePinPath);
+	TargetPin = GetGraph()->FindPin(TargetPinPath);
+}
+
 URigVMPin* URigVMLink::GetOppositePin(const URigVMPin* InPin) const
 {
 	if (InPin == GetSourcePin())
@@ -103,9 +237,9 @@ URigVMPin* URigVMLink::GetOppositePin(const URigVMPin* InPin) const
 	return nullptr;
 }
 
-FString URigVMLink::GetPinPathRepresentation()
+FString URigVMLink::GetPinPathRepresentation() const
 {
-	return GetPinPathRepresentation(GetSourcePin()->GetPinPath(), GetTargetPin()->GetPinPath());
+	return GetPinPathRepresentation(GetSourcePinPath(), GetTargetPinPath());
 }
 
 FString URigVMLink::GetPinPathRepresentation(const FString& InSourcePinPath, const FString& InTargetPinPath)
@@ -119,16 +253,3 @@ bool URigVMLink::SplitPinPathRepresentation(const FString& InString, FString& Ou
 	static constexpr TCHAR Format[] = TEXT(" -> "); 
 	return InString.Split(Format, &OutSource, &OutTarget);
 }
-
-void URigVMLink::PrepareForCopy()
-{
-	if (URigVMPin* CurrentSourcePin = GetSourcePin())
-	{
-		SourcePinPath = CurrentSourcePin->GetPinPath();
-	}
-	if (URigVMPin* CurrenTargetPin = GetTargetPin())
-	{
-		TargetPinPath = CurrenTargetPin->GetPinPath();
-	}
-}
-
diff --git a/Engine/Source/Developer/RigVMDeveloper/Private/RigVMModel/RigVMNode.cpp b/Engine/Source/Developer/RigVMDeveloper/Private/RigVMModel/RigVMNode.cpp
--- a/Engine/Source/Developer/RigVMDeveloper/Private/RigVMModel/RigVMNode.cpp
+++ b/Engine/Source/Developer/RigVMDeveloper/Private/RigVMModel/RigVMNode.cpp
@@ -180,6 +180,11 @@ URigVMGraph* URigVMNode::GetRootGraph() const
 	return nullptr;
 }
 
+int32 URigVMNode::GetGraphDepth() const
+{
+	return GetGraph()->GetGraphDepth();
+}
+
 URigVMInjectionInfo* URigVMNode::GetInjectionInfo() const
 {
 	return Cast<URigVMInjectionInfo>(GetOuter());
diff --git a/Engine/Source/Developer/RigVMDeveloper/Private/RigVMModel/RigVMPin.cpp b/Engine/Source/Developer/RigVMDeveloper/Private/RigVMModel/RigVMPin.cpp
--- a/Engine/Source/Developer/RigVMDeveloper/Private/RigVMModel/RigVMPin.cpp
+++ b/Engine/Source/Developer/RigVMDeveloper/Private/RigVMModel/RigVMPin.cpp
@@ -1484,9 +1484,9 @@ URigVMPin* URigVMPin::FindSubPin(const FString& InPinPath) const
 	return nullptr;
 }
 
-bool URigVMPin::IsLinkedTo(URigVMPin* InPin) const
+bool URigVMPin::IsLinkedTo(const URigVMPin* InPin) const
 {
-	for (URigVMLink* Link : Links)
+	for (const URigVMLink* Link : Links)
 	{
 		if (Link->GetSourcePin() == InPin || Link->GetTargetPin() == InPin)
 		{
@@ -1847,43 +1847,49 @@ bool URigVMPin::CanLink(URigVMPin* InSourcePin, URigVMPin* InTargetPin, FString*
 
 	if (!bInAllowNonArgumentPins)
 	{
-		if (URigVMTemplateNode* SourceTemplateNode = Cast<URigVMTemplateNode>(SourceNode))
+		if (const URigVMTemplateNode* SourceTemplateNode = Cast<URigVMTemplateNode>(SourceNode))
 		{
 			if (!SourceNode->IsA<URigVMFunctionEntryNode>() && !SourceNode->IsA<URigVMFunctionReturnNode>())
 			{
 				if (const FRigVMTemplate* Template = SourceTemplateNode->GetTemplate())
 				{
-					URigVMPin* RootPin = InSourcePin->GetRootPin();
-					if (!Template->FindArgument(RootPin->GetFName()))
+					const URigVMPin* RootPin = InSourcePin->GetRootPin();
+					if(!RootPin->IsOrphanPin())
 					{
-						if(!RootPin->IsExecuteContext())
+						if (!Template->FindArgument(RootPin->GetFName()))
 						{
-							if (OutFailureReason)
+							if(!RootPin->IsExecuteContext())
 							{
-								*OutFailureReason = FString::Printf(TEXT("Library pin %s supported types need to be reduced."), *RootPin->GetPinPath(true));
+								if (OutFailureReason)
+								{
+									*OutFailureReason = FString::Printf(TEXT("Library pin %s supported types need to be reduced."), *RootPin->GetPinPath(true));
+								}
+								return false;
 							}
-							return false;
 						}
 					}
 				}
 			}
 		}
-		if (URigVMTemplateNode* TargetTemplateNode = Cast<URigVMTemplateNode>(TargetNode))
+		if (const URigVMTemplateNode* TargetTemplateNode = Cast<URigVMTemplateNode>(TargetNode))
 		{
 			if (!TargetNode->IsA<URigVMFunctionEntryNode>() && !TargetNode->IsA<URigVMFunctionReturnNode>())
 			{
 				if (const FRigVMTemplate* Template = TargetTemplateNode->GetTemplate())
 				{
-					URigVMPin* RootPin = InTargetPin->GetRootPin();
-					if (!Template->FindArgument(RootPin->GetFName()))
+					const URigVMPin* RootPin = InTargetPin->GetRootPin();
+					if(!RootPin->IsOrphanPin())
 					{
-						if(!RootPin->IsExecuteContext())
+						if (!Template->FindArgument(RootPin->GetFName()))
 						{
-							if (OutFailureReason)
+							if(!RootPin->IsExecuteContext())
 							{
-								*OutFailureReason = FString::Printf(TEXT("Library pin %s supported types need to be reduced."), *RootPin->GetPinPath(true));
+								if (OutFailureReason)
+								{
+									*OutFailureReason = FString::Printf(TEXT("Library pin %s supported types need to be reduced."), *RootPin->GetPinPath(true));
+								}
+								return false;
 							}
-							return false;
 						}
 					}
 				}
@@ -1892,7 +1898,7 @@ bool URigVMPin::CanLink(URigVMPin* InSourcePin, URigVMPin* InTargetPin, FString*
 	}
 
 	// only allow to link to specified input / output pins on an injected node
-	if (URigVMInjectionInfo* SourceInjectionInfo = SourceNode->GetInjectionInfo())
+	if (const URigVMInjectionInfo* SourceInjectionInfo = SourceNode->GetInjectionInfo())
 	{
 		if (SourceInjectionInfo->OutputPin != InSourcePin->GetRootPin())
 		{
@@ -1905,7 +1911,7 @@ bool URigVMPin::CanLink(URigVMPin* InSourcePin, URigVMPin* InTargetPin, FString*
 	}
 
 	// only allow to link to specified input / output pins on an injected node
-	if (URigVMInjectionInfo* TargetInjectionInfo = TargetNode->GetInjectionInfo())
+	if (const URigVMInjectionInfo* TargetInjectionInfo = TargetNode->GetInjectionInfo())
 	{
 		if (TargetInjectionInfo->InputPin != InTargetPin->GetRootPin())
 		{
diff --git a/Engine/Source/Developer/RigVMDeveloper/Public/RigVMModel/RigVMClient.h b/Engine/Source/Developer/RigVMDeveloper/Public/RigVMModel/RigVMClient.h
--- a/Engine/Source/Developer/RigVMDeveloper/Public/RigVMModel/RigVMClient.h
+++ b/Engine/Source/Developer/RigVMDeveloper/Public/RigVMModel/RigVMClient.h
@@ -103,6 +103,7 @@ public:
 	URigVMGraph* GetModel(const FString& InNodePathOrName) const;
 	URigVMGraph* GetModel(const UObject* InEditorSideObject) const;
 	TArray<URigVMGraph*> GetAllModels(bool bIncludeFunctionLibrary, bool bRecursive) const;
+	TArray<URigVMGraph*> GetAllModelsLeavesFirst(bool bIncludeFunctionLibrary) const;
 	URigVMController* GetController(int32 InIndex) const;
 	URigVMController* GetController(const FString& InNodePathOrName) const;
 	URigVMController* GetController(const URigVMGraph* InModel) const;
@@ -145,6 +146,9 @@ public:
 	// backwards compatibility
 	FRigVMClientPatchResult PatchModelsOnLoad();
 
+	// try to reattach detached links and delete remaining ones
+	void ProcessDetachedLinks();
+
 	// work to be done after a duplication of the source asset
 	void PostDuplicateHost(const FString& InOldPathName, const FString& InNewPathName);
 
diff --git a/Engine/Source/Developer/RigVMDeveloper/Public/RigVMModel/RigVMController.h b/Engine/Source/Developer/RigVMDeveloper/Public/RigVMModel/RigVMController.h
--- a/Engine/Source/Developer/RigVMDeveloper/Public/RigVMModel/RigVMController.h
+++ b/Engine/Source/Developer/RigVMDeveloper/Public/RigVMModel/RigVMController.h
@@ -19,6 +19,7 @@
 #include "RigVMModel/RigVMBuildData.h"
 #include "RigVMCore/RigVMUserWorkflow.h"
 #include "UObject/Interface.h"
+#include "Algo/Sort.h"
 #include "RigVMController.generated.h"
 
 class URigVMActionStack;
@@ -88,7 +89,7 @@ DECLARE_DELEGATE_FiveParams(FRigVMController_OnBulkEditProgressDelegate, TSoftOb
 DECLARE_DELEGATE_RetVal_TwoParams(FString, FRigVMController_PinPathRemapDelegate, const FString& /* InPinPath */, bool /* bIsInput */);
 DECLARE_DELEGATE_OneParam(FRigVMController_RequestJumpToHyperlinkDelegate, const UObject* InSubject);
 DECLARE_DELEGATE_OneParam(FRigVMController_ConfigureWorkflowOptionsDelegate, URigVMUserWorkflowOptions* InOutOptions);
-DECLARE_DELEGATE_RetVal_TwoParams(bool, FRigVMController_CheckPinComatibilityDelegate, URigVMPin*, URigVMPin*);
+DECLARE_DELEGATE_RetVal_TwoParams(bool, FRigVMController_CheckPinCompatibilityDelegate, URigVMPin*, URigVMPin*);
 
 USTRUCT(BlueprintType)
 struct RIGVMDEVELOPER_API FRigStructScope
@@ -205,7 +206,7 @@ public:
 	// Pushes a new graph to the stack
 	// This causes a GraphChanged modified event.
 	UFUNCTION(BlueprintCallable, Category = RigVMController)
-	void PushGraph(URigVMGraph* InGraph, bool bSetupUndoRedo = true);
+	bool PushGraph(URigVMGraph* InGraph, bool bSetupUndoRedo = true);
 
 	// Pops the last graph off the stack
 	// This causes a GraphChanged modified event.
@@ -494,11 +495,11 @@ public:
 
 	// Exports the given nodes as text
 	UFUNCTION(BlueprintCallable, Category = RigVMController)
-	FString ExportNodesToText(const TArray<FName>& InNodeNames);
+	FString ExportNodesToText(const TArray<FName>& InNodeNames, bool bIncludeExteriorLinks = false);
 
 	// Exports the selected nodes as text
 	UFUNCTION(BlueprintCallable, Category = RigVMController)
-	FString ExportSelectedNodesToText();
+	FString ExportSelectedNodesToText(bool bIncludeExteriorLinks = false);
 
 	// Exports the given nodes as text
 	UFUNCTION(BlueprintCallable, Category = RigVMController)
@@ -552,13 +553,23 @@ public:
 	// Removes a node from the graph
 	// This causes a NodeRemoved modified event.
 	UFUNCTION(BlueprintCallable, Category = RigVMController)
-	bool RemoveNode(URigVMNode* InNode, bool bSetupUndoRedo = true, bool bRecursive = false, bool bPrintPythonCommand = false, bool bRelinkPins = false);
+	bool RemoveNode(URigVMNode* InNode, bool bSetupUndoRedo = true, bool bPrintPythonCommand = false);
 
 	// Removes a node from the graph given the node's name.
 	// This causes a NodeRemoved modified event.
 	UFUNCTION(BlueprintCallable, Category = RigVMController)
-	bool RemoveNodeByName(const FName& InNodeName, bool bSetupUndoRedo = true, bool bRecursive = false, bool bPrintPythonCommand = false, bool bRelinkPins = false);
-	
+	bool RemoveNodeByName(const FName& InNodeName, bool bSetupUndoRedo = true, bool bPrintPythonCommand = false);
+
+	// Removes a list of nodes from the graph
+	// This causes a NodeRemoved modified event.
+	UFUNCTION(BlueprintCallable, Category = RigVMController)
+	bool RemoveNodes(TArray<URigVMNode*> InNodes, bool bSetupUndoRedo = true, bool bPrintPythonCommand = false);
+
+	// Removes a list of nodes from the graph given the names
+	// This causes a NodeRemoved modified event.
+	UFUNCTION(BlueprintCallable, Category = RigVMController)
+	bool RemoveNodesByName(const TArray<FName>& InNodeNames, bool bSetupUndoRedo = true, bool bPrintPythonCommand = false);
+
 	// Renames a node in the graph
 	// This causes a NodeRenamed modified event.
 	UFUNCTION(BlueprintCallable, Category = RigVMController)
@@ -908,9 +919,8 @@ public:
 	// A delegate to request to configure an options instance for a node workflow
 	FRigVMController_ConfigureWorkflowOptionsDelegate ConfigureWorkflowOptionsDelegate; 
 
-	int32 DetachLinksFromPinObjects(const TArray<URigVMLink*>* InLinks = nullptr);
-	int32 ReattachLinksToPinObjects(bool bFollowCoreRedirectors = false, const TArray<URigVMLink*>* InLinks = nullptr, bool bSetupOrphanedPins = false, bool bAllowNonArgumentLinks = false, bool bRecursive = true);
 	void AddPinRedirector(bool bInput, bool bOutput, const FString& OldPinPath, const FString& NewPinPath);
+	void ClearPinRedirectors();
 
 	// Removes nodes which went stale.
 	void RemoveStaleNodes();
@@ -919,11 +929,12 @@ public:
 	bool ShouldRedirectPin(UScriptStruct* InOwningStruct, const FString& InOldRelativePinPath, FString& InOutNewRelativePinPath) const;
 	bool ShouldRedirectPin(const FString& InOldPinPath, FString& InOutNewPinPath) const;
 
-	void RepopulatePinsOnNode(URigVMNode* InNode, bool bFollowCoreRedirectors = true, bool bSetupOrphanedPins = false, bool bDetachAndReattachLinks = false);
+	void RepopulatePinsOnNode(URigVMNode* InNode, bool bFollowCoreRedirectors = true, bool bSetupOrphanedPins = false, bool bRecreateLinks = false);
 	void RemovePinsDuringRepopulate(URigVMNode* InNode, TArray<URigVMPin*>& InPins, bool bSetupOrphanedPins);
 
 	// removes any orphan pins that no longer holds a link
-	bool RemoveUnusedOrphanedPins(URigVMNode* InNode);
+	// @param bRelayLinks If true we'll try to relay the links back to their original pins 
+	bool RemoveUnusedOrphanedPins(URigVMNode* InNode, bool bRelayLinks = false);
 
 	// Update filtered permutations, and propagate both ways of the link before adding this link
 	bool PrepareToLink(URigVMPin* FirstToResolve, URigVMPin* SecondToResolve, bool bSetupUndoRedo);
@@ -1063,9 +1074,9 @@ private:
 	FString GetValidNodeName(const FString& InPrefix);
 	bool IsValidGraph() const;
 	bool IsGraphEditable() const;
-	bool IsValidNodeForGraph(URigVMNode* InNode);
-	bool IsValidPinForGraph(URigVMPin* InPin);
-	bool IsValidLinkForGraph(URigVMLink* InLink);
+	bool IsValidNodeForGraph(const URigVMNode* InNode);
+	bool IsValidPinForGraph(const URigVMPin* InPin);
+	bool IsValidLinkForGraph(const URigVMLink* InLink);
 	bool CanAddNode(URigVMNode* InNode, bool bReportErrors, bool bIgnoreFunctionEntryReturnNodes = false);
 	TObjectPtr<URigVMNode> FindEventNode(const UScriptStruct* InScriptStruct) const;
 	bool CanAddEventNode(UScriptStruct* InScriptStruct, const bool bReportErrors) const;
@@ -1083,7 +1094,7 @@ private:
 	static FString GetPinInitialDefaultValue(const URigVMPin* InPin);
 	static FString GetPinInitialDefaultValueFromStruct(UScriptStruct* ScriptStruct, const URigVMPin* InPin, uint32 InOffset);
 	URigVMPin* InsertArrayPin(URigVMPin* ArrayPin, int32 InIndex, const FString& InDefaultValue, bool bSetupUndoRedo);
-	bool RemovePin(URigVMPin* InPinToRemove, bool bSetupUndoRedo);
+	bool RemovePin(URigVMPin* InPinToRemove, bool bSetupUndoRedo, bool bForceBreakLinks = false);
 	FProperty* FindPropertyForPin(const FString& InPinPath);
 	bool BindPinToVariable(URigVMPin* InPin, const FString& InNewBoundVariablePath, bool bSetupUndoRedo, const FString& InVariableNodeName = FString());
 	bool UnbindPinFromVariable(URigVMPin* InPin, bool bSetupUndoRedo);
@@ -1094,11 +1105,12 @@ private:
 	URigVMNode* EjectNodeFromPin(URigVMPin* InPin, bool bSetupUndoRedo = true, bool bPrintPythonCommands = false);
 	bool EjectAllInjectedNodes(URigVMNode* InNode, bool bSetupUndoRedo = true, bool bPrintPythonCommands = false);
 
+public:
+
 	// try to reconnect source and target pins after a node deletion
 	void RelinkSourceAndTargetPins(URigVMNode* RigNode, bool bSetupUndoRedo = true);
 
-public:
-	bool AddLink(URigVMPin* OutputPin, URigVMPin* InputPin, bool bSetupUndoRedo = true, ERigVMPinDirection InUserDirection = ERigVMPinDirection::Invalid, bool bCreateCastNode = false);
+	bool AddLink(URigVMPin* OutputPin, URigVMPin* InputPin, bool bSetupUndoRedo = true, ERigVMPinDirection InUserDirection = ERigVMPinDirection::Invalid, bool bCreateCastNode = false, bool bIsRestoringLinks = false, FString* OutFailureReason = nullptr);
 	bool BreakLink(URigVMPin* OutputPin, URigVMPin* InputPin, bool bSetupUndoRedo = true);
 	bool BreakAllLinks(URigVMPin* Pin, bool bAsInput, bool bSetupUndoRedo = true);
 	void EnableTypeCasting(bool bEnabled = true) { bEnableTypeCasting = bEnabled; }
@@ -1120,7 +1132,7 @@ private:
 
 	void RefreshFunctionPins(URigVMNode* InNode);
 
-	void ReportRemovedLink(const FString& InSourcePinPath, const FString& InTargetPinPath);
+	void ReportRemovedLink(const FString& InSourcePinPath, const FString& InTargetPinPath, const FString& Reason = FString());
 
 	struct FPinState
 	{
@@ -1188,25 +1200,68 @@ public:
 	static FString GetSanitizedPinName(const FString& InName);
 	static FString GetSanitizedPinPath(const FString& InName);
 	static void SanitizeName(FString& InOutName, bool bAllowPeriod, bool bAllowSpace);
-	static TArray<TPair<FString, FString>> GetLinkedPinPaths(URigVMNode* InNode, bool bIncludeInjectionNodes = false);
-	static TArray<TPair<FString, FString>> GetLinkedPinPaths(const TArray<URigVMNode*>& InNodes, bool bIncludeInjectionNodes = false);
-	bool BreakLinkedPaths(const TArray<TPair<FString, FString>>& InLinkedPaths, bool bSetupUndoRedo);
-	bool RestoreLinkedPaths(
-		const TArray<TPair<FString, FString>>& InLinkedPaths,
-		const TMap<FString, FString>& InNodeNameMap,
-		const TMap<FString,FRigVMController_PinPathRemapDelegate>& InRemapDelegates,
-		FRigVMController_CheckPinComatibilityDelegate InCompatibilityDelegate,
-		bool bSetupUndoRedo,
-		ERigVMPinDirection InUserDirection = ERigVMPinDirection::Invalid);
-	bool RestoreLinkedPaths(
-		const TArray<TPair<FString, FString>>& InLinkedPaths,
-		const TMap<FString, FString>& InNodeNameMap,
-		const TMap<FString,FRigVMController_PinPathRemapDelegate>& InRemapDelegates,
-		bool bSetupUndoRedo,
-		ERigVMPinDirection InUserDirection = ERigVMPinDirection::Invalid)
+
+	struct FLinkedPath
 	{
-		return RestoreLinkedPaths(InLinkedPaths, InNodeNameMap, InRemapDelegates, FRigVMController_CheckPinComatibilityDelegate(), bSetupUndoRedo, InUserDirection);
-	}
+		FLinkedPath()
+			: bSourceNodeIsInjected(false)
+			, bTargetNodeIsInjected(false)
+		{}
+
+		FLinkedPath(URigVMLink* InLink);
+
+		URigVMGraph* GetGraph(URigVMGraph* InGraph = nullptr) const;
+		FString GetPinPathRepresentation() const;
+		URigVMPin* GetSourcePin(URigVMGraph* InGraph = nullptr) const;
+		URigVMPin* GetTargetPin(URigVMGraph* InGraph = nullptr) const;
+
+		friend uint32 GetTypeHash(const FLinkedPath& InPath);
+
+		bool operator ==(const FLinkedPath& InOther) const { return GetTypeHash(*this) == GetTypeHash(InOther); }
+		bool operator !=(const FLinkedPath& InOther) const { return !(*this == InOther); }
+		
+		TSoftObjectPtr<URigVMGraph> GraphPtr;
+		FString SourcePinPath;
+		FString TargetPinPath;
+		bool bSourceNodeIsInjected;
+		bool bTargetNodeIsInjected;
+	};
+
+	struct FRestoreLinkedPathSettings
+	{
+		FRestoreLinkedPathSettings()
+			: bFollowCoreRedirectors(false)
+			, bRelayToOrphanPins(false)
+			, bIsImportingFromText(false)
+			, UserDirection(ERigVMPinDirection::Invalid)
+		{}
+
+		bool bFollowCoreRedirectors;
+		bool bRelayToOrphanPins;
+		bool bIsImportingFromText;
+		ERigVMPinDirection UserDirection;
+		TMap<FString, FString> NodeNameMap;
+		TMap<FString,FRigVMController_PinPathRemapDelegate> RemapDelegates;
+		FRigVMController_CheckPinCompatibilityDelegate CompatibilityDelegate;
+	};
+
+	TArray<FLinkedPath> GetLinkedPaths() const;
+	static TArray<FLinkedPath> GetLinkedPaths(const TArray<URigVMLink*>& InLinks);
+	static TArray<FLinkedPath> GetLinkedPaths(URigVMNode* InNode, bool bIncludeInjectionNodes = false);
+	static TArray<FLinkedPath> GetLinkedPaths(const TArray<URigVMNode*>& InNodes, bool bIncludeInjectionNodes = false);
+	static TArray<FLinkedPath> GetLinkedPaths(const URigVMPin* InPin, bool bSourceLinksRecursive = false, bool bTargetLinksRecursive = false);
+	bool BreakLinkedPaths(const TArray<FLinkedPath>& InLinkedPaths, bool bSetupUndoRedo = false, bool bRelyOnBreakLink = true);
+	bool RestoreLinkedPaths(
+		const TArray<FLinkedPath>& InLinkedPaths,
+		const FRestoreLinkedPathSettings& InSettings = FRestoreLinkedPathSettings(),
+		bool bSetupUndoRedo = false);
+	TArray<FLinkedPath> RemapLinkedPaths(
+		const TArray<FLinkedPath>& InLinkedPaths,
+		const FRestoreLinkedPathSettings& InSettings = FRestoreLinkedPathSettings(),
+		bool bSetupUndoRedo = false);
+	bool FastBreakLinkedPaths(const TArray<FLinkedPath>& InLinkedPaths, bool bSetupUndoRedo = false);
+	URigVMLink* FindLinkFromPinPathRepresentation(const FString& InPinPathRepresentation, bool bLookForDetachedLink) const;
+	void ProcessDetachedLinks(const FRestoreLinkedPathSettings& InSettings = FRestoreLinkedPathSettings());
 
 #if WITH_EDITOR
 	// Registers this template node's use for later determining the commonly used types
@@ -1228,6 +1283,64 @@ protected:
 
 	// work to do after a duplication of the host asset
 	void PostDuplicateHost(const FString& InOldPathName, const FString& InNewPathName);
+
+	template<typename T>
+	static void SortGraphElementsByGraphDepth(TArray<T*>& InOutElements, bool bReverse = false)
+	{
+		if(InOutElements.IsEmpty())
+		{
+			return;
+		}
+
+		int32 MinDepth = InOutElements[0]->GetGraphDepth();
+		int32 MaxDepth = MinDepth;
+		int32 Increment = 1;
+		TMap<int32, TArray<T*>> ElementsPerDepth;
+		for(T* Element : InOutElements)
+		{
+			const int32 Depth = Element->GetGraphDepth(); 
+			ElementsPerDepth.FindOrAdd(Depth).Add(Element);
+			MinDepth = FMath::Min<int32>(MinDepth, Depth);
+			MaxDepth = FMath::Max<int32>(MaxDepth, Depth);
+		}
+
+		if(bReverse)
+		{
+			Swap(MinDepth, MaxDepth);
+			Increment = -1;
+		}
+		
+		InOutElements.Reset();
+		for(int32 Depth = MinDepth; /* no exit condition */; Depth += Increment)
+		{
+			if(const TArray<T*>* Elements = ElementsPerDepth.Find(Depth))
+			{
+				InOutElements.Append(*Elements);
+			}
+
+			if(Depth == MaxDepth)
+			{
+				break;
+			}
+		}
+	}
+
+	template<typename T>
+	static void SortGraphElementsByImportOrder(
+		TArray<TObjectPtr<T>> InOutElements,
+		const TArray<T*>& InElementsInImportOrder,
+		const TArray<T*>& InPreviousElementPriorToImport
+	) {
+		Algo::SortBy(InOutElements, [InElementsInImportOrder, InPreviousElementPriorToImport](T* Element) -> int32
+		{
+			const int32 ImportOrderIndex = InElementsInImportOrder.Find(Element);
+			if(ImportOrderIndex != INDEX_NONE)
+			{
+				return ImportOrderIndex + InPreviousElementPriorToImport.Num();
+			}
+			return InPreviousElementPriorToImport.Find(Element);
+		});
+	}
 	
 private: 
 	UPROPERTY(transient)
@@ -1306,6 +1419,7 @@ private:
 	friend struct FRigVMAddRerouteNodeAction;
 	friend struct FRigVMChangePinTypeAction;
 	friend struct FRigVMInjectNodeIntoPinAction;
+	friend struct FRigVMEjectNodeFromPinAction;
 	friend class FRigVMParserAST;
 	friend class FRigVMControllerCompileBracketScope;
 	friend class FRigVMControllerGraphGuard;
@@ -1321,13 +1435,22 @@ public:
 	FRigVMControllerGraphGuard(URigVMController* InController, URigVMGraph* InGraph, bool bSetupUndoRedo = true)
 		: Controller(InController)
 		, bUndo(bSetupUndoRedo)
+		, bEnabled(true)
 	{
-		Controller->PushGraph(InGraph, bUndo);
-		NumGraphs = Controller->Graphs.Num();
+		bEnabled = Controller->PushGraph(InGraph, bUndo);
+		if(bEnabled)
+		{
+			NumGraphs = Controller->Graphs.Num();
+		}
 	}
 
 	~FRigVMControllerGraphGuard()
 	{
+		if(!bEnabled)
+		{
+			return;
+		}
+		
 		// an action can be cancelled in the middle of a graph guard,
 		// in that case CancelAction should have already popped the graph
 		if (Controller->Graphs.Num() < NumGraphs)
@@ -1341,6 +1464,7 @@ private:
 
 	URigVMController* Controller;
 	bool bUndo;
+	bool bEnabled;
 
 	int32 NumGraphs;
 };
diff --git a/Engine/Source/Developer/RigVMDeveloper/Public/RigVMModel/RigVMControllerActions.h b/Engine/Source/Developer/RigVMDeveloper/Public/RigVMModel/RigVMControllerActions.h
--- a/Engine/Source/Developer/RigVMDeveloper/Public/RigVMModel/RigVMControllerActions.h
+++ b/Engine/Source/Developer/RigVMDeveloper/Public/RigVMModel/RigVMControllerActions.h
@@ -89,6 +89,23 @@ private:
 	TSharedPtr<FStructOnScope> StructOnScope;
 };
 
+/**
+ * A tuple 
+ * access to sub actions, merge functionality as well as undo and redo
+ * base implementations.
+ */
+USTRUCT()
+struct FRigVMActionNodeContent
+{
+	GENERATED_BODY()
+
+	UPROPERTY()
+	FString Old;
+
+	UPROPERTY()
+	FString New;
+};
+
 /**
  * The base action is the base struct for all actions, and provides
  * access to sub actions, merge functionality as well as undo and redo
@@ -128,6 +145,9 @@ struct FRigVMBaseAction
 	// returns true if successfull.
 	virtual bool Merge(const FRigVMBaseAction* Other);
 
+	// Returns true if this action makes the other action obsolete.
+	virtual bool MakesObsolete(const FRigVMBaseAction* Other) const;
+
 	// Un-does the action and returns true if successfull.
 	virtual bool Undo(URigVMController* InController);
 
@@ -143,6 +163,9 @@ struct FRigVMBaseAction
 		SubActions.Add(Key);
 	}
 
+	bool StoreNode(const URigVMNode* InNode, URigVMController* InController, bool bIsPriorChange = true);
+	bool RestoreNode(const FName& InNodeName, URigVMController* InController, bool bIsUndoing);
+
 #if RIGVM_ACTIONSTACK_VERBOSE_LOG
 	
 	// Logs the action to the console / output log
@@ -155,6 +178,9 @@ struct FRigVMBaseAction
 
 	UPROPERTY()
 	TArray<FRigVMActionKey> SubActions;
+
+	UPROPERTY()
+	TMap<FName, FRigVMActionNodeContent> ExportedNodes;
 };
 
 /**
@@ -252,28 +278,42 @@ public:
 			ActionList = &CurrentActions[CurrentActions.Num()-1]->SubActions;
 		}
 
-		bool bMergeIfPossible = bPerformMerge;
-		if (bMergeIfPossible)
+		bool bMergeIfPossible = false;
+		bool bIgnoreAction = false;
+		
+		if (ActionList->Num() > 0 && InAction.SubActions.Num() == 0)
 		{
-			bMergeIfPossible = false;
-			if (ActionList->Num() > 0 && InAction.SubActions.Num() == 0)
+			const FRigVMActionWrapper Wrapper((*ActionList)[ActionList->Num() - 1]);
+			if (Wrapper.GetAction()->SubActions.Num() == 0)
 			{
+				FRigVMBaseAction* OtherAction = Wrapper.GetAction();
+
 				if (ActionList->Last().ScriptStructPath == ActionType::StaticStruct()->GetPathName())
 				{
-					FRigVMActionWrapper Wrapper((*ActionList)[ActionList->Num() - 1]);
-					if (Wrapper.GetAction()->SubActions.Num() == 0)
+					if(bPerformMerge)
 					{
-						bMergeIfPossible = Wrapper.GetAction()->Merge(&InAction);
-						if (bMergeIfPossible)
+						bMergeIfPossible = OtherAction->Merge(&InAction);
+						if(bMergeIfPossible)
 						{
 							(*ActionList)[ActionList->Num()-1].ExportedText = Wrapper.ExportText();
 						}
 					}
 				}
+
+				if(!bMergeIfPossible)
+				{
+					if(InAction.MakesObsolete(OtherAction))
+					{
+						// if this new action cancels out the last one we don't
+						// need to perform either of the actions
+						ActionList->Pop();
+						bIgnoreAction = true;
+					}
+				}
 			}
 		}
 
-		if (!bMergeIfPossible)
+		if (!bMergeIfPossible && !bIgnoreAction)
 		{
 			FRigVMActionKey Key;
 			Key.Set<ActionType>(InAction);
@@ -363,275 +403,80 @@ private:
 	friend class URigVMController;
 };
 
-/**
- * An action which inverses the child actions,
- * it performs undo on redo and vice versa.
- */
-USTRUCT()
-struct FRigVMInverseAction : public FRigVMBaseAction
-{
-	GENERATED_BODY()
-
-	virtual ~FRigVMInverseAction() {};
-	virtual UScriptStruct* GetScriptStruct() const override { return FRigVMInverseAction::StaticStruct(); }
-
-	virtual bool Undo(URigVMController* InController);
-	virtual bool Redo(URigVMController* InController);
-};
-
-/**
- * An action adding a unit node to the graph.
- */
-USTRUCT()
-struct FRigVMAddUnitNodeAction : public FRigVMBaseAction
-{
-	GENERATED_BODY()
-
-public:
-
-	FRigVMAddUnitNodeAction();
-	FRigVMAddUnitNodeAction(URigVMUnitNode* InNode);
-	virtual ~FRigVMAddUnitNodeAction() {};
-	virtual UScriptStruct* GetScriptStruct() const override { return FRigVMAddUnitNodeAction::StaticStruct(); }
-	virtual bool Undo(URigVMController* InController) override;
-	virtual bool Redo(URigVMController* InController) override;
-
-	UPROPERTY()
-	FString ScriptStructPath;
-
-	UPROPERTY()
-	FName MethodName;
-	
-	UPROPERTY()
-	FVector2D Position;
-
-	UPROPERTY()
-	FString NodePath;
-};
-
-/**
- * An action adding a variable node to the graph.
- */
-USTRUCT()
-struct FRigVMAddVariableNodeAction : public FRigVMBaseAction
-{
-	GENERATED_BODY()
-
-public:
-
-	FRigVMAddVariableNodeAction();
-	FRigVMAddVariableNodeAction(URigVMVariableNode* InNode);
-	virtual ~FRigVMAddVariableNodeAction() {};
-	virtual UScriptStruct* GetScriptStruct() const override { return FRigVMAddVariableNodeAction::StaticStruct(); }
-	virtual bool Undo(URigVMController* InController) override;
-	virtual bool Redo(URigVMController* InController) override;
-
-	UPROPERTY()
-	FName VariableName;
-	
-	UPROPERTY()
-	FString CPPType;
-
-	UPROPERTY()
-	FString CPPTypeObjectPath;
-
-	UPROPERTY()
-	bool bIsGetter;
-
-	UPROPERTY()
-	FString DefaultValue;
-
-	UPROPERTY()
-	FVector2D Position;
-
-	UPROPERTY()
-	FString NodePath;
-};
-
-/**
- * An action adding a comment node to the graph.
- */
-USTRUCT()
-struct FRigVMAddCommentNodeAction : public FRigVMBaseAction
-{
-	GENERATED_BODY()
-
-public:
-
-	FRigVMAddCommentNodeAction();
-	FRigVMAddCommentNodeAction(URigVMCommentNode* InNode);
-	virtual ~FRigVMAddCommentNodeAction() {};
-	virtual UScriptStruct* GetScriptStruct() const override { return FRigVMAddCommentNodeAction::StaticStruct(); }
-	virtual bool Undo(URigVMController* InController) override;
-	virtual bool Redo(URigVMController* InController) override;
-
-	UPROPERTY()
-	FString CommentText;
-
-	UPROPERTY()
-	FVector2D Position;
-
-	UPROPERTY()
-	FVector2D Size;
-
-	UPROPERTY()
-	FLinearColor Color;
-
-	UPROPERTY()
-	FString NodePath;
-};
 
 /**
- * An action adding a reroute node to the graph.
+ * An action injecting a node into a pin
  */
 USTRUCT()
-struct FRigVMAddRerouteNodeAction : public FRigVMBaseAction
+struct FRigVMInjectNodeIntoPinAction : public FRigVMBaseAction
 {
 	GENERATED_BODY()
 
 public:
 
-	FRigVMAddRerouteNodeAction();
-	FRigVMAddRerouteNodeAction(URigVMRerouteNode* InNode);
-	virtual ~FRigVMAddRerouteNodeAction() {};
-	virtual UScriptStruct* GetScriptStruct() const override { return FRigVMAddRerouteNodeAction::StaticStruct(); }
+	FRigVMInjectNodeIntoPinAction();
+	FRigVMInjectNodeIntoPinAction(URigVMInjectionInfo* InInjectionInfo);
+	virtual ~FRigVMInjectNodeIntoPinAction() {};
+	virtual UScriptStruct* GetScriptStruct() const override { return FRigVMInjectNodeIntoPinAction::StaticStruct(); }
 	virtual bool Undo(URigVMController* InController) override;
 	virtual bool Redo(URigVMController* InController) override;
 
 	UPROPERTY()
-	bool bShowAsFullNode;
-
-	UPROPERTY()
-	FString CPPType;
-
-	UPROPERTY()
-	FName CPPTypeObjectPath;
-
-	UPROPERTY()
-	FString DefaultValue;
+	FString PinPath;
 
 	UPROPERTY()
-	bool bIsConstant;
+	bool bAsInput;
 
 	UPROPERTY()
-	FName CustomWidgetName;
+	FName InputPinName;
 
 	UPROPERTY()
-	FVector2D Position;
+	FName OutputPinName;
 
 	UPROPERTY()
 	FString NodePath;
 };
 
 /**
- * An action adding an enum node to the graph.
+ * An action ejecting a node from a pin
  */
 USTRUCT()
-struct FRigVMAddEnumNodeAction : public FRigVMBaseAction
+struct FRigVMEjectNodeFromPinAction : public FRigVMInjectNodeIntoPinAction
 {
 	GENERATED_BODY()
 
 public:
 
-	FRigVMAddEnumNodeAction();
-	FRigVMAddEnumNodeAction(URigVMEnumNode* InNode);
-	virtual ~FRigVMAddEnumNodeAction() {};
-	virtual UScriptStruct* GetScriptStruct() const override { return FRigVMAddEnumNodeAction::StaticStruct(); }
+	FRigVMEjectNodeFromPinAction();
+	FRigVMEjectNodeFromPinAction(URigVMInjectionInfo* InInjectionInfo);
+	virtual ~FRigVMEjectNodeFromPinAction() {};
+	virtual UScriptStruct* GetScriptStruct() const override { return FRigVMEjectNodeFromPinAction::StaticStruct(); }
 	virtual bool Undo(URigVMController* InController) override;
 	virtual bool Redo(URigVMController* InController) override;
-
-	UPROPERTY()
-	FString CPPType;
-
-	UPROPERTY()
-	FName CPPTypeObjectPath;
-
-	UPROPERTY()
-	FVector2D Position;
-
-	UPROPERTY()
-	FString NodePath;
 };
 
 /**
- * An action adding a template node to the graph.
+ * An action removing one or more nodes from the graph.
  */
 USTRUCT()
-struct FRigVMAddTemplateNodeAction : public FRigVMBaseAction
+struct FRigVMRemoveNodesAction : public FRigVMBaseAction
 {
 	GENERATED_BODY()
 
 public:
 
-	FRigVMAddTemplateNodeAction();
-	FRigVMAddTemplateNodeAction(URigVMTemplateNode* InNode);
-	virtual ~FRigVMAddTemplateNodeAction() {};
-	virtual UScriptStruct* GetScriptStruct() const override { return FRigVMAddTemplateNodeAction::StaticStruct(); }
+	FRigVMRemoveNodesAction() {}
+	FRigVMRemoveNodesAction(TArray<URigVMNode*> InNodes, URigVMController* InController);
+	virtual ~FRigVMRemoveNodesAction() {};
+	virtual UScriptStruct* GetScriptStruct() const override { return FRigVMRemoveNodesAction::StaticStruct(); }
 	virtual bool Undo(URigVMController* InController) override;
 	virtual bool Redo(URigVMController* InController) override;
 
 	UPROPERTY()
-	FName TemplateNotation;
-
-	UPROPERTY()
-	FVector2D Position;
-
-	UPROPERTY()
-	FString NodePath;
-};
-
-/**
- * An action injecting a node into a pin
- */
-USTRUCT()
-struct FRigVMInjectNodeIntoPinAction : public FRigVMBaseAction
-{
-	GENERATED_BODY()
-
-public:
-
-	FRigVMInjectNodeIntoPinAction();
-	FRigVMInjectNodeIntoPinAction(URigVMInjectionInfo* InInjectionInfo);
-	virtual ~FRigVMInjectNodeIntoPinAction() {};
-	virtual UScriptStruct* GetScriptStruct() const override { return FRigVMInjectNodeIntoPinAction::StaticStruct(); }
-	virtual bool Undo(URigVMController* InController) override;
-	virtual bool Redo(URigVMController* InController) override;
-
-	UPROPERTY()
-	FString PinPath;
-
-	UPROPERTY()
-	bool bAsInput;
-
-	UPROPERTY()
-	FName InputPinName;
-
-	UPROPERTY()
-	FName OutputPinName;
-
-	UPROPERTY()
-	FString NodePath;
-};
-
-/**
- * An action removing a node from the graph.
- */
-USTRUCT()
-struct FRigVMRemoveNodeAction : public FRigVMBaseAction
-{
-	GENERATED_BODY()
-
-public:
-
-	FRigVMRemoveNodeAction() {}
-	FRigVMRemoveNodeAction(URigVMNode* InNode, URigVMController* InController);
-	virtual ~FRigVMRemoveNodeAction() {};
-	virtual UScriptStruct* GetScriptStruct() const override { return FRigVMRemoveNodeAction::StaticStruct(); }
-	virtual bool Undo(URigVMController* InController) override;
-	virtual bool Redo(URigVMController* InController) override;
+	TArray<FName> NodeNames;
 
 	UPROPERTY()
-	FRigVMActionKey InverseActionKey;
+	FString ExportedContent;	
 };
 
 /**
@@ -1179,33 +1024,6 @@ public:
 	bool bRemoveSubPins;
 };
 
-/**
- * An action to add a node from a text buffer
- */
-USTRUCT()
-struct FRigVMImportNodeFromTextAction : public FRigVMBaseAction
-{
-	GENERATED_BODY()
-
-public:
-
-	FRigVMImportNodeFromTextAction();
-	FRigVMImportNodeFromTextAction(URigVMNode* InNode, URigVMController* InController);
-	virtual ~FRigVMImportNodeFromTextAction() {};
-	virtual UScriptStruct* GetScriptStruct() const override { return FRigVMImportNodeFromTextAction::StaticStruct(); }
-	virtual bool Undo(URigVMController* InController) override;
-	virtual bool Redo(URigVMController* InController) override;
-
-	UPROPERTY()
-	FVector2D Position;
-
-	UPROPERTY()
-	FString NodePath;
-
-	UPROPERTY()
-	FString ExportedText;
-};
-
 /**
  * An action to collapse a selection of nodes
  */
@@ -1232,9 +1050,6 @@ public:
 	UPROPERTY()
 	TArray<FString> CollapsedNodesPaths;
 
-	UPROPERTY()
-	TArray<FString> CollapsedNodesLinks;
-
 	UPROPERTY()
 	bool bIsAggregate;
 };
@@ -1262,9 +1077,6 @@ public:
 	UPROPERTY()
 	FString LibraryNodeContent;
 
-	UPROPERTY()
-	TArray<FString> LibraryNodeLinks;
-
 	UPROPERTY()
 	TArray<FString> ExpandedNodePaths;
 };
@@ -1309,6 +1121,7 @@ public:
 	virtual UScriptStruct* GetScriptStruct() const override { return FRigVMPushGraphAction::StaticStruct(); }
 	virtual bool Undo(URigVMController* InController) override;
 	virtual bool Redo(URigVMController* InController) override;
+	virtual bool MakesObsolete(const FRigVMBaseAction* Other) const override;
 
 	UPROPERTY()
 	FSoftObjectPath GraphPath;
@@ -1330,6 +1143,7 @@ public:
 	virtual UScriptStruct* GetScriptStruct() const override { return FRigVMPopGraphAction::StaticStruct(); }
 	virtual bool Undo(URigVMController* InController) override;
 	virtual bool Redo(URigVMController* InController) override;
+	virtual bool MakesObsolete(const FRigVMBaseAction* Other) const override;
 
 	UPROPERTY()
 	FSoftObjectPath GraphPath;
@@ -1630,52 +1444,70 @@ public:
 };
 
 /**
- * An action adding an invoke entry node to the graph.
+ * An action marking a function as public/private.
  */
 USTRUCT()
-struct FRigVMAddInvokeEntryNodeAction : public FRigVMBaseAction
+struct FRigVMMarkFunctionPublicAction : public FRigVMBaseAction
 {
 	GENERATED_BODY()
 
 public:
 
-	FRigVMAddInvokeEntryNodeAction();
-	FRigVMAddInvokeEntryNodeAction(URigVMInvokeEntryNode* InNode);
-	virtual ~FRigVMAddInvokeEntryNodeAction() {};
-	virtual UScriptStruct* GetScriptStruct() const override { return FRigVMAddInvokeEntryNodeAction::StaticStruct(); }
+	FRigVMMarkFunctionPublicAction();
+	FRigVMMarkFunctionPublicAction(const FName& InFunctionName, bool bInIsPublic);
+	virtual ~FRigVMMarkFunctionPublicAction() {};
+	virtual UScriptStruct* GetScriptStruct() const override { return FRigVMMarkFunctionPublicAction::StaticStruct(); }
 	virtual bool Undo(URigVMController* InController) override;
 	virtual bool Redo(URigVMController* InController) override;
 
 	UPROPERTY()
-	FName EntryName;
+	FName FunctionName;
 	
 	UPROPERTY()
-	FVector2D Position;
-
-	UPROPERTY()
-	FString NodePath;
+	bool bIsPublic;
 };
 
 /**
- * An action marking a function as public/private.
+ * An action importing nodes and links from text
  */
 USTRUCT()
-struct FRigVMMarkFunctionPublicAction : public FRigVMBaseAction
+struct FRigVMImportFromTextAction : public FRigVMBaseAction
 {
 	GENERATED_BODY()
 
 public:
 
-	FRigVMMarkFunctionPublicAction();
-	FRigVMMarkFunctionPublicAction(const FName& InFunctionName, bool bInIsPublic);
-	virtual ~FRigVMMarkFunctionPublicAction() {};
-	virtual UScriptStruct* GetScriptStruct() const override { return FRigVMMarkFunctionPublicAction::StaticStruct(); }
+	FRigVMImportFromTextAction() {}
+	FRigVMImportFromTextAction(const FString& InContent, const TArray<FName>& InTopLevelNodeNames);
+	FRigVMImportFromTextAction(URigVMNode* InNode, URigVMController* InController, bool bIncludeExteriorLinks = true);
+	FRigVMImportFromTextAction(const TArray<URigVMNode*>& InNodes, URigVMController* InController, bool bIncludeExteriorLinks = true);
+	void SetContent(const TArray<URigVMNode*>& InNodes, URigVMController* InController, bool bIncludeExteriorLinks = true);
+	virtual ~FRigVMImportFromTextAction() {};
+	virtual UScriptStruct* GetScriptStruct() const override { return FRigVMImportFromTextAction::StaticStruct(); }
 	virtual bool Undo(URigVMController* InController) override;
 	virtual bool Redo(URigVMController* InController) override;
 
 	UPROPERTY()
-	FName FunctionName;
-	
+	FString Content;
+
 	UPROPERTY()
-	bool bIsPublic;
+	TArray<FName> TopLevelNodeNames;
+};
+
+/**
+ * An action to add store / restore a single node
+ */
+USTRUCT()
+struct FRigVMReplaceNodesAction : public FRigVMBaseAction
+{
+	GENERATED_BODY()
+
+public:
+
+	FRigVMReplaceNodesAction() {}
+	FRigVMReplaceNodesAction(const TArray<URigVMNode*>& InNodes, URigVMController* InController);
+	virtual ~FRigVMReplaceNodesAction() {};
+	virtual UScriptStruct* GetScriptStruct() const override { return FRigVMReplaceNodesAction::StaticStruct(); }
+	virtual bool Undo(URigVMController* InController) override;
+	virtual bool Redo(URigVMController* InController) override;
 };
diff --git a/Engine/Source/Developer/RigVMDeveloper/Public/RigVMModel/RigVMGraph.h b/Engine/Source/Developer/RigVMDeveloper/Public/RigVMModel/RigVMGraph.h
--- a/Engine/Source/Developer/RigVMDeveloper/Public/RigVMModel/RigVMGraph.h
+++ b/Engine/Source/Developer/RigVMDeveloper/Public/RigVMModel/RigVMGraph.h
@@ -34,6 +34,8 @@ public:
 	// Default constructor
 	URigVMGraph();
 
+	virtual void PostLoad() override;
+
 	// Returns all of the Nodes within this Graph.
 	UFUNCTION(BlueprintCallable, Category = RigVMGraph)
 	const TArray<URigVMNode*>& GetNodes() const;
@@ -42,6 +44,10 @@ public:
 	UFUNCTION(BlueprintCallable, Category = RigVMGraph)
 	const TArray<URigVMLink*>& GetLinks() const;
 
+	// Returns true if the graph contains a link given its string representation 
+	UFUNCTION(BlueprintCallable, Category = RigVMGraph)
+	bool ContainsLink(const FString& InPinPathRepresentation) const;
+
 	// Returns all of the contained graphs
 	UFUNCTION(BlueprintCallable, Category = RigVMGraph)
 	TArray<URigVMGraph*> GetContainedGraphs(bool bRecursive = false) const;
@@ -54,6 +60,10 @@ public:
 	UFUNCTION(BlueprintCallable, Category = RigVMGraph)
     URigVMGraph* GetRootGraph() const;
 
+	// Returns the root / top level parent graph of this graph (or this if it is the root graph)
+	UFUNCTION(BlueprintPure, Category = RigVMGraph)
+	int32 GetGraphDepth() const;
+
 	// Returns true if this graph is a root / top level graph
 	UFUNCTION(BlueprintCallable, Category = RigVMGraph)
     bool IsRootGraph() const;
@@ -162,6 +172,9 @@ private:
 	UPROPERTY()
 	TArray<TObjectPtr<URigVMLink>> Links;
 
+	UPROPERTY(transient)
+	TArray<TObjectPtr<URigVMLink>> DetachedLinks;
+
 	UPROPERTY()
 	TArray<FName> SelectedNodes;
 
@@ -188,7 +201,7 @@ private:
 	UPROPERTY()
 	TArray<FRigVMGraphVariableDescription> LocalVariables;
 
-	bool IsNameAvailable(const FString& InName);
+	bool IsNameAvailable(const FString& InName) const;
 
 	friend class URigVMController;
 	friend class UControlRigBlueprint;
diff --git a/Engine/Source/Developer/RigVMDeveloper/Public/RigVMModel/RigVMLink.h b/Engine/Source/Developer/RigVMDeveloper/Public/RigVMModel/RigVMLink.h
--- a/Engine/Source/Developer/RigVMDeveloper/Public/RigVMModel/RigVMLink.h
+++ b/Engine/Source/Developer/RigVMDeveloper/Public/RigVMModel/RigVMLink.h
@@ -39,6 +39,10 @@ public:
 	UFUNCTION(BlueprintCallable, Category = RigVMLink)
 	URigVMGraph* GetGraph() const;
 
+	// Returns the graph nesting depth of this link
+	UFUNCTION(BlueprintCallable, Category = RigVMLink)
+	int32 GetGraphDepth() const;
+
 	// Returns the source Pin of this Link (or nullptr)
 	UFUNCTION(BlueprintCallable, Category = RigVMLink)
 	URigVMPin* GetSourcePin() const;
@@ -47,6 +51,29 @@ public:
 	UFUNCTION(BlueprintCallable, Category = RigVMLink)
 	URigVMPin* GetTargetPin() const;
 
+	// Returns the source Node of this Link (or nullptr)
+	UFUNCTION(BlueprintCallable, Category = RigVMLink)
+	URigVMNode* GetSourceNode() const;
+
+	// Returns the target Node of this Link (or nullptr)
+	UFUNCTION(BlueprintCallable, Category = RigVMLink)
+	URigVMNode* GetTargetNode() const;
+
+	// Returns the source pin's path pin of this Link
+	FString GetSourcePinPath() const;
+
+	// Returns the target pin's path pin of this Link
+	FString GetTargetPinPath() const;
+
+	// Sets the source pin's path pin of this Link
+	bool SetSourcePinPath(const FString& InPinPath);
+
+	// Sets the target pin's path pin of this Link
+	bool SetTargetPinPath(const FString& InPinPath);
+
+	// Sets the target pin's path pin of this Link
+	bool SetSourceAndTargetPinPaths(const FString& InSourcePinPath, const FString& InTargetPinPath);
+
 	// Returns the opposite Pin of this Link given one of its edges (or nullptr)
 	UFUNCTION(BlueprintCallable, Category = RigVMLink)
 	URigVMPin* GetOppositePin(const URigVMPin* InPin) const;
@@ -55,21 +82,29 @@ public:
 	// for example: "NodeA.Color.R -> NodeB.Translation.X"
 	// note: can be split again using SplitPinPathRepresentation
 	UFUNCTION(BlueprintCallable, Category = RigVMLink)
-	FString GetPinPathRepresentation();
+	FString GetPinPathRepresentation() const;
 
 	// Returns a string representation of the Link given the two pin paths
 	// for example: "NodeA.Color.R -> NodeB.Translation.X"
 	// note: can be split again using SplitPinPathRepresentation
 	static FString GetPinPathRepresentation(const FString& InSourcePinPath, const FString& InTargetPinPath);
 
-	// Splits a pin path represenation of a link
+	// Splits a pin path representation of a link
 	// for example: "NodeA.Color.R -> NodeB.Translation.X"
 	// into its two pin paths
 	static bool SplitPinPathRepresentation(const FString& InString, FString& OutSource, FString& OutTarget);
 
+	void UpdatePinPaths();
+	void UpdatePinPointers() const;
+	bool Detach();
+	
 private:
 
-	void PrepareForCopy();
+	// Returns true if the link is attached.
+	// Attached links rely on the pin pointers first and the pin path second.
+	// Deattached links never rely on the pin pointers and always try to resolve from string.
+	bool IsAttached() const;
+	bool Attach(FString* OutFailureReason = nullptr);
 
 	UPROPERTY()
 	FString SourcePinPath;
@@ -79,7 +114,5 @@ private:
 
 	mutable URigVMPin* SourcePin;
 	mutable URigVMPin* TargetPin;
-
-	friend class URigVMController;
 };
 
diff --git a/Engine/Source/Developer/RigVMDeveloper/Public/RigVMModel/RigVMNode.h b/Engine/Source/Developer/RigVMDeveloper/Public/RigVMModel/RigVMNode.h
--- a/Engine/Source/Developer/RigVMDeveloper/Public/RigVMModel/RigVMNode.h
+++ b/Engine/Source/Developer/RigVMDeveloper/Public/RigVMModel/RigVMNode.h
@@ -86,6 +86,10 @@ public:
 	UFUNCTION(BlueprintCallable, Category = RigVMNode)
     URigVMGraph* GetRootGraph() const;
 
+	// Returns the graph nesting depth of this node
+	UFUNCTION(BlueprintCallable, Category = RigVMNode)
+	int32 GetGraphDepth() const;
+
 	// Returns the injection info of this Node (or nullptr)
 	UFUNCTION(BlueprintCallable, Category = RigVMNode)
 	URigVMInjectionInfo* GetInjectionInfo() const;
diff --git a/Engine/Source/Developer/RigVMDeveloper/Public/RigVMModel/RigVMPin.h b/Engine/Source/Developer/RigVMDeveloper/Public/RigVMModel/RigVMPin.h
--- a/Engine/Source/Developer/RigVMDeveloper/Public/RigVMModel/RigVMPin.h
+++ b/Engine/Source/Developer/RigVMDeveloper/Public/RigVMModel/RigVMPin.h
@@ -359,7 +359,7 @@ public:
 
 	// Returns true if this Pin is linked to another Pin
 	UFUNCTION(BlueprintCallable, Category = RigVMPin)
-	bool IsLinkedTo(URigVMPin* InPin) const;
+	bool IsLinkedTo(const URigVMPin* InPin) const;
 
 	// Returns true if the pin has any link
 	bool IsLinked(bool bRecursive = false) const;
@@ -531,6 +531,7 @@ private:
 	friend class UControlRigBlueprint;
 	friend class URigVMGraph;
 	friend class URigVMNode;
+	friend class URigVMLink;
 	friend class FRigVMParserAST;
 };
 
diff --git a/Engine/Source/Editor/UnrealEd/Private/Factories/EditorFactories.cpp b/Engine/Source/Editor/UnrealEd/Private/Factories/EditorFactories.cpp
--- a/Engine/Source/Editor/UnrealEd/Private/Factories/EditorFactories.cpp
+++ b/Engine/Source/Editor/UnrealEd/Private/Factories/EditorFactories.cpp
@@ -5537,7 +5537,7 @@ void FCustomizableTextObjectFactory::ProcessBuffer(UObject* InParent, EObjectFla
 				UObject* ObjectParent = InParent ? InParent : GetParentForNewObject(ObjClass);
 
 				// Make sure this name is not used by anything else. Will rename other stuff if necessary
-				UpdateObjectName(ObjClass, ObjName);
+				UpdateObjectName(ObjClass, ObjectParent, ObjName);
 				ClearObjectNameUsage(ObjectParent, ObjName);
 
 				// Spawn the object and reset it's archetype
diff --git a/Engine/Source/Editor/UnrealEd/Public/Factories.h b/Engine/Source/Editor/UnrealEd/Public/Factories.h
--- a/Engine/Source/Editor/UnrealEd/Public/Factories.h
+++ b/Engine/Source/Editor/UnrealEd/Public/Factories.h
@@ -57,7 +57,7 @@ protected:
 	virtual void PostProcessConstructedObjects() {};
 
 	/** Provide the opportunity to change the name during import */
-	virtual void UpdateObjectName(UClass* ObjectClass, FName& InOutObjName) {};
+	virtual void UpdateObjectName(UClass* ObjectClass, UObject* InParent, FName& InOutObjName) {};
 
 	/** Util to ensure that InName is a valid name for a new object within InParent. Will rename any existing object within InParent if it is called InName. */
 	static void ClearObjectNameUsage(UObject* InParent, FName InName);
