Control Rig: Fix undo function creation failing and crashing the editor

diff --git a/Engine/Plugins/Animation/ControlRig/Source/ControlRigEditor/Private/ControlRigLocalVariableDetails.cpp b/Engine/Plugins/Animation/ControlRig/Source/ControlRigEditor/Private/ControlRigLocalVariableDetails.cpp
--- a/Engine/Plugins/Animation/ControlRig/Source/ControlRigEditor/Private/ControlRigLocalVariableDetails.cpp
+++ b/Engine/Plugins/Animation/ControlRig/Source/ControlRigEditor/Private/ControlRigLocalVariableDetails.cpp
@@ -32,7 +32,10 @@ void FRigVMLocalVariableDetails::CustomizeDetails(IDetailLayoutBuilder& DetailBu
 	{
 		VariableDescription = ObjectsBeingCustomized[0]->GetContent<FRigVMGraphVariableDescription>();
 		GraphBeingCustomized = ObjectsBeingCustomized[0]->GetTypedOuter<URigVMGraph>();
-		BlueprintBeingCustomized = GraphBeingCustomized->GetTypedOuter<UControlRigBlueprint>();
+		if (GraphBeingCustomized)
+		{
+			BlueprintBeingCustomized = GraphBeingCustomized->GetTypedOuter<UControlRigBlueprint>();
+		}
 
 		NameValidator = FControlRigLocalVariableNameValidator(BlueprintBeingCustomized, GraphBeingCustomized, VariableDescription.Name);
 	}
diff --git a/Engine/Source/Developer/RigVMDeveloper/Private/RigVMModel/RigVMController.cpp b/Engine/Source/Developer/RigVMDeveloper/Private/RigVMModel/RigVMController.cpp
--- a/Engine/Source/Developer/RigVMDeveloper/Private/RigVMModel/RigVMController.cpp
+++ b/Engine/Source/Developer/RigVMDeveloper/Private/RigVMModel/RigVMController.cpp
@@ -193,7 +193,8 @@ bool URigVMController::PushGraph(URigVMGraph* InGraph, bool bSetupUndoRedo)
 
 URigVMGraph* URigVMController::PopGraph(bool bSetupUndoRedo)
 {
-	if(Graphs.Num() == 1)
+	ensure(Graphs.Num() > 1);
+	if (Graphs.Num() == 1)
 	{
 		return nullptr;
 	}
@@ -11316,13 +11317,13 @@ URigVMLibraryNode* URigVMController::AddFunctionToLibrary(const FName& InFunctio
 	{
 		const UScriptStruct* ExecuteContextStruct = FRigVMExecuteContext::StaticStruct();
 		
-		FRigVMControllerGraphGuard GraphGuard(this, CollapseNode->ContainedGraph, bSetupUndoRedo);
+		FRigVMControllerGraphGuard GraphGuard(this, CollapseNode->ContainedGraph, false);
 		AddExposedPin(FRigVMStruct::ExecuteContextName,
 			ERigVMPinDirection::IO,
 			FString::Printf(TEXT("F%s"), *ExecuteContextStruct->GetName()),
 			*ExecuteContextStruct->GetPathName(),
 			FString(),
-			bSetupUndoRedo);
+			false);
 	}
 
 	{
