RigVM: Move action stack owner ship from controller to client

diff --git a/Engine/Plugins/Animation/ControlRig/Source/ControlRigDeveloper/Private/ControlRigBlueprint.cpp b/Engine/Plugins/Animation/ControlRig/Source/ControlRigDeveloper/Private/ControlRigBlueprint.cpp
--- a/Engine/Plugins/Animation/ControlRig/Source/ControlRigDeveloper/Private/ControlRigBlueprint.cpp
+++ b/Engine/Plugins/Animation/ControlRig/Source/ControlRigDeveloper/Private/ControlRigBlueprint.cpp
@@ -143,7 +143,7 @@ UControlRigBlueprint::UControlRigBlueprint(const FObjectInitializer& ObjectIniti
 	{
 		TGuardValue<bool> DisableClientNotifs(RigVMClient.bSuspendNotifications, true);
 		RigVMClient.GetOrCreateFunctionLibrary(false, &ObjectInitializer, false);
-		RigVMClient.AddModel(RigVMModelPrefix, false, &ObjectInitializer, false);
+		RigVMClient.AddModel(FRigVMClient::RigVMModelPrefix, false, &ObjectInitializer, false);
 	}
 	RigVMClient.SetExecuteContextStruct(FControlRigExecuteContext::StaticStruct());
 	
@@ -1984,7 +1984,7 @@ URigVMFunctionLibrary* UControlRigBlueprint::GetLocalFunctionLibrary() const
 URigVMGraph* UControlRigBlueprint::AddModel(FString InName, bool bSetupUndoRedo, bool bPrintPythonCommand)
 {
 	const FString DesiredName = FString::Printf(TEXT("%s %s"),
-    	UControlRigBlueprint::RigVMModelPrefix, *InName);
+    	FRigVMClient::RigVMModelPrefix, *InName);
 
 	TGuardValue<bool> EnablePythonPrint(bSuspendPythonMessagesForRigVMClient, !bPrintPythonCommand);
 	return RigVMClient.AddModel(*DesiredName, bSetupUndoRedo);
@@ -5193,7 +5193,7 @@ UEdGraph* UControlRigBlueprint::CreateEdGraph(URigVMGraph* InModel, bool bForce)
 	}
 
 	FString GraphName = InModel->GetName();
-	GraphName.RemoveFromStart(RigVMModelPrefix);
+	GraphName.RemoveFromStart(FRigVMClient::RigVMModelPrefix);
 	GraphName.TrimStartAndEndInline();
 
 	if(GraphName.IsEmpty())
diff --git a/Engine/Plugins/Animation/ControlRig/Source/ControlRigDeveloper/Private/Graph/ControlRigGraph.cpp b/Engine/Plugins/Animation/ControlRig/Source/ControlRigDeveloper/Private/Graph/ControlRigGraph.cpp
--- a/Engine/Plugins/Animation/ControlRig/Source/ControlRigDeveloper/Private/Graph/ControlRigGraph.cpp
+++ b/Engine/Plugins/Animation/ControlRig/Source/ControlRigDeveloper/Private/Graph/ControlRigGraph.cpp
@@ -68,7 +68,7 @@ void UControlRigGraph::HandleRigVMGraphRenamed(const FString& InOldNodePath, con
 			GraphName = ModelNodePath;
 		}
 		GraphName.RemoveFromEnd(TEXT("::"));
-		GraphName.RemoveFromStart(UControlRigBlueprint::RigVMModelPrefix);
+		GraphName.RemoveFromStart(FRigVMClient::RigVMModelPrefix);
 		GraphName.TrimStartAndEndInline();
 
 		if(GraphName.IsEmpty())
diff --git a/Engine/Plugins/Animation/ControlRig/Source/ControlRigDeveloper/Private/Graph/ControlRigGraphSchema.cpp b/Engine/Plugins/Animation/ControlRig/Source/ControlRigDeveloper/Private/Graph/ControlRigGraphSchema.cpp
--- a/Engine/Plugins/Animation/ControlRig/Source/ControlRigDeveloper/Private/Graph/ControlRigGraphSchema.cpp
+++ b/Engine/Plugins/Animation/ControlRig/Source/ControlRigDeveloper/Private/Graph/ControlRigGraphSchema.cpp
@@ -1427,7 +1427,7 @@ void UControlRigGraphSchema::GetGraphDisplayInformation(const UEdGraph& Graph, /
 						static constexpr TCHAR NodePathSuffix[] = TEXT("::");
 						FString NodePath = NodePathParts[0];
 						NodePath.RemoveFromEnd(NodePathSuffix);
-						NodePath.RemoveFromStart(UControlRigBlueprint::RigVMModelPrefix);
+						NodePath.RemoveFromStart(FRigVMClient::RigVMModelPrefix);
 						NodePath.TrimStartAndEndInline();
 						DisplayInfo.DisplayName = FText::FromString(NodePath);
 						DisplayInfo.PlainName = DisplayInfo.DisplayName;
@@ -1643,7 +1643,7 @@ bool UControlRigGraphSchema::TryRenameGraph(UEdGraph* GraphToRename, const FName
 			{
 				if(Model->IsRootGraph())
 				{
-					const FString NewName = FString::Printf(TEXT("%s %s"), RigBlueprint->RigVMModelPrefix, *InNewName.ToString()); 
+					const FString NewName = FString::Printf(TEXT("%s %s"), FRigVMClient::RigVMModelPrefix, *InNewName.ToString()); 
 					RigBlueprint->RenameGraph(Model->GetNodePath(), *NewName);
 				}
 				else if (URigVMGraph* RootModel = Model->GetRootGraph())
diff --git a/Engine/Plugins/Animation/ControlRig/Source/ControlRigDeveloper/Public/ControlRigBlueprint.h b/Engine/Plugins/Animation/ControlRig/Source/ControlRigDeveloper/Public/ControlRigBlueprint.h
--- a/Engine/Plugins/Animation/ControlRig/Source/ControlRigDeveloper/Public/ControlRigBlueprint.h
+++ b/Engine/Plugins/Animation/ControlRig/Source/ControlRigDeveloper/Public/ControlRigBlueprint.h
@@ -802,8 +802,6 @@ public:
 
 #endif
 
-	static constexpr TCHAR RigVMModelPrefix[] = TEXT("RigVMModel");
-
 protected:
 
 	static FSoftObjectPath PreDuplicateAssetPath;
diff --git a/Engine/Source/Developer/RigVMDeveloper/Private/RigVMModel/RigVMClient.cpp b/Engine/Source/Developer/RigVMDeveloper/Private/RigVMModel/RigVMClient.cpp
--- a/Engine/Source/Developer/RigVMDeveloper/Private/RigVMModel/RigVMClient.cpp
+++ b/Engine/Source/Developer/RigVMDeveloper/Private/RigVMModel/RigVMClient.cpp
@@ -5,6 +5,7 @@
 #include "Exporters/Exporter.h"
 #include "Algo/Sort.h"
 #include "UnrealExporter.h"
+#include "RigVMModel/RigVMControllerActions.h"
 
 #include UE_INLINE_GENERATED_CPP_BY_NAME(RigVMClient)
 
@@ -19,6 +20,15 @@ void FRigVMClient::SetOuterClientHost(UObject* InOuterClientHost, const FName& I
 
 	check(OuterClientHost->Implements<URigVMClientHost>());
 	check(GetOuterClientProperty() != nullptr);
+
+	// create the null graph / default controller.
+	// we need this to react to notifs without a valid graph
+	// such as interaction brackets
+	static const URigVMGraph* NullGraph = nullptr;
+	if(!Controllers.Contains(NullGraph))
+	{
+		CreateController(NullGraph);
+	}
 }
 
 void FRigVMClient::SetFromDeprecatedData(URigVMGraph* InDefaultGraph, URigVMFunctionLibrary* InFunctionLibrary)
@@ -625,10 +635,15 @@ void FRigVMClient::NotifyOuterOfPropertyChange(EPropertyChangeType::Type ChangeT
 
 URigVMController* FRigVMClient::CreateController(const URigVMGraph* InModel)
 {
-	const FName SafeControllerName = GetUniqueName(*FString::Printf(TEXT("%s_Controller"), *InModel->GetName()));
+	const FString ModelName = InModel ? InModel->GetName() : TEXT("__NullGraph");
+	const FName SafeControllerName = GetUniqueName(*FString::Printf(TEXT("%s_Controller"), *ModelName));
 	URigVMController* Controller = NewObject<URigVMController>(GetOuter(), SafeControllerName);
 	Controllers.Add(InModel, Controller);
-	Controller->SetGraph((URigVMGraph*)InModel);
+	Controller->SetActionStack(GetOrCreateActionStack());
+	if(InModel)
+	{
+		Controller->SetGraph((URigVMGraph*)InModel);
+	}
 	Controller->OnModified().AddLambda([this](ERigVMGraphNotifType InNotifType, URigVMGraph* InGraph, UObject* InSubject)
 	{
 		HandleGraphModifiedEvent(InNotifType, InGraph, InSubject);
@@ -639,11 +654,31 @@ URigVMController* FRigVMClient::CreateController(const URigVMGraph* InModel)
 		IRigVMClientHost* ClientHost = Cast<IRigVMClientHost>(GetOuter());
 		ClientHost->HandleConfigureRigVMController(this, Controller);
 	}
-	
-	Controller->RemoveStaleNodes();
+
+	if(InModel)
+	{
+		Controller->RemoveStaleNodes();
+	}
 	return Controller;
 }
 
+URigVMActionStack* FRigVMClient::GetOrCreateActionStack()
+{
+	if(UObject* Outer = GetOuter())
+	{
+		if(ActionStack == nullptr)
+		{
+			ActionStack = NewObject<URigVMActionStack>(Outer, TEXT("ActionStack"));
+		}
+		else if(ActionStack->GetOuter() != Outer)
+		{
+			// make sure the action stack is parented correctly
+			ActionStack->Rename(nullptr, Outer, REN_ForceNoResetLoaders | REN_DoNotDirty | REN_DontCreateRedirectors | REN_NonTransactional);
+		}
+	}
+	return ActionStack;
+}
+
 FName FRigVMClient::GetUniqueName(const FName& InDesiredName) const
 {
 	return GetUniqueName(GetOuter(), InDesiredName);
diff --git a/Engine/Source/Developer/RigVMDeveloper/Private/RigVMModel/RigVMController.cpp b/Engine/Source/Developer/RigVMDeveloper/Private/RigVMModel/RigVMController.cpp
--- a/Engine/Source/Developer/RigVMDeveloper/Private/RigVMModel/RigVMController.cpp
+++ b/Engine/Source/Developer/RigVMDeveloper/Private/RigVMModel/RigVMController.cpp
@@ -113,11 +113,6 @@ URigVMController::URigVMController(const FObjectInitializer& ObjectInitializer)
 	, bRegisterTemplateNodeUsage(true)
 #endif
 {
-	ActionStack = CreateDefaultSubobject<URigVMActionStack>(TEXT("ActionStack"));
-
-	ActionStack->OnModified().AddLambda([&](ERigVMGraphNotifType NotifType, URigVMGraph* InGraph, UObject* InSubject) -> void {
-		Notify(NotifType, InSubject);
-	});
 }
 
 URigVMController::~URigVMController()
@@ -185,7 +180,7 @@ bool URigVMController::PushGraph(URigVMGraph* InGraph, bool bSetupUndoRedo)
 
 	if (bSetupUndoRedo)
 	{
-		ActionStack->AddAction(FRigVMPushGraphAction(InGraph));
+		GetActionStack()->AddAction(FRigVMPushGraphAction(this, InGraph));
 	}
 
 	return true;
@@ -214,7 +209,7 @@ URigVMGraph* URigVMController::PopGraph(bool bSetupUndoRedo)
 
 	if (bSetupUndoRedo)
 	{
-		ActionStack->AddAction(FRigVMPopGraphAction(LastGraph));
+		GetActionStack()->AddAction(FRigVMPopGraphAction(this, LastGraph));
 	}
 
 	return LastGraph;
@@ -259,6 +254,10 @@ void URigVMController::Notify(ERigVMGraphNotifType InNotifType, UObject* InSubje
 	{
 		Graph->Notify(InNotifType, InSubject);
 	}
+	else
+	{
+		const_cast<URigVMController*>(this)->HandleModifiedEvent(InNotifType, nullptr, InSubject);
+	}
 }
 
 void URigVMController::ResendAllNotifications()
@@ -1345,12 +1344,12 @@ URigVMUnitNode* URigVMController::AddUnitNode(UScriptStruct* InScriptStruct, con
 
 	FRigVMControllerCompileBracketScope CompileScope(this);
 
-	FRigVMBaseAction Action;
+	FRigVMBaseAction Action(this);
 	if (bSetupUndoRedo)
 	{
-		Action.Title = FString::Printf(TEXT("Add %s Node"), *Node->GetNodeTitle());
-		ActionStack->BeginAction(Action);
-		ActionStack->AddAction(FRigVMImportFromTextAction(Node, this));
+		Action.SetTitle(FString::Printf(TEXT("Add %s Node"), *Node->GetNodeTitle()));
+		GetActionStack()->BeginAction(Action);
+		GetActionStack()->AddAction(FRigVMImportFromTextAction(this, Node));
 	}
 
 	Notify(ERigVMGraphNotifType::NodeAdded, Node);
@@ -1367,7 +1366,7 @@ URigVMUnitNode* URigVMController::AddUnitNode(UScriptStruct* InScriptStruct, con
 
 	if (bSetupUndoRedo)
 	{
-		ActionStack->EndAction(Action);
+		GetActionStack()->EndAction(Action);
 	}
 
 	if (bPrintPythonCommand)
@@ -1687,12 +1686,12 @@ URigVMVariableNode* URigVMController::AddVariableNode(const FName& InVariableNam
 	}
 
 	FRigVMControllerCompileBracketScope CompileScope(this);
-	FRigVMBaseAction Action;
+	FRigVMBaseAction Action(this);
 	if (bSetupUndoRedo)
 	{
-		Action.Title = FString::Printf(TEXT("Add %s Variable"), *InVariableName.ToString());
-		ActionStack->BeginAction(Action);
-		ActionStack->AddAction(FRigVMImportFromTextAction(Node, this));
+		Action.SetTitle(FString::Printf(TEXT("Add %s Variable"), *InVariableName.ToString()));
+		GetActionStack()->BeginAction(Action);
+		GetActionStack()->AddAction(FRigVMImportFromTextAction(this, Node));
 	}
 
 	Notify(ERigVMGraphNotifType::NodeAdded, Node);
@@ -1700,7 +1699,7 @@ URigVMVariableNode* URigVMController::AddVariableNode(const FName& InVariableNam
 
 	if (bSetupUndoRedo)
 	{
-		ActionStack->EndAction(Action);
+		GetActionStack()->EndAction(Action);
 	}
 
 	if (bPrintPythonCommand)
@@ -2257,11 +2256,11 @@ bool URigVMController::UnresolveTemplateNodes(const TArray<URigVMNode*>& InNodes
 		return false;
 	}
 	
-	FRigVMBaseAction Action;
+	FRigVMBaseAction Action(this);
 	if(bSetupUndoRedo)
 	{
-		Action.Title = TEXT("Unresolve nodes");
-		ActionStack->BeginAction(Action);
+		Action.SetTitle(TEXT("Unresolve nodes"));
+		GetActionStack()->BeginAction(Action);
 	}
 
 	{
@@ -2338,7 +2337,7 @@ bool URigVMController::UnresolveTemplateNodes(const TArray<URigVMNode*>& InNodes
 
 	if(bSetupUndoRedo)
 	{
-		ActionStack->EndAction(Action);
+		GetActionStack()->EndAction(Action);
 	}
 
 	return true;
@@ -2428,11 +2427,11 @@ TArray<URigVMNode*> URigVMController::UpgradeNodes(const TArray<URigVMNode*>& In
 		return InNodes;
 	}
 
-	FRigVMBaseAction Action;
+	FRigVMBaseAction Action(this);
 	if(bSetupUndoRedo)
 	{
-		Action.Title = TEXT("Upgrade nodes");
-		ActionStack->BeginAction(Action);
+		Action.SetTitle(TEXT("Upgrade nodes"));
+		GetActionStack()->BeginAction(Action);
 	}
 
 	// find all links affecting the nodes to upgrade
@@ -2441,7 +2440,7 @@ TArray<URigVMNode*> URigVMController::UpgradeNodes(const TArray<URigVMNode*>& In
 	{
 		if(bSetupUndoRedo)
 		{
-			ActionStack->CancelAction(Action, this);
+			GetActionStack()->CancelAction(Action);
 		}
 		return TArray<URigVMNode*>();
 	}
@@ -2473,7 +2472,7 @@ TArray<URigVMNode*> URigVMController::UpgradeNodes(const TArray<URigVMNode*>& In
 
 	if(bSetupUndoRedo)
 	{
-		ActionStack->EndAction(Action);
+		GetActionStack()->EndAction(Action);
 	}
 
 	return UpgradedNodes;
@@ -2498,11 +2497,11 @@ URigVMNode* URigVMController::UpgradeNode(URigVMNode* InNode, bool bSetupUndoRed
 	const FString NodeName = InNode->GetName();
 	const FVector2D NodePosition = InNode->GetPosition();
 
-	FRigVMBaseAction Action;
+	FRigVMBaseAction Action(this);
 	if(bSetupUndoRedo)
 	{
-		Action.Title = TEXT("Upgrade node");
-		ActionStack->BeginAction(Action);
+		Action.SetTitle(TEXT("Upgrade node"));
+		GetActionStack()->BeginAction(Action);
 	}
 
 	URigVMNode* UpgradedNode = nullptr;
@@ -2528,7 +2527,7 @@ URigVMNode* URigVMController::UpgradeNode(URigVMNode* InNode, bool bSetupUndoRed
 	{
 		if(bSetupUndoRedo)
 		{
-			ActionStack->CancelAction(Action, this);
+			GetActionStack()->CancelAction(Action);
 		}
 		ReportErrorf(TEXT("Unable to remove node %s."), *NodeName);
 		return nullptr;
@@ -2574,7 +2573,7 @@ URigVMNode* URigVMController::UpgradeNode(URigVMNode* InNode, bool bSetupUndoRed
 	{
 		if(bSetupUndoRedo)
 		{
-			ActionStack->CancelAction(Action, this);
+			GetActionStack()->CancelAction(Action);
 		}
 		ReportErrorf(TEXT("Unable to upgrade node %s."), *NodeName);
 		return nullptr;
@@ -2625,7 +2624,7 @@ URigVMNode* URigVMController::UpgradeNode(URigVMNode* InNode, bool bSetupUndoRed
 
 	if(bSetupUndoRedo)
 	{
-		ActionStack->EndAction(Action);
+		GetActionStack()->EndAction(Action);
 	}
 
 	return UpgradedNode;
@@ -2699,19 +2698,19 @@ URigVMCommentNode* URigVMController::AddCommentNode(const FString& InCommentText
 	}
 
 	FRigVMControllerCompileBracketScope CompileScope(this);
-	FRigVMBaseAction Action;
+	FRigVMBaseAction Action(this);
 	if (bSetupUndoRedo)
 	{
-		Action.Title = FString::Printf(TEXT("Add Comment"));
-		ActionStack->BeginAction(Action);
-		ActionStack->AddAction(FRigVMImportFromTextAction(Node, this));
+		Action.SetTitle(FString::Printf(TEXT("Add Comment")));
+		GetActionStack()->BeginAction(Action);
+		GetActionStack()->AddAction(FRigVMImportFromTextAction(this, Node));
 	}
 
 	Notify(ERigVMGraphNotifType::NodeAdded, Node);
 
 	if (bSetupUndoRedo)
 	{
-		ActionStack->EndAction(Action);
+		GetActionStack()->EndAction(Action);
 	}
 
 	if (bPrintPythonCommand)
@@ -2751,11 +2750,11 @@ URigVMRerouteNode* URigVMController::AddRerouteNodeOnLink(URigVMLink* InLink, bo
 	TGuardValue<bool> GuardCompactness(bIgnoreRerouteCompactnessChanges, true);
 
 	FRigVMControllerCompileBracketScope CompileScope(this);
-	FRigVMBaseAction Action;
+	FRigVMBaseAction Action(this);
 	if (bSetupUndoRedo)
 	{
-		Action.Title = FString::Printf(TEXT("Add Reroute"));
-		ActionStack->BeginAction(Action);
+		Action.SetTitle(FString::Printf(TEXT("Add Reroute")));
+		GetActionStack()->BeginAction(Action);
 	}
 
 	URigVMRerouteNode* Node = AddRerouteNodeOnPin(TargetPin->GetPinPath(), true, bShowAsFullNode, InPosition, InNodeName, bSetupUndoRedo);
@@ -2763,7 +2762,7 @@ URigVMRerouteNode* URigVMController::AddRerouteNodeOnLink(URigVMLink* InLink, bo
 	{
 		if (bSetupUndoRedo)
 		{
-			ActionStack->CancelAction(Action, this);
+			GetActionStack()->CancelAction(Action);
 		}
 		return nullptr;
 	}
@@ -2773,7 +2772,7 @@ URigVMRerouteNode* URigVMController::AddRerouteNodeOnLink(URigVMLink* InLink, bo
 
 	if (bSetupUndoRedo)
 	{
-		ActionStack->EndAction(Action);
+		GetActionStack()->EndAction(Action);
 	}
 
 	if (bPrintPythonCommand)
@@ -2843,11 +2842,11 @@ URigVMRerouteNode* URigVMController::AddRerouteNodeOnPin(const FString& InPinPat
 	TGuardValue<bool> GuardCompactness(bIgnoreRerouteCompactnessChanges, true);
 
 	FRigVMControllerCompileBracketScope CompileScope(this);
-	FRigVMBaseAction Action;
+	FRigVMBaseAction Action(this);
 	if (bSetupUndoRedo)
 	{
-		Action.Title = FString::Printf(TEXT("Add Reroute"));
-		ActionStack->BeginAction(Action);
+		Action.SetTitle(FString::Printf(TEXT("Add Reroute")));
+		GetActionStack()->BeginAction(Action);
 	}
 
 	//in case an injected node is present, use its pins for any new links
@@ -2891,7 +2890,7 @@ URigVMRerouteNode* URigVMController::AddRerouteNodeOnPin(const FString& InPinPat
 
 	if (bSetupUndoRedo)
 	{
-		ActionStack->AddAction(FRigVMImportFromTextAction(Node, this, false));
+		GetActionStack()->AddAction(FRigVMImportFromTextAction(this, Node, false));
 	}
 
 	Notify(ERigVMGraphNotifType::NodeAdded, Node);
@@ -2907,7 +2906,7 @@ URigVMRerouteNode* URigVMController::AddRerouteNodeOnPin(const FString& InPinPat
 
 	if (bSetupUndoRedo)
 	{
-		ActionStack->EndAction(Action);
+		GetActionStack()->EndAction(Action);
 	}
 
 	if (bPrintPythonCommand)
@@ -3013,11 +3012,11 @@ URigVMInjectionInfo* URigVMController::AddInjectedNode(const FString& InPinPath,
 	// 3.- Inject node into pin
 
 	FRigVMControllerCompileBracketScope CompileScope(this);
-	FRigVMBaseAction Action;
+	FRigVMBaseAction Action(this);
 	if (bSetupUndoRedo)
 	{
-		Action.Title = FString::Printf(TEXT("Add Injected Node"));
-		ActionStack->BeginAction(Action);
+		Action.SetTitle(FString::Printf(TEXT("Add Injected Node")));
+		GetActionStack()->BeginAction(Action);
 	}
 
 	// 1.- Create unit node
@@ -3033,7 +3032,7 @@ URigVMInjectionInfo* URigVMController::AddInjectedNode(const FString& InPinPath,
 		{
 			if (bSetupUndoRedo)
 			{
-				ActionStack->CancelAction(Action, this);
+				GetActionStack()->CancelAction(Action);
 			}
 			return nullptr;
 		}
@@ -3043,7 +3042,7 @@ URigVMInjectionInfo* URigVMController::AddInjectedNode(const FString& InPinPath,
 			RemoveNode(UnitNode, false);
 			if (bSetupUndoRedo)
 			{
-				ActionStack->CancelAction(Action, this);
+				GetActionStack()->CancelAction(Action);
 			}
 			return nullptr;
 		}
@@ -3059,7 +3058,7 @@ URigVMInjectionInfo* URigVMController::AddInjectedNode(const FString& InPinPath,
 			ReportErrorf(TEXT("Injected node %s is using incompatible input and output pins."), *InScriptStruct->GetName());
 			if (bSetupUndoRedo)
 			{
-				ActionStack->CancelAction(Action, this);
+				GetActionStack()->CancelAction(Action);
 			}
 			return nullptr;
 		}
@@ -3070,7 +3069,7 @@ URigVMInjectionInfo* URigVMController::AddInjectedNode(const FString& InPinPath,
 			ReportErrorf(TEXT("Injected node %s is using incompatible pin."), *InScriptStruct->GetName());
 			if (bSetupUndoRedo)
 			{
-				ActionStack->CancelAction(Action, this);
+				GetActionStack()->CancelAction(Action);
 			}
 			return nullptr;
 		}
@@ -3123,7 +3122,7 @@ URigVMInjectionInfo* URigVMController::AddInjectedNode(const FString& InPinPath,
 	
 	if (bSetupUndoRedo)
 	{
-		ActionStack->EndAction(Action);
+		GetActionStack()->EndAction(Action);
 	}
 
 	if (bPrintPythonCommand)
@@ -3206,11 +3205,11 @@ bool URigVMController::RemoveInjectedNode(const FString& InPinPath, bool bAsInpu
 	// 3.- Remove node
 
 	FRigVMControllerCompileBracketScope CompileScope(this);
-	FRigVMBaseAction Action;
+	FRigVMBaseAction Action(this);
 	if (bSetupUndoRedo)
 	{
-		Action.Title = FString::Printf(TEXT("Remove Injected Node"));
-		ActionStack->BeginAction(Action);
+		Action.SetTitle(FString::Printf(TEXT("Remove Injected Node")));
+		GetActionStack()->BeginAction(Action);
 	}
 
 	URigVMInjectionInfo* InjectionInfo = Pin->InjectionInfos.Last();
@@ -3221,7 +3220,7 @@ bool URigVMController::RemoveInjectedNode(const FString& InPinPath, bool bAsInpu
 	URigVMNode* NodeEjected = EjectNodeFromPin(InPinPath, bSetupUndoRedo);
 	if (!NodeEjected)
 	{
-		ActionStack->CancelAction(Action, this);
+		GetActionStack()->CancelAction(Action);
 		return false;
 	}
 
@@ -3245,7 +3244,7 @@ bool URigVMController::RemoveInjectedNode(const FString& InPinPath, bool bAsInpu
 	// 3.- Remove node
 	if (!RemoveNode(NodeEjected))
 	{
-		ActionStack->CancelAction(Action, this);
+		GetActionStack()->CancelAction(Action);
 		return false;
 	}
 
@@ -3256,7 +3255,7 @@ bool URigVMController::RemoveInjectedNode(const FString& InPinPath, bool bAsInpu
 	
 	if (bSetupUndoRedo)
 	{
-		ActionStack->EndAction(Action);
+		GetActionStack()->EndAction(Action);
 	}
 
 	if (bPrintPythonCommand)
@@ -3357,11 +3356,11 @@ URigVMInjectionInfo* URigVMController::InjectNodeIntoPin(URigVMPin* InPin, bool
 	}
 
 	FRigVMControllerCompileBracketScope CompileScope(this);
-	FRigVMBaseAction Action;
+	FRigVMBaseAction Action(this);
 	if (bSetupUndoRedo)
 	{
-		Action.Title = FString::Printf(TEXT("Inject Node"));
-		ActionStack->BeginAction(Action);
+		Action.SetTitle(FString::Printf(TEXT("Inject Node")));
+		GetActionStack()->BeginAction(Action);
 	}
 	
 	URigVMInjectionInfo* InjectionInfo = NewObject<URigVMInjectionInfo>(InPin);
@@ -3408,8 +3407,8 @@ URigVMInjectionInfo* URigVMController::InjectNodeIntoPin(URigVMPin* InPin, bool
 
 	if (bSetupUndoRedo)
 	{
-		ActionStack->AddAction(FRigVMInjectNodeIntoPinAction(InjectionInfo));
-		ActionStack->EndAction(Action);
+		GetActionStack()->AddAction(FRigVMInjectNodeIntoPinAction(this, InjectionInfo));
+		GetActionStack()->EndAction(Action);
 	}
 
 	return InjectionInfo;
@@ -3470,13 +3469,13 @@ URigVMNode* URigVMController::EjectNodeFromPin(URigVMPin* InPin, bool bSetupUndo
 
 	
 	FRigVMControllerCompileBracketScope CompileScope(this);
-	FRigVMBaseAction Action;
+	FRigVMBaseAction Action(this);
 	if (bSetupUndoRedo)
 	{
-		Action.Title = TEXT("Eject node");
+		Action.SetTitle(TEXT("Eject node"));
 
-		ActionStack->BeginAction(Action);
-		ActionStack->AddAction(FRigVMEjectNodeFromPinAction(Injection));
+		GetActionStack()->BeginAction(Action);
+		GetActionStack()->AddAction(FRigVMEjectNodeFromPinAction(this, Injection));
 	}
 
 	FVector2D Position = InPin->GetNode()->GetPosition() + FVector2D(0.f, 12.f) * float(InPin->GetPinIndex());
@@ -3542,7 +3541,7 @@ URigVMNode* URigVMController::EjectNodeFromPin(URigVMPin* InPin, bool bSetupUndo
 
 	if (bSetupUndoRedo)
 	{
-		ActionStack->EndAction(Action);
+		GetActionStack()->EndAction(Action);
 	}
 
 	if (bPrintPythonCommand)
@@ -3576,10 +3575,10 @@ bool URigVMController::EjectAllInjectedNodes(URigVMNode* InNode, bool bSetupUndo
 		return false;
 	}
 
-	FRigVMBaseAction EjectAllInjectedNodesAction;
+	FRigVMBaseAction EjectAllInjectedNodesAction(this);
 	if (bSetupUndoRedo)
 	{
-		ActionStack->BeginAction(EjectAllInjectedNodesAction);
+		GetActionStack()->BeginAction(EjectAllInjectedNodesAction);
 	}
 
 	for(URigVMPin* Pin : InNode->GetPins())
@@ -3595,7 +3594,7 @@ bool URigVMController::EjectAllInjectedNodes(URigVMNode* InNode, bool bSetupUndo
 
 	if(bSetupUndoRedo)
 	{
-		ActionStack->EndAction(EjectAllInjectedNodesAction);
+		GetActionStack()->EndAction(EjectAllInjectedNodesAction);
 	}
 
 	return true;
@@ -3609,7 +3608,7 @@ bool URigVMController::Undo()
 		return false;
 	}
 
-	return ActionStack->Undo(this);
+	return GetActionStack()->Undo(this);
 }
 
 bool URigVMController::Redo()
@@ -3619,7 +3618,7 @@ bool URigVMController::Redo()
 		return false;
 	}
 
-	return ActionStack->Redo(this);
+	return GetActionStack()->Redo(this);
 }
 
 bool URigVMController::OpenUndoBracket(const FString& InTitle)
@@ -3628,7 +3627,7 @@ bool URigVMController::OpenUndoBracket(const FString& InTitle)
 	{
 		return false;
 	}
-	return ActionStack->OpenUndoBracket(InTitle);
+	return GetActionStack()->OpenUndoBracket(this, InTitle);
 }
 
 bool URigVMController::CloseUndoBracket()
@@ -3637,7 +3636,7 @@ bool URigVMController::CloseUndoBracket()
 	{
 		return false;
 	}
-	return ActionStack->CloseUndoBracket(this);
+	return GetActionStack()->CloseUndoBracket(this);
 }
 
 bool URigVMController::CancelUndoBracket()
@@ -3646,7 +3645,7 @@ bool URigVMController::CancelUndoBracket()
 	{
 		return false;
 	}
-	return ActionStack->CancelUndoBracket(this);
+	return GetActionStack()->CancelUndoBracket(this);
 }
 
 FString URigVMController::ExportNodesToText(const TArray<FName>& InNodeNames, bool bIncludeExteriorLinks)
@@ -4096,9 +4095,9 @@ TArray<FName> URigVMController::ImportNodesFromText(const FString& InText, bool
 
 	if (bSetupUndoRedo)
 	{
-		FRigVMImportFromTextAction Action(InText, NodeNames);
-		Action.Title = TEXT("Importing Nodes from Text");
-		ActionStack->AddAction(Action);
+		FRigVMImportFromTextAction Action(this, InText, NodeNames);
+		Action.SetTitle(TEXT("Importing Nodes from Text"));
+		GetActionStack()->AddAction(Action);
 	}
 
 #if WITH_EDITOR
@@ -4633,8 +4632,8 @@ URigVMCollapseNode* URigVMController::CollapseNodes(const TArray<URigVMNode*>& I
 	if (bSetupUndoRedo)
 	{
 		CollapseAction = FRigVMCollapseNodesAction(this, Nodes, CollapseNodeName, bIsAggregate); 
-		CollapseAction.Title = TEXT("Collapse Nodes");
-		ActionStack->BeginAction(CollapseAction);
+		CollapseAction.SetTitle(TEXT("Collapse Nodes"));
+		GetActionStack()->BeginAction(CollapseAction);
 	}
 
 #if UE_RIGVM_AGGREGATE_NODES_ENABLED
@@ -5013,7 +5012,7 @@ URigVMCollapseNode* URigVMController::CollapseNodes(const TArray<URigVMNode*>& I
 
 	if (bSetupUndoRedo)
 	{
-		ActionStack->EndAction(CollapseAction);
+		GetActionStack()->EndAction(CollapseAction);
 	}
 
 	return CollapseNode;
@@ -5067,8 +5066,8 @@ TArray<URigVMNode*> URigVMController::ExpandLibraryNode(URigVMLibraryNode* InNod
 	if (bSetupUndoRedo)
 	{
 		ExpandAction = FRigVMExpandNodeAction(this, InNode);
-		ExpandAction.Title = FString::Printf(TEXT("Expand '%s' Node"), *InNode->GetName());
-		ActionStack->BeginAction(ExpandAction);
+		ExpandAction.SetTitle(FString::Printf(TEXT("Expand '%s' Node"), *InNode->GetName()));
+		GetActionStack()->BeginAction(ExpandAction);
 	}
 
 	TArray<FName> NodeNames;
@@ -5107,7 +5106,7 @@ TArray<URigVMNode*> URigVMController::ExpandLibraryNode(URigVMLibraryNode* InNod
 	{
 		if (bSetupUndoRedo)
 		{
-			ActionStack->CancelAction(ExpandAction, this);
+			GetActionStack()->CancelAction(ExpandAction);
 		}
 		return TArray<URigVMNode*>();
 	}
@@ -5159,7 +5158,7 @@ TArray<URigVMNode*> URigVMController::ExpandLibraryNode(URigVMLibraryNode* InNod
 							ReportErrorf(TEXT("Found variable %s of incompatible type with a local variable inside function %s"), *LocalVariable.Name.ToString(), *FunctionReferenceNode->GetReferencedFunctionHeader().Name.ToString());
 							if (bSetupUndoRedo)
 							{
-								ActionStack->CancelAction(ExpandAction, this);
+								GetActionStack()->CancelAction(ExpandAction);
 							}
 							return TArray<URigVMNode*>();
 						}
@@ -5732,7 +5731,7 @@ TArray<URigVMNode*> URigVMController::ExpandLibraryNode(URigVMLibraryNode* InNod
 		{
 			ExpandAction.ExpandedNodePaths.Add(ExpandedNode->GetName());
 		}
-		ActionStack->EndAction(ExpandAction);
+		GetActionStack()->EndAction(ExpandAction);
 	}
 
 	return ExpandedNodes;
@@ -5928,7 +5927,7 @@ URigVMFunctionReferenceNode* URigVMController::PromoteCollapseNodeToFunctionRefe
 
 		if (bSetupUndoRedo)
 		{
-			ActionStack->AddAction(FRigVMPromoteNodeAction(InCollapseNode, NodeName, FString()));
+			GetActionStack()->AddAction(FRigVMPromoteNodeAction(this, InCollapseNode, NodeName, FString()));
 		}
 	}
 
@@ -6012,7 +6011,7 @@ URigVMCollapseNode* URigVMController::PromoteFunctionReferenceNodeToCollapseNode
 
 	if (bSetupUndoRedo)
 	{
-		ActionStack->AddAction(FRigVMPromoteNodeAction(InFunctionRefNode, NodeName, FunctionDefinition->GetPathName()));
+		GetActionStack()->AddAction(FRigVMPromoteNodeAction(this, InFunctionRefNode, NodeName, FunctionDefinition->GetPathName()));
 	}
 
 	FastBreakLinkedPaths(LinkedPaths);
@@ -6252,21 +6251,20 @@ bool URigVMController::RemoveNodes(TArray<URigVMNode*> InNodes, bool bSetupUndoR
 	TGuardValue<bool> GuardCompactness(bIgnoreRerouteCompactnessChanges, true);
 
 	FRigVMControllerCompileBracketScope CompileScope(this);
-	FRigVMBaseAction Action;
+	FRigVMBaseAction Action(this);
 	if (bSetupUndoRedo)
 	{
-		Action = FRigVMBaseAction();
 		if(InNodes.Num() == 1)
 		{
-			Action.Title = FString::Printf(TEXT("Remove %s Node"), *InNodes[0]->GetNodeTitle());
+			Action.SetTitle(FString::Printf(TEXT("Remove %s Node"), *InNodes[0]->GetNodeTitle()));
 		}
 		else
 		{
 			static const FString RemoveNodesTitle = TEXT("Remove nodes");
-			Action.Title = RemoveNodesTitle;
+			Action.SetTitle(RemoveNodesTitle);
 		}
-		ActionStack->BeginAction(Action);
-		ActionStack->AddAction(FRigVMRemoveNodesAction(InNodes, this));
+		GetActionStack()->BeginAction(Action);
+		GetActionStack()->AddAction(FRigVMRemoveNodesAction(this, InNodes));
 	}
 
 	for(URigVMNode* InNode : InNodes)
@@ -6278,7 +6276,7 @@ bool URigVMController::RemoveNodes(TArray<URigVMNode*> InNodes, bool bSetupUndoR
 
 			if (!EjectNodeFromPin(Pin->GetPinPath(), bSetupUndoRedo))
 			{
-				ActionStack->CancelAction(Action, this);
+				GetActionStack()->CancelAction(Action);
 				return false;
 			}
 			
@@ -6354,7 +6352,7 @@ bool URigVMController::RemoveNodes(TArray<URigVMNode*> InNodes, bool bSetupUndoR
 
 					if (bSetupUndoRedo)
 					{
-						ActionStack->AddAction(FRigVMMarkFunctionPublicAction(LibraryNode->GetFName(), true));
+						GetActionStack()->AddAction(FRigVMMarkFunctionPublicAction(this, LibraryNode->GetFName(), true));
 					}
 				}
 			}
@@ -6451,7 +6449,7 @@ bool URigVMController::RemoveNodes(TArray<URigVMNode*> InNodes, bool bSetupUndoR
 
 	if (bSetupUndoRedo)
 	{
-		ActionStack->EndAction(Action);
+		GetActionStack()->EndAction(Action);
 	}
 
 	return true;
@@ -6523,8 +6521,8 @@ bool URigVMController::RenameNode(URigVMNode* InNode, const FName& InNewName, bo
 	FRigVMRenameNodeAction Action;
 	if (bSetupUndoRedo)
 	{
-		Action = FRigVMRenameNodeAction(InNode->GetFName(), ValidNewName);
-		ActionStack->BeginAction(Action);
+		Action = FRigVMRenameNodeAction(this, InNode->GetFName(), ValidNewName);
+		GetActionStack()->BeginAction(Action);
 	}
 
 	// loop over all links and remove them
@@ -6535,7 +6533,7 @@ bool URigVMController::RenameNode(URigVMNode* InNode, const FName& InNewName, bo
 	InNode->PreviousName = InNode->GetFName();
 	if (!RenameObject(InNode, *ValidNewName.ToString()))
 	{
-		ActionStack->CancelAction(Action, this);
+		GetActionStack()->CancelAction(Action);
 		return false;
 	}
 
@@ -6570,7 +6568,7 @@ bool URigVMController::RenameNode(URigVMNode* InNode, const FName& InNewName, bo
 
 	if (bSetupUndoRedo)
 	{
-		ActionStack->EndAction(Action);
+		GetActionStack()->EndAction(Action);
 	}
 
 	if (GetGraph()->IsA<URigVMFunctionLibrary>())
@@ -6665,7 +6663,7 @@ bool URigVMController::SetNodeSelection(const TArray<FName>& InNodeNames, bool b
 	URigVMGraph* Graph = GetGraph();
 	check(Graph);
 
-	FRigVMSetNodeSelectionAction Action(Graph, InNodeNames);
+	FRigVMSetNodeSelectionAction Action(this, Graph, InNodeNames);
 	bool bSelectionChanged = false;
 
 	TArray<FName> PreviousSelection = Graph->GetSelectNodes();
@@ -6700,20 +6698,20 @@ bool URigVMController::SetNodeSelection(const TArray<FName>& InNodeNames, bool b
 			const TArray<FName>& SelectedNodes = Graph->GetSelectNodes();
 			if (SelectedNodes.Num() == 0)
 			{
-				Action.Title = TEXT("Deselect all nodes.");
+				Action.SetTitle(TEXT("Deselect all nodes."));
 			}
 			else
 			{
 				if (SelectedNodes.Num() == 1)
 				{
-					Action.Title = FString::Printf(TEXT("Selected node '%s'."), *SelectedNodes[0].ToString());
+					Action.SetTitle(FString::Printf(TEXT("Selected node '%s'."), *SelectedNodes[0].ToString()));
 				}
 				else
 				{
-					Action.Title = TEXT("Selected multiple nodes.");
+					Action.SetTitle(TEXT("Selected multiple nodes."));
 				}
 			}
-			ActionStack->AddAction(Action);
+			GetActionStack()->AddAction(Action);
 		}
 	}
 
@@ -6766,9 +6764,9 @@ bool URigVMController::SetNodePosition(URigVMNode* InNode, const FVector2D& InPo
 	FRigVMSetNodePositionAction Action;
 	if (bSetupUndoRedo)
 	{
-		Action = FRigVMSetNodePositionAction(InNode, InPosition);
-		Action.Title = FString::Printf(TEXT("Set Node Position"));
-		ActionStack->BeginAction(Action);
+		Action = FRigVMSetNodePositionAction(this, InNode, InPosition);
+		Action.SetTitle(FString::Printf(TEXT("Set Node Position")));
+		GetActionStack()->BeginAction(Action);
 	}
 
 	InNode->Position = InPosition;
@@ -6776,7 +6774,7 @@ bool URigVMController::SetNodePosition(URigVMNode* InNode, const FVector2D& InPo
 
 	if (bSetupUndoRedo)
 	{
-		ActionStack->EndAction(Action, bMergeUndoAction);
+		GetActionStack()->EndAction(Action, bMergeUndoAction);
 	}
 
 	if (bPrintPythonCommand)
@@ -6834,9 +6832,9 @@ bool URigVMController::SetNodeSize(URigVMNode* InNode, const FVector2D& InSize,
 	FRigVMSetNodeSizeAction Action;
 	if (bSetupUndoRedo)
 	{
-		Action = FRigVMSetNodeSizeAction(InNode, InSize);
-		Action.Title = FString::Printf(TEXT("Set Node Size"));
-		ActionStack->BeginAction(Action);
+		Action = FRigVMSetNodeSizeAction(this, InNode, InSize);
+		Action.SetTitle(FString::Printf(TEXT("Set Node Size")));
+		GetActionStack()->BeginAction(Action);
 	}
 
 	InNode->Size = InSize;
@@ -6844,7 +6842,7 @@ bool URigVMController::SetNodeSize(URigVMNode* InNode, const FVector2D& InSize,
 
 	if (bSetupUndoRedo)
 	{
-		ActionStack->EndAction(Action, bMergeUndoAction);
+		GetActionStack()->EndAction(Action, bMergeUndoAction);
 	}
 
 	if (bPrintPythonCommand)
@@ -6901,9 +6899,9 @@ bool URigVMController::SetNodeColor(URigVMNode* InNode, const FLinearColor& InCo
 	FRigVMSetNodeColorAction Action;
 	if (bSetupUndoRedo)
 	{
-		Action = FRigVMSetNodeColorAction(InNode, InColor);
-		Action.Title = FString::Printf(TEXT("Set Node Color"));
-		ActionStack->BeginAction(Action);
+		Action = FRigVMSetNodeColorAction(this, InNode, InColor);
+		Action.SetTitle(FString::Printf(TEXT("Set Node Color")));
+		GetActionStack()->BeginAction(Action);
 	}
 
 	InNode->NodeColor = InColor;
@@ -6923,7 +6921,7 @@ bool URigVMController::SetNodeColor(URigVMNode* InNode, const FLinearColor& InCo
 
 	if (bSetupUndoRedo)
 	{
-		ActionStack->EndAction(Action, bMergeUndoAction);
+		GetActionStack()->EndAction(Action, bMergeUndoAction);
 	}
 
 	if (bPrintPythonCommand)
@@ -6980,9 +6978,9 @@ bool URigVMController::SetNodeCategory(URigVMCollapseNode* InNode, const FString
 	FRigVMSetNodeCategoryAction Action;
 	if (bSetupUndoRedo)
 	{
-		Action = FRigVMSetNodeCategoryAction(InNode, InCategory);
-		Action.Title = FString::Printf(TEXT("Set Node Category"));
-		ActionStack->BeginAction(Action);
+		Action = FRigVMSetNodeCategoryAction(this, InNode, InCategory);
+		Action.SetTitle(FString::Printf(TEXT("Set Node Category")));
+		GetActionStack()->BeginAction(Action);
 	}
 
 	InNode->NodeCategory = InCategory;
@@ -6990,7 +6988,7 @@ bool URigVMController::SetNodeCategory(URigVMCollapseNode* InNode, const FString
 
 	if (bSetupUndoRedo)
 	{
-		ActionStack->EndAction(Action, bMergeUndoAction);
+		GetActionStack()->EndAction(Action, bMergeUndoAction);
 	}
 
 	if (bPrintPythonCommand)
@@ -7047,9 +7045,9 @@ bool URigVMController::SetNodeKeywords(URigVMCollapseNode* InNode, const FString
 	FRigVMSetNodeKeywordsAction Action;
 	if (bSetupUndoRedo)
 	{
-		Action = FRigVMSetNodeKeywordsAction(InNode, InKeywords);
-		Action.Title = FString::Printf(TEXT("Set Node Keywords"));
-		ActionStack->BeginAction(Action);
+		Action = FRigVMSetNodeKeywordsAction(this, InNode, InKeywords);
+		Action.SetTitle(FString::Printf(TEXT("Set Node Keywords")));
+		GetActionStack()->BeginAction(Action);
 	}
 
 	InNode->NodeKeywords = InKeywords;
@@ -7057,7 +7055,7 @@ bool URigVMController::SetNodeKeywords(URigVMCollapseNode* InNode, const FString
 
 	if (bSetupUndoRedo)
 	{
-		ActionStack->EndAction(Action, bMergeUndoAction);
+		GetActionStack()->EndAction(Action, bMergeUndoAction);
 	}
 
 	if (bPrintPythonCommand)
@@ -7114,9 +7112,9 @@ bool URigVMController::SetNodeDescription(URigVMCollapseNode* InNode, const FStr
 	FRigVMSetNodeDescriptionAction Action;
 	if (bSetupUndoRedo)
 	{
-		Action = FRigVMSetNodeDescriptionAction(InNode, InDescription);
-		Action.Title = FString::Printf(TEXT("Set Node Description"));
-		ActionStack->BeginAction(Action);
+		Action = FRigVMSetNodeDescriptionAction(this, InNode, InDescription);
+		Action.SetTitle(FString::Printf(TEXT("Set Node Description")));
+		GetActionStack()->BeginAction(Action);
 	}
 
 	InNode->NodeDescription = InDescription;
@@ -7124,7 +7122,7 @@ bool URigVMController::SetNodeDescription(URigVMCollapseNode* InNode, const FStr
 
 	if (bSetupUndoRedo)
 	{
-		ActionStack->EndAction(Action, bMergeUndoAction);
+		GetActionStack()->EndAction(Action, bMergeUndoAction);
 	}
 
 	if (bPrintPythonCommand)
@@ -7183,9 +7181,9 @@ bool URigVMController::SetCommentText(URigVMNode* InNode, const FString& InComme
 		FRigVMSetCommentTextAction Action;
 		if (bSetupUndoRedo)
 		{
-			Action = FRigVMSetCommentTextAction(CommentNode, InCommentText, InCommentFontSize, bInCommentBubbleVisible, bInCommentColorBubble);
-			Action.Title = FString::Printf(TEXT("Set Comment Text"));
-			ActionStack->BeginAction(Action);
+			Action = FRigVMSetCommentTextAction(this, CommentNode, InCommentText, InCommentFontSize, bInCommentBubbleVisible, bInCommentColorBubble);
+			Action.SetTitle(FString::Printf(TEXT("Set Comment Text")));
+			GetActionStack()->BeginAction(Action);
 		}
 
 		CommentNode->CommentText = InCommentText;
@@ -7196,7 +7194,7 @@ bool URigVMController::SetCommentText(URigVMNode* InNode, const FString& InComme
 
 		if (bSetupUndoRedo)
 		{
-			ActionStack->EndAction(Action);
+			GetActionStack()->EndAction(Action);
 		}
 
 		if (bPrintPythonCommand)
@@ -7258,9 +7256,9 @@ bool URigVMController::SetRerouteCompactness(URigVMNode* InNode, bool bShowAsFul
 		FRigVMSetRerouteCompactnessAction Action;
 		if (bSetupUndoRedo)
 		{
-			Action = FRigVMSetRerouteCompactnessAction(RerouteNode, bShowAsFullNode);
-			Action.Title = FString::Printf(TEXT("Set Reroute Size"));
-			ActionStack->BeginAction(Action);
+			Action = FRigVMSetRerouteCompactnessAction(this, RerouteNode, bShowAsFullNode);
+			Action.SetTitle(FString::Printf(TEXT("Set Reroute Size")));
+			GetActionStack()->BeginAction(Action);
 		}
 
 		RerouteNode->bShowAsFullNode = bShowAsFullNode;
@@ -7268,7 +7266,7 @@ bool URigVMController::SetRerouteCompactness(URigVMNode* InNode, bool bShowAsFul
 
 		if (bSetupUndoRedo)
 		{
-			ActionStack->EndAction(Action);
+			GetActionStack()->EndAction(Action);
 		}
 
 		return true;
@@ -7340,9 +7338,9 @@ bool URigVMController::RenameVariable(const FName& InOldName, const FName& InNew
 	FRigVMRenameVariableAction Action;
 	if (bSetupUndoRedo)
 	{
-		Action = FRigVMRenameVariableAction(InOldName, InNewName);
-		Action.Title = FString::Printf(TEXT("Rename Variable"));
-		ActionStack->BeginAction(Action);
+		Action = FRigVMRenameVariableAction(this, InOldName, InNewName);
+		Action.SetTitle(FString::Printf(TEXT("Rename Variable")));
+		GetActionStack()->BeginAction(Action);
 	}
 
 	TArray<URigVMNode*> RenamedNodes;
@@ -7371,11 +7369,11 @@ bool URigVMController::RenameVariable(const FName& InOldName, const FName& InNew
 	{
 		if (RenamedNodes.Num() > 0)
 		{
-			ActionStack->EndAction(Action);
+			GetActionStack()->EndAction(Action);
 		}
 		else
 		{
-			ActionStack->CancelAction(Action, this);
+			GetActionStack()->CancelAction(Action);
 		}
 	}
 
@@ -7474,9 +7472,9 @@ bool URigVMController::SetPinExpansion(URigVMPin* InPin, bool bIsExpanded, bool
 	FRigVMSetPinExpansionAction Action;
 	if (bSetupUndoRedo)
 	{
-		Action = FRigVMSetPinExpansionAction(InPin, bIsExpanded);
-		Action.Title = bIsExpanded ? TEXT("Expand Pin") : TEXT("Collapse Pin");
-		ActionStack->BeginAction(Action);
+		Action = FRigVMSetPinExpansionAction(this, InPin, bIsExpanded);
+		Action.SetTitle(bIsExpanded ? TEXT("Expand Pin") : TEXT("Collapse Pin"));
+		GetActionStack()->BeginAction(Action);
 	}
 
 	InPin->bIsExpanded = bIsExpanded;
@@ -7489,7 +7487,7 @@ bool URigVMController::SetPinExpansion(URigVMPin* InPin, bool bIsExpanded, bool
 
 	if (bSetupUndoRedo)
 	{
-		ActionStack->EndAction(Action);
+		GetActionStack()->EndAction(Action);
 	}
 
 	return true;
@@ -7549,9 +7547,9 @@ bool URigVMController::SetPinIsWatched(URigVMPin* InPin, bool bIsWatched, bool b
 	FRigVMSetPinWatchAction Action;
 	if (bSetupUndoRedo)
 	{
-		Action = FRigVMSetPinWatchAction(InPin, bIsWatched);
-		Action.Title = bIsWatched ? TEXT("Watch Pin") : TEXT("Unwatch Pin");
-		ActionStack->BeginAction(Action);
+		Action = FRigVMSetPinWatchAction(this, InPin, bIsWatched);
+		Action.SetTitle(bIsWatched ? TEXT("Watch Pin") : TEXT("Unwatch Pin"));
+		GetActionStack()->BeginAction(Action);
 	}
 
 	InPin->bRequiresWatch = bIsWatched;
@@ -7564,7 +7562,7 @@ bool URigVMController::SetPinIsWatched(URigVMPin* InPin, bool bIsWatched, bool b
 
 	if (bSetupUndoRedo)
 	{
-		ActionStack->EndAction(Action);
+		GetActionStack()->EndAction(Action);
 	}
 
 	return true;
@@ -7682,9 +7680,9 @@ bool URigVMController::SetPinDefaultValue(URigVMPin* InPin, const FString& InDef
 	FRigVMSetPinDefaultValueAction Action;
 	if (bSetupUndoRedo)
 	{
-		Action = FRigVMSetPinDefaultValueAction(InPin, InDefaultValue);
-		Action.Title = FString::Printf(TEXT("Set Pin Default Value"));
-		ActionStack->BeginAction(Action);
+		Action = FRigVMSetPinDefaultValueAction(this, InPin, InDefaultValue);
+		Action.SetTitle(FString::Printf(TEXT("Set Pin Default Value")));
+		GetActionStack()->BeginAction(Action);
 	}
 
 	const FString ClampedDefaultValue = InPin->IsRootPin() ? InPin->ClampDefaultValueFromMetaData(InDefaultValue) : InDefaultValue;
@@ -7763,7 +7761,7 @@ bool URigVMController::SetPinDefaultValue(URigVMPin* InPin, const FString& InDef
 
 	if (bSetupUndoRedo)
 	{
-		ActionStack->EndAction(Action, bMergeUndoAction);
+		GetActionStack()->EndAction(Action, bMergeUndoAction);
 	}
 
 	return true;
@@ -8039,8 +8037,8 @@ FString URigVMController::AddAggregatePin(URigVMNode* InNode, const FString& InP
 	FRigVMReplaceNodesAction Action;
 	if(bSetupUndoRedo)
 	{
-		Action = FRigVMReplaceNodesAction({InNode}, this);
-		Action.Title = TEXT("Add aggregate pin");
+		Action = FRigVMReplaceNodesAction(this, {InNode});
+		Action.SetTitle(TEXT("Add aggregate pin"));
 	}
 
 	if (AggregateNode == nullptr)
@@ -8311,8 +8309,8 @@ FString URigVMController::AddAggregatePin(URigVMNode* InNode, const FString& InP
 
 	if (bSetupUndoRedo)
 	{
-		Action.StoreNode(AggregateNode, this, false);
-		ActionStack->AddAction(Action);
+		Action.StoreNode(AggregateNode, false);
+		GetActionStack()->AddAction(Action);
 	}
 
 	if (bPrintPythonCommand)
@@ -8381,8 +8379,8 @@ bool URigVMController::RemoveAggregatePin(URigVMPin* InPin, bool bSetupUndoRedo,
 	FRigVMReplaceNodesAction Action;
 	if(bSetupUndoRedo)
 	{
-		Action = FRigVMReplaceNodesAction({InPin->GetNode()}, this);
-		Action.Title = TEXT("Remove aggregate pin");
+		Action = FRigVMReplaceNodesAction(this, {InPin->GetNode()});
+		Action.SetTitle(TEXT("Remove aggregate pin"));
 	}
 
 	URigVMNode* Node = InPin->GetNode();
@@ -8534,8 +8532,8 @@ bool URigVMController::RemoveAggregatePin(URigVMPin* InPin, bool bSetupUndoRedo,
 	{
 		if (bSuccess)
 		{
-			Action.StoreNode(Node, this, false);
-			ActionStack->AddAction(Action);
+			Action.StoreNode(Node, false);
+			GetActionStack()->AddAction(Action);
 		}
 	}
 
@@ -8672,9 +8670,9 @@ URigVMPin* URigVMController::InsertArrayPin(URigVMPin* ArrayPin, int32 InIndex,
 	FRigVMInsertArrayPinAction Action;
 	if (bSetupUndoRedo)
 	{
-		Action = FRigVMInsertArrayPinAction(ArrayPin, InIndex, InDefaultValue);
-		Action.Title = FString::Printf(TEXT("Insert Array Pin"));
-		ActionStack->BeginAction(Action);
+		Action = FRigVMInsertArrayPinAction(this, ArrayPin, InIndex, InDefaultValue);
+		Action.SetTitle(FString::Printf(TEXT("Insert Array Pin")));
+		GetActionStack()->BeginAction(Action);
 	}
 
 	for (int32 ExistingIndex = ArrayPin->GetSubPins().Num() - 1; ExistingIndex >= InIndex; ExistingIndex--)
@@ -8726,7 +8724,7 @@ URigVMPin* URigVMController::InsertArrayPin(URigVMPin* ArrayPin, int32 InIndex,
 
 	if (bSetupUndoRedo)
 	{
-		ActionStack->EndAction(Action);
+		GetActionStack()->EndAction(Action);
 	}
 
 	return Pin;
@@ -8768,9 +8766,9 @@ bool URigVMController::RemoveArrayPin(const FString& InArrayElementPinPath, bool
 	FRigVMRemoveArrayPinAction Action;
 	if (bSetupUndoRedo)
 	{
-		Action = FRigVMRemoveArrayPinAction(ArrayElementPin);
-		Action.Title = FString::Printf(TEXT("Remove Array Pin"));
-		ActionStack->BeginAction(Action);
+		Action = FRigVMRemoveArrayPinAction(this, ArrayElementPin);
+		Action.SetTitle(FString::Printf(TEXT("Remove Array Pin")));
+		GetActionStack()->BeginAction(Action);
 	}
 
 	// we need to keep at least one element for fixed size arrays
@@ -8803,7 +8801,7 @@ bool URigVMController::RemoveArrayPin(const FString& InArrayElementPinPath, bool
 
 	if (bSetupUndoRedo)
 	{
-		ActionStack->EndAction(Action);
+		GetActionStack()->EndAction(Action);
 	}
 
 	if (bPrintPythonCommand)
@@ -8899,11 +8897,11 @@ bool URigVMController::SetArrayPinSize(const FString& InArrayPinPath, int32 InSi
 	}
 
 	FRigVMControllerCompileBracketScope CompileScope(this);
-	FRigVMBaseAction Action;
+	FRigVMBaseAction Action(this);
 	if (bSetupUndoRedo)
 	{
-		Action.Title = FString::Printf(TEXT("Set Array Pin Size (%d)"), InSize);
-		ActionStack->BeginAction(Action);
+		Action.SetTitle(FString::Printf(TEXT("Set Array Pin Size (%d)"), InSize));
+		GetActionStack()->BeginAction(Action);
 	}
 
 	InSize = FMath::Max<int32>(InSize, 0);
@@ -8926,7 +8924,7 @@ bool URigVMController::SetArrayPinSize(const FString& InArrayPinPath, int32 InSi
 		{
 			if (bSetupUndoRedo)
 			{
-				ActionStack->CancelAction(Action, this);
+				GetActionStack()->CancelAction(Action);
 			}
 			return false;
 		}
@@ -8939,7 +8937,7 @@ bool URigVMController::SetArrayPinSize(const FString& InArrayPinPath, int32 InSi
 		{
 			if (bSetupUndoRedo)
 			{
-				ActionStack->CancelAction(Action, this);
+				GetActionStack()->CancelAction(Action);
 			}
 			return false;
 		}
@@ -8955,11 +8953,11 @@ bool URigVMController::SetArrayPinSize(const FString& InArrayPinPath, int32 InSi
 	{
 		if (RemovedPins > 0 || AddedPins > 0)
 		{
-			ActionStack->EndAction(Action);
+			GetActionStack()->EndAction(Action);
 		}
 		else
 		{
-			ActionStack->CancelAction(Action, this);
+			GetActionStack()->CancelAction(Action);
 		}
 	}
 
@@ -9067,11 +9065,11 @@ bool URigVMController::BindPinToVariable(URigVMPin* InPin, const FString& InNewB
 	}
 	
 	FRigVMControllerCompileBracketScope CompileScope(this);
-	FRigVMBaseAction Action;
+	FRigVMBaseAction Action(this);
 	if (bSetupUndoRedo)
 	{
-		Action.Title = TEXT("Bind pin to variable");
-		ActionStack->BeginAction(Action);
+		Action.SetTitle(TEXT("Bind pin to variable"));
+		GetActionStack()->BeginAction(Action);
 	}
 
 	// Unbind any other variables, remove any other injections, and break all links to the input pin
@@ -9102,7 +9100,7 @@ bool URigVMController::BindPinToVariable(URigVMPin* InPin, const FString& InNewB
 		{
 			if (bSetupUndoRedo)
 			{
-				ActionStack->CancelAction(Action, this);
+				GetActionStack()->CancelAction(Action);
 			}
 			return false;
 		}
@@ -9122,7 +9120,7 @@ bool URigVMController::BindPinToVariable(URigVMPin* InPin, const FString& InNewB
 			{
 				if (bSetupUndoRedo)
 				{
-					ActionStack->CancelAction(Action, this);
+					GetActionStack()->CancelAction(Action);
 				}
 				return false;
 			}
@@ -9134,14 +9132,14 @@ bool URigVMController::BindPinToVariable(URigVMPin* InPin, const FString& InNewB
 	{
 		if (bSetupUndoRedo)
 		{
-			ActionStack->CancelAction(Action, this);
+			GetActionStack()->CancelAction(Action);
 		}
 		return false;
 	}
 	
 	if (bSetupUndoRedo)
 	{
-		ActionStack->EndAction(Action);
+		GetActionStack()->EndAction(Action);
 	}
 
 	return true;
@@ -9209,18 +9207,18 @@ bool URigVMController::UnbindPinFromVariable(URigVMPin* InPin, bool bSetupUndoRe
 	}
 
 	FRigVMControllerCompileBracketScope CompileScope(this);
-	FRigVMBaseAction Action;
+	FRigVMBaseAction Action(this);
 	if (bSetupUndoRedo)
 	{
-		Action.Title = TEXT("Unbind pin from variable");
-		ActionStack->BeginAction(Action);
+		Action.SetTitle(TEXT("Unbind pin from variable"));
+		GetActionStack()->BeginAction(Action);
 	}
 
 	RemoveInjectedNode(InPin->GetPinPath(), true, bSetupUndoRedo);
 
 	if (bSetupUndoRedo)
 	{
-		ActionStack->EndAction(Action);
+		GetActionStack()->EndAction(Action);
 	}
 
 	return true;
@@ -9606,11 +9604,11 @@ bool URigVMController::AddLink(URigVMPin* OutputPin, URigVMPin* InputPin, bool b
 	ensure(!InputPin->IsLinkedTo(OutputPin));
 
 	FRigVMControllerCompileBracketScope CompileScope(this);
-	FRigVMBaseAction Action;
+	FRigVMBaseAction Action(this);
 	if (bSetupUndoRedo)
 	{
-		Action.Title = TEXT("Add Link");
-		ActionStack->BeginAction(Action);
+		Action.SetTitle(TEXT("Add Link"));
+		GetActionStack()->BeginAction(Action);
 	}
 
 	// check if we need to inject cast node
@@ -9636,7 +9634,7 @@ bool URigVMController::AddLink(URigVMPin* OutputPin, URigVMPin* InputPin, bool b
 				{
 					if (bSetupUndoRedo)
 					{
-						ActionStack->CancelAction(Action, this);
+						GetActionStack()->CancelAction(Action);
 					}
 					return false;
 				}
@@ -9675,7 +9673,7 @@ bool URigVMController::AddLink(URigVMPin* OutputPin, URigVMPin* InputPin, bool b
 				{
 					if (bSetupUndoRedo)
 					{
-						ActionStack->CancelAction(Action, this);
+						GetActionStack()->CancelAction(Action);
 					}
 					return false;
 				}
@@ -9691,7 +9689,7 @@ bool URigVMController::AddLink(URigVMPin* OutputPin, URigVMPin* InputPin, bool b
 					{
 						if (bSetupUndoRedo)
 						{
-							ActionStack->CancelAction(Action, this);
+							GetActionStack()->CancelAction(Action);
 						}
 						return false;
 					}
@@ -9703,7 +9701,7 @@ bool URigVMController::AddLink(URigVMPin* OutputPin, URigVMPin* InputPin, bool b
 					{
 						if (bSetupUndoRedo)
 						{
-							ActionStack->CancelAction(Action, this);
+							GetActionStack()->CancelAction(Action);
 						}
 						return false;
 					}
@@ -9711,8 +9709,8 @@ bool URigVMController::AddLink(URigVMPin* OutputPin, URigVMPin* InputPin, bool b
 
 				if(bSetupUndoRedo)
 				{
-					Action.Title = TEXT("Add Link with Cast");
-					ActionStack->EndAction(Action);
+					Action.SetTitle(TEXT("Add Link with Cast"));
+					GetActionStack()->EndAction(Action);
 				}
 				return true;
 			}
@@ -9792,7 +9790,7 @@ bool URigVMController::AddLink(URigVMPin* OutputPin, URigVMPin* InputPin, bool b
 		{
 			if (bSetupUndoRedo)
 			{
-				ActionStack->CancelAction(Action, this);
+				GetActionStack()->CancelAction(Action);
 			}
 			return false;
 		}
@@ -9800,7 +9798,7 @@ bool URigVMController::AddLink(URigVMPin* OutputPin, URigVMPin* InputPin, bool b
 
 	if (bSetupUndoRedo)
 	{
-		ActionStack->AddAction(FRigVMAddLinkAction(OutputPin, InputPin));
+		GetActionStack()->AddAction(FRigVMAddLinkAction(this, OutputPin, InputPin));
 	}
 
 	URigVMLink* Link = nullptr;
@@ -9894,7 +9892,7 @@ bool URigVMController::AddLink(URigVMPin* OutputPin, URigVMPin* InputPin, bool b
 		}
 #endif
 		
-		ActionStack->EndAction(Action);
+		GetActionStack()->EndAction(Action);
 	}
 
 	if (!bIsTransacting)
@@ -10049,9 +10047,9 @@ bool URigVMController::BreakLink(URigVMPin* OutputPin, URigVMPin* InputPin, bool
 			FRigVMBreakLinkAction Action;
 			if (bSetupUndoRedo)
 			{
-				Action = FRigVMBreakLinkAction(OutputPin, InputPin);
-				Action.Title = FString::Printf(TEXT("Break Link"));
-				ActionStack->BeginAction(Action);
+				Action = FRigVMBreakLinkAction(this, OutputPin, InputPin);
+				Action.SetTitle(FString::Printf(TEXT("Break Link")));
+				GetActionStack()->BeginAction(Action);
 			}
 
 			OutputPin->Links.Remove(Link);
@@ -10079,7 +10077,7 @@ bool URigVMController::BreakLink(URigVMPin* OutputPin, URigVMPin* InputPin, bool
 
 			if (bSetupUndoRedo)
 			{
-				ActionStack->EndAction(Action);
+				GetActionStack()->EndAction(Action);
 			}
 
 			return true;
@@ -10144,11 +10142,11 @@ bool URigVMController::BreakAllLinks(URigVMPin* Pin, bool bAsInput, bool bSetupU
 	}
 	
 	FRigVMControllerCompileBracketScope CompileScope(this);
-	FRigVMBaseAction Action;
+	FRigVMBaseAction Action(this);
 	if (bSetupUndoRedo)
 	{
-		Action.Title = FString::Printf(TEXT("Break All Links"));
-		ActionStack->BeginAction(Action);
+		Action.SetTitle(FString::Printf(TEXT("Break All Links")));
+		GetActionStack()->BeginAction(Action);
 	}
 
 	int32 LinksBroken = 0;
@@ -10178,11 +10176,11 @@ bool URigVMController::BreakAllLinks(URigVMPin* Pin, bool bAsInput, bool bSetupU
 	{
 		if (LinksBroken > 0)
 		{
-			ActionStack->EndAction(Action);
+			GetActionStack()->EndAction(Action);
 		}
 		else
 		{
-			ActionStack->CancelAction(Action, this);
+			GetActionStack()->CancelAction(Action);
 		}
 	}
 
@@ -10332,10 +10330,10 @@ FName URigVMController::AddExposedPin(const FName& InPinName, ERigVMPinDirection
 	}
 
 	FRigVMControllerCompileBracketScope CompileScope(this);
-	FRigVMAddExposedPinAction Action(Pin);
+	FRigVMAddExposedPinAction Action(this, Pin);
 	if (bSetupUndoRedo)
 	{
-		ActionStack->BeginAction(Action);
+		GetActionStack()->BeginAction(Action);
 	}
 
 	{
@@ -10383,7 +10381,7 @@ FName URigVMController::AddExposedPin(const FName& InPinName, ERigVMPinDirection
 
 	if (bSetupUndoRedo)
 	{
-		ActionStack->EndAction(Action);
+		GetActionStack()->EndAction(Action);
 	}
 
 	if (bPrintPythonCommand)
@@ -10456,10 +10454,10 @@ bool URigVMController::RemoveExposedPin(const FName& InPinName, bool bSetupUndoR
 	}
 
 	FRigVMControllerCompileBracketScope CompileScope(this);
-	FRigVMRemoveExposedPinAction Action(Pin);
+	FRigVMRemoveExposedPinAction Action(this, Pin);
 	if (bSetupUndoRedo)
 	{
-		ActionStack->BeginAction(Action);
+		GetActionStack()->BeginAction(Action);
 	}
 
 	bool bSuccessfullyRemovedPin = false;
@@ -10494,11 +10492,11 @@ bool URigVMController::RemoveExposedPin(const FName& InPinName, bool bSetupUndoR
 	{
 		if (bSuccessfullyRemovedPin)
 		{
-			ActionStack->EndAction(Action);
+			GetActionStack()->EndAction(Action);
 		}
 		else
 		{
-			ActionStack->CancelAction(Action, this);
+			GetActionStack()->CancelAction(Action);
 		}
 	}
 
@@ -10585,8 +10583,8 @@ bool URigVMController::RenameExposedPin(const FName& InOldPinName, const FName&
 	FRigVMRenameExposedPinAction Action;
 	if (bSetupUndoRedo)
 	{
-		Action = FRigVMRenameExposedPinAction(Pin->GetFName(), PinName);
-		ActionStack->BeginAction(Action);
+		Action = FRigVMRenameExposedPinAction(this, Pin->GetFName(), PinName);
+		GetActionStack()->BeginAction(Action);
 	}
 
 	struct Local
@@ -10641,7 +10639,7 @@ bool URigVMController::RenameExposedPin(const FName& InOldPinName, const FName&
 
 	if (!Local::RenamePin(this, Pin, PinName))
 	{
-		ActionStack->CancelAction(Action, this);
+		GetActionStack()->CancelAction(Action);
 		return false;
 	}
 
@@ -10682,7 +10680,7 @@ bool URigVMController::RenameExposedPin(const FName& InOldPinName, const FName&
 
 	if (bSetupUndoRedo)
 	{
-		ActionStack->EndAction(Action);
+		GetActionStack()->EndAction(Action);
 	}
 
 	if (bPrintPythonCommand)
@@ -10802,11 +10800,11 @@ bool URigVMController::ChangeExposedPinType(const FName& InPinName, const FStrin
 	}
 
 	FRigVMControllerCompileBracketScope CompileScope(this);
-	FRigVMBaseAction Action;
+	FRigVMBaseAction Action(this);
 	if (bSetupUndoRedo)
 	{
-		Action.Title = FString::Printf(TEXT("Change Exposed Pin Type"));
-		ActionStack->BeginAction(Action);
+		Action.SetTitle(FString::Printf(TEXT("Change Exposed Pin Type")));
+		GetActionStack()->BeginAction(Action);
 	}
 
 	FRigVMRegistry& Registry = FRigVMRegistry::Get();
@@ -10869,7 +10867,7 @@ bool URigVMController::ChangeExposedPinType(const FName& InPinName, const FStrin
 		{
 			if (bSetupUndoRedo)
 			{
-				ActionStack->CancelAction(Action, this);
+				GetActionStack()->CancelAction(Action);
 			}
 			return false;
 		}
@@ -10922,7 +10920,7 @@ bool URigVMController::ChangeExposedPinType(const FName& InPinName, const FStrin
 
 	if (bSetupUndoRedo)
 	{
-		ActionStack->EndAction(Action);
+		GetActionStack()->EndAction(Action);
 	}
 
 	if (bPrintPythonCommand)
@@ -10990,7 +10988,7 @@ bool URigVMController::SetExposedPinIndex(const FName& InPinName, int32 InNewInd
 
 	FRigVMControllerCompileBracketScope CompileBracketScope(this);
 
-	FRigVMSetPinIndexAction PinIndexAction(Pin, InNewIndex);
+	FRigVMSetPinIndexAction PinIndexAction(this, Pin, InNewIndex);
 	{
 		LibraryNode->Pins.Remove(Pin);
 		LibraryNode->Pins.Insert(Pin, InNewIndex);
@@ -11005,7 +11003,7 @@ bool URigVMController::SetExposedPinIndex(const FName& InPinName, int32 InNewInd
 	
 	if (bSetupUndoRedo)
 	{
-		ActionStack->AddAction(PinIndexAction);
+		GetActionStack()->AddAction(PinIndexAction);
 	}
 
 	if (bPrintPythonCommand)
@@ -11099,9 +11097,9 @@ URigVMFunctionReferenceNode* URigVMController::AddFunctionReferenceNodeFromDescr
 
 	if (bSetupUndoRedo)
 	{
-		FRigVMImportFromTextAction Action(FunctionRefNode, this);
-		Action.Title = TEXT("Add function node");
-		ActionStack->AddAction(Action);
+		FRigVMImportFromTextAction Action(this, FunctionRefNode);
+		Action.SetTitle(TEXT("Add function node"));
+		GetActionStack()->AddAction(Action);
 	}
 
 	if (bPrintPythonCommand)
@@ -11276,7 +11274,7 @@ bool URigVMController::SetRemappedVariable(URigVMFunctionReferenceNode* InFuncti
 	FRigVMControllerCompileBracketScope CompileScope(this);
 	if(bSetupUndoRedo)
 	{
-		ActionStack->AddAction(FRigVMSetRemappedVariableAction(InFunctionRefNode, InInnerVariableName, OldOuterVariableName, InOuterVariableName));
+		GetActionStack()->AddAction(FRigVMSetRemappedVariableAction(this, InFunctionRefNode, InInnerVariableName, OldOuterVariableName, InOuterVariableName));
 	}
 	
 	return true;
@@ -11369,9 +11367,9 @@ URigVMLibraryNode* URigVMController::AddFunctionToLibrary(const FName& InFunctio
 
 	if (bSetupUndoRedo)
 	{
-		FRigVMImportFromTextAction Action(CollapseNode, this);
-		Action.Title = TEXT("Add function to library");
-		ActionStack->AddAction(Action);
+		FRigVMImportFromTextAction Action(this, CollapseNode);
+		Action.SetTitle(TEXT("Add function to library"));
+		GetActionStack()->AddAction(Action);
 	}
 
 	if (bPrintPythonCommand)
@@ -11485,11 +11483,11 @@ bool URigVMController::MarkFunctionAsPublic(const FName& InFunctionName, bool bI
 
 		if (bSetupUndoRedo)
 		{
-			FRigVMBaseAction BaseAction;
-			BaseAction.Title = FString::Printf(TEXT("Mark function %s as %s"), *InFunctionName.ToString(), (bInIsPublic) ? TEXT("Public") : TEXT("Private"));
-			ActionStack->BeginAction(BaseAction);
-			ActionStack->AddAction(FRigVMMarkFunctionPublicAction(InFunctionName, bInIsPublic));
-			ActionStack->EndAction(BaseAction);
+			FRigVMBaseAction BaseAction(this);
+			BaseAction.SetTitle(FString::Printf(TEXT("Mark function %s as %s"), *InFunctionName.ToString(), (bInIsPublic) ? TEXT("Public") : TEXT("Private")));
+			GetActionStack()->BeginAction(BaseAction);
+			GetActionStack()->AddAction(FRigVMMarkFunctionPublicAction(this, InFunctionName, bInIsPublic));
+			GetActionStack()->EndAction(BaseAction);
 		}
 
 		if (bInIsPublic)
@@ -11621,9 +11619,9 @@ FRigVMGraphVariableDescription URigVMController::AddLocalVariable(const FName& I
 
 	if (bSetupUndoRedo)
 	{
-		FRigVMAddLocalVariableAction Action(NewVariable);
-		Action.Title = FString::Printf(TEXT("Add Local Variable %s"), *InVariableName.ToString());
-		ActionStack->AddAction(Action);
+		FRigVMAddLocalVariableAction Action(this, NewVariable);
+		Action.SetTitle(FString::Printf(TEXT("Add Local Variable %s"), *InVariableName.ToString()));
+		GetActionStack()->AddAction(Action);
 	}
 
 	if (!bSuspendNotifications)
@@ -11703,11 +11701,11 @@ bool URigVMController::RemoveLocalVariable(const FName& InVariableName, bool bSe
 	if (FoundIndex != INDEX_NONE)
 	{
 		FRigVMControllerCompileBracketScope CompileScope(this);
-		FRigVMBaseAction BaseAction;
+		FRigVMBaseAction BaseAction(this);
 		if (bSetupUndoRedo)
 		{
-			BaseAction.Title = FString::Printf(TEXT("Remove Local Variable %s"), *InVariableName.ToString());
-			ActionStack->BeginAction(BaseAction);			
+			BaseAction.SetTitle(FString::Printf(TEXT("Remove Local Variable %s"), *InVariableName.ToString()));
+			GetActionStack()->BeginAction(BaseAction);			
 		}	
 		
 		const FString VarNameStr = InVariableName.ToString();
@@ -11788,13 +11786,13 @@ bool URigVMController::RemoveLocalVariable(const FName& InVariableName, bool bSe
 
 		if (bSetupUndoRedo)
 		{
-			ActionStack->AddAction(FRigVMRemoveLocalVariableAction(LocalVariables[FoundIndex]));
+			GetActionStack()->AddAction(FRigVMRemoveLocalVariableAction(this, LocalVariables[FoundIndex]));
 		}
 		LocalVariables.RemoveAt(FoundIndex);
 
 		if (bSetupUndoRedo)
 		{
-			ActionStack->EndAction(BaseAction);
+			GetActionStack()->EndAction(BaseAction);
 		}
 
 		if (bPrintPythonCommand)
@@ -11859,12 +11857,12 @@ bool URigVMController::RenameLocalVariable(const FName& InVariableName, const FN
 	FRigVMControllerCompileBracketScope CompileScope(this);
 	if (bSetupUndoRedo)
 	{
-		FRigVMBaseAction BaseAction;
-		BaseAction.Title = FString::Printf(TEXT("Rename Local Variable %s to %s"), *InVariableName.ToString(), *InNewVariableName.ToString());
+		FRigVMBaseAction BaseAction(this);
+		BaseAction.SetTitle(FString::Printf(TEXT("Rename Local Variable %s to %s"), *InVariableName.ToString(), *InNewVariableName.ToString()));
 
-		ActionStack->BeginAction(BaseAction);
-		ActionStack->AddAction(FRigVMRenameLocalVariableAction(LocalVariables[FoundIndex].Name, InNewVariableName));
-		ActionStack->EndAction(BaseAction);
+		GetActionStack()->BeginAction(BaseAction);
+		GetActionStack()->AddAction(FRigVMRenameLocalVariableAction(this, LocalVariables[FoundIndex].Name, InNewVariableName));
+		GetActionStack()->EndAction(BaseAction);
 	}
 	
 	LocalVariables[FoundIndex].Name = InNewVariableName;
@@ -11943,13 +11941,13 @@ bool URigVMController::SetLocalVariableType(const FName& InVariableName, const F
 	}
 
 	FRigVMControllerCompileBracketScope CompileScope(this);
-	FRigVMBaseAction BaseAction;
+	FRigVMBaseAction BaseAction(this);
 	if (bSetupUndoRedo)
 	{
-		BaseAction.Title = FString::Printf(TEXT("Change Local Variable type %s to %s"), *InVariableName.ToString(), *InCPPType);
-		ActionStack->BeginAction(BaseAction);
+		BaseAction.SetTitle(FString::Printf(TEXT("Change Local Variable type %s to %s"), *InVariableName.ToString(), *InCPPType));
+		GetActionStack()->BeginAction(BaseAction);
 
-		ActionStack->AddAction(FRigVMChangeLocalVariableTypeAction(LocalVariables[FoundIndex], InCPPType, InCPPTypeObject));
+		GetActionStack()->AddAction(FRigVMChangeLocalVariableTypeAction(this, LocalVariables[FoundIndex], InCPPType, InCPPTypeObject));
 	}	
 	
 	LocalVariables[FoundIndex].CPPType = InCPPType;
@@ -11995,7 +11993,7 @@ bool URigVMController::SetLocalVariableType(const FName& InVariableName, const F
 
 	if (bSetupUndoRedo)
 	{
-		ActionStack->EndAction(BaseAction);
+		GetActionStack()->EndAction(BaseAction);
 	}
 
 	if (bPrintPythonCommand)
@@ -12064,9 +12062,9 @@ bool URigVMController::SetLocalVariableDefaultValue(const FName& InVariableName,
 	FRigVMControllerCompileBracketScope CompileScope(this);
 	if (bSetupUndoRedo)
 	{
-		FRigVMChangeLocalVariableDefaultValueAction Action(LocalVariables[FoundIndex], InDefaultValue);
-		Action.Title = FString::Printf(TEXT("Change Local Variable %s default value"), *InVariableName.ToString());
-		ActionStack->AddAction(Action);
+		FRigVMChangeLocalVariableDefaultValueAction Action(this, LocalVariables[FoundIndex], InDefaultValue);
+		Action.SetTitle(FString::Printf(TEXT("Change Local Variable %s default value"), *InVariableName.ToString()));
+		GetActionStack()->AddAction(Action);
 	}
 
 	FRigVMGraphVariableDescription& VariableDescription = LocalVariables[FoundIndex];
@@ -12150,20 +12148,20 @@ bool URigVMController::PerformUserWorkflow(const FRigVMUserWorkflow& InWorkflow,
 	}
 
 
-	FRigVMBaseAction Bracket;
-	Bracket.Title = InWorkflow.GetTitle();
-	ActionStack->BeginAction(Bracket);
+	FRigVMBaseAction Bracket(this);
+	Bracket.SetTitle(InWorkflow.GetTitle());
+	GetActionStack()->BeginAction(Bracket);
 
 	const bool bSuccess = InWorkflow.Perform(InOptions, this);
 
-	ActionStack->EndAction(Bracket);
+	GetActionStack()->EndAction(Bracket);
 
 	if(!bSuccess)
 	{
 		// if the workflow was run as the top level action we'll undo
-		if(ActionStack->CurrentActions.IsEmpty())
+		if(GetActionStack()->CurrentActions.IsEmpty())
 		{
-			ActionStack->Undo(this);
+			GetActionStack()->Undo(this);
 		}
 	}
 
@@ -12436,11 +12434,11 @@ URigVMRerouteNode* URigVMController::AddFreeRerouteNode(bool bShowAsFullNode, co
 	}
 
 	FRigVMControllerCompileBracketScope CompileScope(this);
-	FRigVMBaseAction Action;
+	FRigVMBaseAction Action(this);
 	if (bSetupUndoRedo)
 	{
-		Action.Title = FString::Printf(TEXT("Add Reroute"));
-		ActionStack->BeginAction(Action);
+		Action.SetTitle(FString::Printf(TEXT("Add Reroute")));
+		GetActionStack()->BeginAction(Action);
 	}
 
 	FString Name = GetValidNodeName(InNodeName.IsEmpty() ? FString(TEXT("RerouteNode")) : InNodeName);
@@ -12473,14 +12471,14 @@ URigVMRerouteNode* URigVMController::AddFreeRerouteNode(bool bShowAsFullNode, co
 
 	if (bSetupUndoRedo)
 	{
-		ActionStack->AddAction(FRigVMImportFromTextAction(Node, this));
+		GetActionStack()->AddAction(FRigVMImportFromTextAction(this, Node));
 	}
 
 	Notify(ERigVMGraphNotifType::NodeAdded, Node);
 
 	if (bSetupUndoRedo)
 	{
-		ActionStack->EndAction(Action);
+		GetActionStack()->EndAction(Action);
 	}
 
 	return Node;
@@ -12712,19 +12710,19 @@ URigVMTemplateNode* URigVMController::AddTemplateNode(const FName& InNotation, c
 
 	Notify(ERigVMGraphNotifType::NodeAdded, Node);
 
-	FRigVMBaseAction Action;
+	FRigVMBaseAction Action(this);
 	FRigVMControllerCompileBracketScope CompileScope(this);
 	if (bSetupUndoRedo)
 	{
-		ActionStack->BeginAction(Action);
-		ActionStack->AddAction(FRigVMImportFromTextAction(Node, this));
+		GetActionStack()->BeginAction(Action);
+		GetActionStack()->AddAction(FRigVMImportFromTextAction(this, Node));
 	}
 	
 	ResolveTemplateNodeMetaData(Node, bSetupUndoRedo);
 
 	if (bSetupUndoRedo)
 	{
-		ActionStack->EndAction(Action);
+		GetActionStack()->EndAction(Action);
 	}
 
 	if (bPrintPythonCommand)
@@ -12864,7 +12862,7 @@ URigVMEnumNode* URigVMController::AddEnumNode(const FName& InCPPTypeObjectPath,
 	FRigVMControllerCompileBracketScope CompileScope(this);
 	if (bSetupUndoRedo)
 	{
-		ActionStack->AddAction(FRigVMImportFromTextAction(Node, this));
+		GetActionStack()->AddAction(FRigVMImportFromTextAction(this, Node));
 	}
 
 	if (bPrintPythonCommand)
@@ -13073,9 +13071,9 @@ URigVMInvokeEntryNode* URigVMController::AddInvokeEntryNode(const FName& InEntry
 
 	if (bSetupUndoRedo)
 	{
-		FRigVMImportFromTextAction Action(Node, this);
-		Action.Title = FString::Printf(TEXT("Add Invoke %s Entry"), *InEntryName.ToString());
-		ActionStack->AddAction(Action);
+		FRigVMImportFromTextAction Action(this, Node);
+		Action.SetTitle(FString::Printf(TEXT("Add Invoke %s Entry"), *InEntryName.ToString()));
+		GetActionStack()->AddAction(Action);
 	}
 
 	if (bPrintPythonCommand)
@@ -16210,11 +16208,11 @@ bool URigVMController::ResolveWildCardPin(URigVMPin* InPin, TRigVMTypeIndex InTy
 		return false;
 	}
 	
-	FRigVMBaseAction Action;
+	FRigVMBaseAction Action(this);
 	if (bSetupUndoRedo)
 	{
-		Action.Title = TEXT("Resolve Wildcard Pin");
-		ActionStack->BeginAction(Action);
+		Action.SetTitle(TEXT("Resolve Wildcard Pin"));
+		GetActionStack()->BeginAction(Action);
 	}
 
 	if (!InPin->IsWildCard())
@@ -16229,7 +16227,7 @@ bool URigVMController::ResolveWildCardPin(URigVMPin* InPin, TRigVMTypeIndex InTy
 	{
 		if (bSetupUndoRedo)
 		{
-			ActionStack->CancelAction(Action, this);
+			GetActionStack()->CancelAction(Action);
 		}
 		return false;
 	}
@@ -16238,7 +16236,7 @@ bool URigVMController::ResolveWildCardPin(URigVMPin* InPin, TRigVMTypeIndex InTy
 
 	if (bSetupUndoRedo)
 	{
-		ActionStack->EndAction(Action);
+		GetActionStack()->EndAction(Action);
 	}
 	
 	return true;
@@ -16782,11 +16780,11 @@ bool URigVMController::ChangePinType(URigVMPin* InPin, TRigVMTypeIndex InTypeInd
 	}
 
 	FRigVMControllerCompileBracketScope CompileScope(this);
-	FRigVMBaseAction Action;
+	FRigVMBaseAction Action(this);
 	if (bSetupUndoRedo)
 	{
-		Action.Title = TEXT("Change pin type");
-		ActionStack->BeginAction(Action);
+		Action.SetTitle(TEXT("Change pin type"));
+		GetActionStack()->BeginAction(Action);
 	}
 
 	TArray<FLinkedPath> LinkedPaths;
@@ -16839,7 +16837,7 @@ bool URigVMController::ChangePinType(URigVMPin* InPin, TRigVMTypeIndex InTypeInd
 
 	if (bSetupUndoRedo)
 	{
-		ActionStack->AddAction(FRigVMChangePinTypeAction(InPin, TypeIndex, bSetupOrphanPins, bBreakLinks, bRemoveSubPins));
+		GetActionStack()->AddAction(FRigVMChangePinTypeAction(this, InPin, TypeIndex, bSetupOrphanPins, bBreakLinks, bRemoveSubPins));
 	}
 	
 	// compute the number of remaining wildcard pins
@@ -17065,7 +17063,7 @@ bool URigVMController::ChangePinType(URigVMPin* InPin, TRigVMTypeIndex InTypeInd
 
 	if (bSetupUndoRedo)
 	{
-		ActionStack->EndAction(Action);
+		GetActionStack()->EndAction(Action);
 	}
 
 	if(!LinkedPaths.IsEmpty())
@@ -17547,7 +17545,7 @@ bool URigVMController::BreakLinkedPaths(const TArray<FLinkedPath>& InLinkedPaths
 
 				if(bSetupUndoRedo && SourcePin && TargetPin)
 				{
-					ActionStack->AddAction(FRigVMBreakLinkAction(SourcePin, TargetPin));
+					GetActionStack()->AddAction(FRigVMBreakLinkAction(this, SourcePin, TargetPin));
 				}
 			}
 			else
@@ -18356,6 +18354,37 @@ void URigVMController::PostDuplicateHost(const FString& InOldPathName, const FSt
 
 #endif
 
+URigVMActionStack* URigVMController::GetActionStack() const
+{
+	if(WeakActionStack.IsValid())
+	{
+		return WeakActionStack.Get();
+	}
+	checkNoEntry();
+	return nullptr;
+}
+
+void URigVMController::SetActionStack(URigVMActionStack* InActionStack)
+{
+	if(URigVMActionStack* PreviousActionStack = WeakActionStack.Get())
+	{
+		PreviousActionStack->OnModified().Remove(ActionStackHandle);
+		ActionStackHandle.Reset();
+	}
+
+	WeakActionStack = InActionStack;
+
+	if(URigVMActionStack* NewActionStack = WeakActionStack.Get())
+	{
+		ActionStackHandle = NewActionStack->OnModified().AddLambda([&](ERigVMGraphNotifType NotifType, URigVMGraph* InGraph, UObject* InSubject) -> void {
+			if(InGraph == GetGraph())
+			{
+				Notify(NotifType, InSubject);
+			}
+		});
+	}
+}
+
 URigVMControllerSettings::URigVMControllerSettings(const FObjectInitializer& Initializer)
 	: Super(Initializer)
 {
diff --git a/Engine/Source/Developer/RigVMDeveloper/Private/RigVMModel/RigVMControllerActions.cpp b/Engine/Source/Developer/RigVMDeveloper/Private/RigVMModel/RigVMControllerActions.cpp
--- a/Engine/Source/Developer/RigVMDeveloper/Private/RigVMModel/RigVMControllerActions.cpp
+++ b/Engine/Source/Developer/RigVMDeveloper/Private/RigVMModel/RigVMControllerActions.cpp
@@ -38,9 +38,20 @@ const UScriptStruct* FRigVMActionWrapper::GetScriptStruct() const
 	return CastChecked<UScriptStruct>(StructOnScope->GetStruct());
 }
 
-FRigVMBaseAction* FRigVMActionWrapper::GetAction() const
+FRigVMBaseAction* FRigVMActionWrapper::GetAction(URigVMController* InLastController) const
 {
-	return (FRigVMBaseAction*)StructOnScope->GetStructMemory();
+	FRigVMBaseAction* Action = (FRigVMBaseAction*)StructOnScope->GetStructMemory();
+
+	// reuse the last controller so that we don't have to rely on resolving it
+	// from the soft object path each time
+	if(InLastController != nullptr && !Action->WeakController.IsValid())
+	{
+		if(Action->ControllerPath.ToString() == InLastController->GetPathName())
+		{
+			Action->WeakController = InLastController;
+		}
+	}
+	return Action;
 }
 
 FString FRigVMActionWrapper::ExportText() const
@@ -55,9 +66,9 @@ FString FRigVMActionWrapper::ExportText() const
 	return ExportedText;
 }
 
-bool URigVMActionStack::OpenUndoBracket(const FString& InTitle)
+bool URigVMActionStack::OpenUndoBracket(URigVMController* InController, const FString& InTitle)
 {
-	FRigVMBaseAction* Action = new FRigVMBaseAction;
+	FRigVMBaseAction* Action = new FRigVMBaseAction(InController);
 	Action->Title = InTitle;
 	BracketActions.Add(Action);
 	BeginAction(*Action);
@@ -81,16 +92,14 @@ bool URigVMActionStack::CancelUndoBracket(URigVMController* InController)
 {
 	ensure(BracketActions.Num() > 0);
 	FRigVMBaseAction* Action = BracketActions.Pop();
-	CancelAction(*Action, InController);
+	CancelAction(*Action);
 	delete(Action);
 	return true;
 }
 
 bool URigVMActionStack::Undo(URigVMController* InController)
 {
-	check(InController)
-
-	if (UndoActions.Num() == 0)
+	if (UndoActions.Num() == 0 && InController)
 	{
 		InController->ReportWarning(TEXT("Nothing to undo."));
 		return false;
@@ -106,21 +115,22 @@ bool URigVMActionStack::Undo(URigVMController* InController)
 	LogAction(Wrapper.GetScriptStruct(), *Wrapper.GetAction(), FRigVMBaseAction::UndoPrefix);
 #endif
 	
-	if (Wrapper.GetAction()->Undo(InController))
+	if (Wrapper.GetAction(InController)->Undo())
 	{
 		RedoActions.Add(KeyToUndo);
 		return true;
 	}
 
-	InController->ReportAndNotifyErrorf(TEXT("Error while undoing action %s."), *Wrapper.GetAction()->Title);
+	if(InController)
+	{
+		InController->ReportAndNotifyErrorf(TEXT("Error while undoing action %s."), *Wrapper.GetAction()->Title);
+	}
 	return false;
 }
 
 bool URigVMActionStack::Redo(URigVMController* InController)
 {
-	check(InController)
-
-	if (RedoActions.Num() == 0)
+	if (RedoActions.Num() == 0 && InController)
 	{
 		InController->ReportWarning(TEXT("Nothing to redo."));
 		return false;
@@ -135,14 +145,17 @@ bool URigVMActionStack::Redo(URigVMController* InController)
 	LogAction(Wrapper.GetScriptStruct(), *Wrapper.GetAction(), FRigVMBaseAction::RedoPrefix);
 #endif
 	
-	if (Wrapper.GetAction()->Redo(InController))
+	if (Wrapper.GetAction(InController)->Redo())
 	{
 		UndoActions.Add(KeyToRedo);
 		ActionIndex = UndoActions.Num();
 		return true;
 	}
 
-	InController->ReportAndNotifyErrorf(TEXT("Error while undoing action %s."), *Wrapper.GetAction()->Title);
+	if(InController)
+	{
+		InController->ReportAndNotifyErrorf(TEXT("Error while undoing action %s."), *Wrapper.GetAction()->Title);
+	}
 	return false;
 }
 
@@ -152,12 +165,6 @@ void URigVMActionStack::PostTransacted(const FTransactionObjectEvent& Transactio
 {
 	if (TransactionEvent.GetEventType() == ETransactionObjectEventType::UndoRedo)
 	{
-		URigVMController* Controller = Cast< URigVMController>(GetOuter());
-		if (Controller == nullptr)
-		{
-			return;
-		}
-
 		int32 DesiredActionIndex = ActionIndex;
 		ActionIndex = UndoActions.Num();
 
@@ -174,7 +181,7 @@ void URigVMActionStack::PostTransacted(const FTransactionObjectEvent& Transactio
 			{
 				break;
 			}
-			if (!Undo(Controller))
+			if (!Undo(nullptr))
 			{
 				ModifiedEvent.Broadcast(ERigVMGraphNotifType::InteractionBracketCanceled, nullptr, nullptr);
 				return;
@@ -186,7 +193,7 @@ void URigVMActionStack::PostTransacted(const FTransactionObjectEvent& Transactio
 			{
 				break;
 			}
-			if (!Redo(Controller))
+			if (!Redo(nullptr))
 			{
 				ModifiedEvent.Broadcast(ERigVMGraphNotifType::InteractionBracketCanceled, nullptr, nullptr);
 				return;
@@ -206,7 +213,7 @@ void URigVMActionStack::LogAction(const UScriptStruct* InActionStruct, const FRi
 		return;
 	}
 	
-	if(URigVMController* Controller = GetTypedOuter<URigVMController>())
+	if(URigVMController* Controller = InAction.GetController())
 	{
 		static TArray<FString> TabPrefix = {TEXT(""), TEXT("  ")};
 
@@ -221,11 +228,12 @@ void URigVMActionStack::LogAction(const UScriptStruct* InActionStruct, const FRi
 
 	if(InPrefix == FRigVMBaseAction::AddActionPrefix)
 	{
+		URigVMController* Controller = InAction.GetController();
 		TGuardValue<int32> ActionDepthGuard(LogActionDepth, LogActionDepth + 1);
 		for(const FRigVMActionKey& SubAction : InAction.SubActions)
 		{
 			const FRigVMActionWrapper Wrapper(SubAction);
-			LogAction(Wrapper.GetScriptStruct(), *Wrapper.GetAction(), InPrefix);
+			LogAction(Wrapper.GetScriptStruct(), *Wrapper.GetAction(Controller), InPrefix);
 		}
 	}
 }
@@ -235,6 +243,18 @@ void URigVMActionStack::LogAction(const UScriptStruct* InActionStruct, const FRi
 #endif
 
 
+URigVMController* FRigVMBaseAction::GetController() const
+{
+	if(!WeakController.IsValid())
+	{
+		if(URigVMController* Controller = Cast<URigVMController>(ControllerPath.ResolveObject()))
+		{
+			WeakController = Controller;
+		}
+	}
+	return WeakController.Get();
+}
+
 bool FRigVMBaseAction::Merge(const FRigVMBaseAction* Other)
 {
 	return SubActions.Num() == 0 && Other->SubActions.Num() == 0;
@@ -245,62 +265,79 @@ bool FRigVMBaseAction::MakesObsolete(const FRigVMBaseAction* Other) const
 	return false;
 }
 
-bool FRigVMBaseAction::Undo(URigVMController* InController)
+bool FRigVMBaseAction::Undo()
 {
+	URigVMController* Controller = GetController();
+	if(Controller == nullptr)
+	{
+		return false;
+	}
+
 #if RIGVM_ACTIONSTACK_VERBOSE_LOG		
-	URigVMActionStack* Stack = InController->ActionStack;
+	URigVMActionStack* Stack = Controller->GetActionStack();
 	check(Stack);
 	TGuardValue<int32> TabDepthGuard(Stack->LogActionDepth, Stack->LogActionDepth + 1);
 #endif
-	TGuardValue<bool> TransactionGuard(InController->bIsTransacting, true);
+	TGuardValue<bool> TransactionGuard(Controller->bIsTransacting, true);
 
 	bool Result = true;
 	for (int32 KeyIndex = SubActions.Num() - 1; KeyIndex >= 0; KeyIndex--)
 	{
 		FRigVMActionWrapper Wrapper(SubActions[KeyIndex]);
 #if RIGVM_ACTIONSTACK_VERBOSE_LOG		
-		Wrapper.GetAction()->LogAction(InController, UndoPrefix);
+		Wrapper.GetAction(Controller)->LogAction(UndoPrefix);
 #endif
-		if(!Wrapper.GetAction()->Undo(InController))
+		if(!Wrapper.GetAction(Controller)->Undo())
 		{
-			InController->ReportAndNotifyErrorf(TEXT("Error while undoing action '%s'."), *Wrapper.GetAction()->Title);
+			Controller->ReportAndNotifyErrorf(TEXT("Error while undoing action '%s'."), *Wrapper.GetAction(Controller)->Title);
 			Result = false;
 		}
 	}
 	return Result;
 }
 
-bool FRigVMBaseAction::Redo(URigVMController* InController)
+bool FRigVMBaseAction::Redo()
 {
+	URigVMController* Controller = GetController();
+	if(Controller == nullptr)
+	{
+		return false;
+	}
+
 #if RIGVM_ACTIONSTACK_VERBOSE_LOG		
-	URigVMActionStack* Stack = InController->ActionStack;
+	URigVMActionStack* Stack = Controller->GetActionStack();
 	check(Stack);
 	TGuardValue<int32> TabDepthGuard(Stack->LogActionDepth, Stack->LogActionDepth + 1);
 #endif
-	TGuardValue<bool> TransactionGuard(InController->bIsTransacting, true);
+	TGuardValue<bool> TransactionGuard(Controller->bIsTransacting, true);
 
 	bool Result = true;
 	for (int32 KeyIndex = 0; KeyIndex < SubActions.Num(); KeyIndex++)
 	{
 		FRigVMActionWrapper Wrapper(SubActions[KeyIndex]);
 #if RIGVM_ACTIONSTACK_VERBOSE_LOG		
-		Wrapper.GetAction()->LogAction(InController, RedoPrefix);
+		Wrapper.GetAction()->LogAction(RedoPrefix);
 #endif
-		if (!Wrapper.GetAction()->Redo(InController))
+		if (!Wrapper.GetAction()->Redo())
 		{
-			InController->ReportAndNotifyErrorf(TEXT("Error while redoing action '%s'."), *Wrapper.GetAction()->Title);
+			Controller->ReportAndNotifyErrorf(TEXT("Error while redoing action '%s'."), *Wrapper.GetAction()->Title);
 			Result = false;
 		}
 	}
 	return Result;
 }
 
-bool FRigVMBaseAction::StoreNode(const URigVMNode* InNode, URigVMController* InController, bool bIsPriorChange)
+bool FRigVMBaseAction::StoreNode(const URigVMNode* InNode, bool bIsPriorChange)
 {
+	URigVMController* Controller = GetController();
+	if(Controller == nullptr)
+	{
+		return false;
+	}
+
 	check(InNode);
-	check(InController);
 
-	const FString& Content = InController->ExportNodesToText({InNode->GetFName()}, true);
+	const FString& Content = Controller->ExportNodesToText({InNode->GetFName()}, true);
 	if(!Content.IsEmpty())
 	{
 		FRigVMActionNodeContent& StoredContent = ExportedNodes.FindOrAdd(InNode->GetFName());
@@ -316,25 +353,29 @@ bool FRigVMBaseAction::StoreNode(const URigVMNode* InNode, URigVMController* InC
 	return !Content.IsEmpty();
 }
 
-bool FRigVMBaseAction::RestoreNode(const FName& InNodeName, URigVMController* InController, bool bIsUndoing)
+bool FRigVMBaseAction::RestoreNode(const FName& InNodeName, bool bIsUndoing)
 {
-	URigVMGraph* Graph = InController->GetGraph();
+	URigVMController* Controller = GetController();
+	if(Controller == nullptr)
+	{
+		return false;
+	}
+
+	URigVMGraph* Graph = Controller->GetGraph();
 	check(Graph);
 
 	FRigVMActionNodeContent& Content = ExportedNodes.FindChecked(InNodeName);
 
-	//TArray<URigVMController::FLinkedPath> LinkedPaths;
 	if(URigVMNode* ExistingNode = Graph->FindNodeByName(InNodeName))
 	{
-		//LinkedPaths = InController->GetLinkedPaths(ExistingNode);
-		if(!InController->RemoveNode(ExistingNode, false, false))
+		if(!Controller->RemoveNode(ExistingNode, false, false))
 		{
 			return false;
 		}
 	}
 
 	const FString& ContentToImport = bIsUndoing ? Content.Old : Content.New;
-	const TArray<FName> ImportedNodeNames = InController->ImportNodesFromText(ContentToImport, false, false);
+	const TArray<FName> ImportedNodeNames = Controller->ImportNodesFromText(ContentToImport, false, false);
 	if(ImportedNodeNames.Num() == 1)
 	{
 		if(!ImportedNodeNames[0].IsEqual(InNodeName, ENameCase::CaseSensitive))
@@ -343,25 +384,20 @@ bool FRigVMBaseAction::RestoreNode(const FName& InNodeName, URigVMController* In
 		}
 	}
 
-	/*
-	if(!LinkedPaths.IsEmpty())
-	{
-		if(!InController->RestoreLinkedPaths(LinkedPaths))
-		{
-			return false;
-		}
-	}
-	*/
-	
 	return true;
 }
 
 #if RIGVM_ACTIONSTACK_VERBOSE_LOG
 
-void FRigVMBaseAction::LogAction(URigVMController* InController, const FString& InPrefix) const
+void FRigVMBaseAction::LogAction(const FString& InPrefix) const
 {
-	check(InController);
-	URigVMActionStack* Stack = InController->ActionStack;
+	URigVMController* Controller = GetController();
+	if(Controller == nullptr)
+	{
+		return;
+	}
+	
+	URigVMActionStack* Stack = Controller->GetActionStack();
 	check(Stack);
 	TGuardValue<int32> TabDepthGuard(Stack->LogActionDepth, Stack->LogActionDepth + 1);
 
@@ -371,18 +407,20 @@ void FRigVMBaseAction::LogAction(URigVMController* InController, const FString&
 #endif
 
 FRigVMInjectNodeIntoPinAction::FRigVMInjectNodeIntoPinAction()
-	: PinPath()
-	, bAsInput(false)
-	, InputPinName(NAME_None)
-	, OutputPinName(NAME_None)
-	, NodePath()
+: FRigVMBaseAction(nullptr)
+, PinPath()
+, bAsInput(false)
+, InputPinName(NAME_None)
+, OutputPinName(NAME_None)
+, NodePath()
 {
 }
 
-FRigVMInjectNodeIntoPinAction::FRigVMInjectNodeIntoPinAction(URigVMInjectionInfo* InInjectionInfo)
-	: PinPath(InInjectionInfo->GetPin()->GetPinPath())
-	, bAsInput(InInjectionInfo->bInjectedAsInput)
-	, NodePath(InInjectionInfo->Node->GetName())
+FRigVMInjectNodeIntoPinAction::FRigVMInjectNodeIntoPinAction(URigVMController* InController, URigVMInjectionInfo* InInjectionInfo)
+: FRigVMBaseAction(InController)
+, PinPath(InInjectionInfo->GetPin()->GetPinPath())
+, bAsInput(InInjectionInfo->bInjectedAsInput)
+, NodePath(InInjectionInfo->Node->GetName())
 {
 	if (InInjectionInfo->InputPin)
 	{
@@ -394,60 +432,74 @@ FRigVMInjectNodeIntoPinAction::FRigVMInjectNodeIntoPinAction(URigVMInjectionInfo
 	}
 }
 
-bool FRigVMInjectNodeIntoPinAction::Undo(URigVMController* InController)
+bool FRigVMInjectNodeIntoPinAction::Undo()
 {
-	if (!FRigVMBaseAction::Undo(InController))
+	if (!FRigVMBaseAction::Undo())
 	{
 		return false;
 	}
-	return InController->EjectNodeFromPin(*PinPath, false) != nullptr;
+	return GetController()->EjectNodeFromPin(*PinPath, false) != nullptr;
 }
 
-bool FRigVMInjectNodeIntoPinAction::Redo(URigVMController* InController)
+bool FRigVMInjectNodeIntoPinAction::Redo()
 {
+	if(!CanUndoRedo())
+	{
+		return false;
+	}
 #if WITH_EDITOR
-	if (URigVMInjectionInfo* InjectionInfo = InController->InjectNodeIntoPin(PinPath, bAsInput, InputPinName, OutputPinName, false))
+	if (URigVMInjectionInfo* InjectionInfo = GetController()->InjectNodeIntoPin(PinPath, bAsInput, InputPinName, OutputPinName, false))
 	{
-		return FRigVMBaseAction::Redo(InController);
+		return FRigVMBaseAction::Redo();
 	}
 #endif
 	return false;
 }
 
 FRigVMEjectNodeFromPinAction::FRigVMEjectNodeFromPinAction()
-	: FRigVMInjectNodeIntoPinAction()
+: FRigVMInjectNodeIntoPinAction()
 {
 }
 
-FRigVMEjectNodeFromPinAction::FRigVMEjectNodeFromPinAction(URigVMInjectionInfo* InInjectionInfo)
-	: FRigVMInjectNodeIntoPinAction(InInjectionInfo)
+FRigVMEjectNodeFromPinAction::FRigVMEjectNodeFromPinAction(URigVMController* InController, URigVMInjectionInfo* InInjectionInfo)
+: FRigVMInjectNodeIntoPinAction(InController, InInjectionInfo)
 {
 }
 
-bool FRigVMEjectNodeFromPinAction::Undo(URigVMController* InController)
+bool FRigVMEjectNodeFromPinAction::Undo()
 {
-	if (!FRigVMBaseAction::Undo(InController))
+	if (!FRigVMBaseAction::Undo())
 	{
 		return false;
 	}
 #if WITH_EDITOR
-	const URigVMInjectionInfo* InjectionInfo = InController->InjectNodeIntoPin(PinPath, bAsInput, InputPinName, OutputPinName, false);
+	const URigVMInjectionInfo* InjectionInfo = GetController()->InjectNodeIntoPin(PinPath, bAsInput, InputPinName, OutputPinName, false);
 	return InjectionInfo != nullptr;
 #else
 	return false;
 #endif
 }
 
-bool FRigVMEjectNodeFromPinAction::Redo(URigVMController* InController)
+bool FRigVMEjectNodeFromPinAction::Redo()
 {
-	if (InController->EjectNodeFromPin(*PinPath, false) != nullptr)
+	if(!CanUndoRedo())
 	{
-		return FRigVMBaseAction::Redo(InController);
+		return false;
+	}
+	if (GetController()->EjectNodeFromPin(*PinPath, false) != nullptr)
+	{
+		return FRigVMBaseAction::Redo();
 	}
 	return false;
 }
 
-FRigVMRemoveNodesAction::FRigVMRemoveNodesAction(TArray<URigVMNode*> InNodes, URigVMController* InController)
+FRigVMRemoveNodesAction::FRigVMRemoveNodesAction()
+: FRigVMBaseAction(nullptr)
+{
+}
+
+FRigVMRemoveNodesAction::FRigVMRemoveNodesAction(URigVMController* InController, TArray<URigVMNode*> InNodes)
+: FRigVMBaseAction(InController)
 {
 	Algo::Transform(InNodes, NodeNames, [](const URigVMNode* Node) -> FName
 	{
@@ -456,55 +508,73 @@ FRigVMRemoveNodesAction::FRigVMRemoveNodesAction(TArray<URigVMNode*> InNodes, UR
 	ExportedContent = InController->ExportNodesToText(NodeNames, true);
 }
 
-bool FRigVMRemoveNodesAction::Undo(URigVMController* InController)
+bool FRigVMRemoveNodesAction::Undo()
 {
-	if (!FRigVMBaseAction::Undo(InController))
+	if (!FRigVMBaseAction::Undo())
 	{
 		return false;
 	}
-	const TArray<FName> NewNodeNames = InController->ImportNodesFromText(ExportedContent, false, false);
+	const TArray<FName> NewNodeNames = GetController()->ImportNodesFromText(ExportedContent, false, false);
 	return NewNodeNames.Num() == NodeNames.Num();
 }
 
-bool FRigVMRemoveNodesAction::Redo(URigVMController* InController)
+bool FRigVMRemoveNodesAction::Redo()
 {
-	if(InController->RemoveNodesByName(NodeNames, false, false))
+	if(!CanUndoRedo())
 	{
-		return FRigVMBaseAction::Redo(InController);
+		return false;
+	}
+	if(GetController()->RemoveNodesByName(NodeNames, false, false))
+	{
+		return FRigVMBaseAction::Redo();
 	}
 	return false;
 }
 
 FRigVMSetNodeSelectionAction::FRigVMSetNodeSelectionAction()
+: FRigVMBaseAction(nullptr)
 {
 }
 
-FRigVMSetNodeSelectionAction::FRigVMSetNodeSelectionAction(URigVMGraph* InGraph, TArray<FName> InNewSelection)
+FRigVMSetNodeSelectionAction::FRigVMSetNodeSelectionAction(URigVMController* InController, URigVMGraph* InGraph, TArray<FName> InNewSelection)
+: FRigVMBaseAction(InController)
 {
 	OldSelection = InGraph->GetSelectNodes();
 	NewSelection = InNewSelection;
 }
 
-bool FRigVMSetNodeSelectionAction::Undo(URigVMController* InController)
+bool FRigVMSetNodeSelectionAction::Undo()
 {
-	if(!FRigVMBaseAction::Undo(InController))
+	if(!FRigVMBaseAction::Undo())
 	{
 		return false;
 	}
-	return InController->SetNodeSelection(OldSelection, false);
+	return GetController()->SetNodeSelection(OldSelection, false);
 }
 
-bool FRigVMSetNodeSelectionAction::Redo(URigVMController* InController)
+bool FRigVMSetNodeSelectionAction::Redo()
 {
-	if(!InController->SetNodeSelection(NewSelection, false))
+	if(!CanUndoRedo())
+	{
+		return false;
+	}
+	if(!GetController()->SetNodeSelection(NewSelection, false))
 	{
 		return false;
 	}
-	return FRigVMBaseAction::Redo(InController);
+	return FRigVMBaseAction::Redo();
+}
+
+FRigVMSetNodePositionAction::FRigVMSetNodePositionAction()
+: FRigVMBaseAction(nullptr)
+, OldPosition(FVector2D::ZeroVector)
+, NewPosition(FVector2D::ZeroVector)
+{
 }
 
-FRigVMSetNodePositionAction::FRigVMSetNodePositionAction(URigVMNode* InNode, const FVector2D& InNewPosition)
-: NodePath(InNode->GetNodePath())
+FRigVMSetNodePositionAction::FRigVMSetNodePositionAction(URigVMController* InController, URigVMNode* InNode, const FVector2D& InNewPosition)
+: FRigVMBaseAction(InController)
+, NodePath(InNode->GetNodePath())
 , OldPosition(InNode->GetPosition())
 , NewPosition(InNewPosition)
 {
@@ -527,26 +597,38 @@ bool FRigVMSetNodePositionAction::Merge(const FRigVMBaseAction* Other)
 	return true;
 }
 
-bool FRigVMSetNodePositionAction::Undo(URigVMController* InController)
+bool FRigVMSetNodePositionAction::Undo()
 {
-	if(!FRigVMBaseAction::Undo(InController))
+	if(!FRigVMBaseAction::Undo())
 	{
 		return false;
 	}
-	return InController->SetNodePositionByName(*NodePath, OldPosition, false);
+	return GetController()->SetNodePositionByName(*NodePath, OldPosition, false);
 }
 
-bool FRigVMSetNodePositionAction::Redo(URigVMController* InController)
+bool FRigVMSetNodePositionAction::Redo()
 {
-	if(!InController->SetNodePositionByName(*NodePath, NewPosition, false))
+	if(!CanUndoRedo())
 	{
 		return false;
 	}
-	return FRigVMBaseAction::Redo(InController);
+	if(!GetController()->SetNodePositionByName(*NodePath, NewPosition, false))
+	{
+		return false;
+	}
+	return FRigVMBaseAction::Redo();
+}
+
+FRigVMSetNodeSizeAction::FRigVMSetNodeSizeAction()
+: FRigVMBaseAction(nullptr)
+, OldSize(FVector2D::ZeroVector)
+, NewSize(FVector2D::ZeroVector)
+{
 }
 
-FRigVMSetNodeSizeAction::FRigVMSetNodeSizeAction(URigVMNode* InNode, const FVector2D& InNewSize)
-: NodePath(InNode->GetNodePath())
+FRigVMSetNodeSizeAction::FRigVMSetNodeSizeAction(URigVMController* InController, URigVMNode* InNode, const FVector2D& InNewSize)
+: FRigVMBaseAction(InController)
+, NodePath(InNode->GetNodePath())
 , OldSize(InNode->GetSize())
 , NewSize(InNewSize)
 {
@@ -569,26 +651,38 @@ bool FRigVMSetNodeSizeAction::Merge(const FRigVMBaseAction* Other)
 	return true;
 }
 
-bool FRigVMSetNodeSizeAction::Undo(URigVMController* InController)
+bool FRigVMSetNodeSizeAction::Undo()
 {
-	if(!FRigVMBaseAction::Undo(InController))
+	if(!FRigVMBaseAction::Undo())
 	{
 		return false;
 	}
-	return InController->SetNodeSizeByName(*NodePath, OldSize, false);
+	return GetController()->SetNodeSizeByName(*NodePath, OldSize, false);
 }
 
-bool FRigVMSetNodeSizeAction::Redo(URigVMController* InController)
+bool FRigVMSetNodeSizeAction::Redo()
 {
-	if(!InController->SetNodeSizeByName(*NodePath, NewSize, false))
+	if(!CanUndoRedo())
+	{
+		return false;
+	}
+	if(!GetController()->SetNodeSizeByName(*NodePath, NewSize, false))
 	{
 		return false;
 	}
-	return FRigVMBaseAction::Redo(InController);
+	return FRigVMBaseAction::Redo();
+}
+
+FRigVMSetNodeColorAction::FRigVMSetNodeColorAction()
+: FRigVMBaseAction(nullptr)
+, OldColor(FLinearColor::Black)
+, NewColor(FLinearColor::Black)
+{
 }
 
-FRigVMSetNodeColorAction::FRigVMSetNodeColorAction(URigVMNode* InNode, const FLinearColor& InNewColor)
-: NodePath(InNode->GetNodePath())
+FRigVMSetNodeColorAction::FRigVMSetNodeColorAction(URigVMController* InController, URigVMNode* InNode, const FLinearColor& InNewColor)
+: FRigVMBaseAction(InController)
+, NodePath(InNode->GetNodePath())
 , OldColor(InNode->GetNodeColor())
 , NewColor(InNewColor)
 {
@@ -611,28 +705,40 @@ bool FRigVMSetNodeColorAction::Merge(const FRigVMBaseAction* Other)
 	return true;
 }
 
-bool FRigVMSetNodeColorAction::Undo(URigVMController* InController)
+bool FRigVMSetNodeColorAction::Undo()
 {
-	if(!FRigVMBaseAction::Undo(InController))
+	if(!FRigVMBaseAction::Undo())
 	{
 		return false;
 	}
-	return InController->SetNodeColorByName(*NodePath, OldColor, false);
+	return GetController()->SetNodeColorByName(*NodePath, OldColor, false);
 }
 
-bool FRigVMSetNodeColorAction::Redo(URigVMController* InController)
+bool FRigVMSetNodeColorAction::Redo()
 {
-	if(!InController->SetNodeColorByName(*NodePath, NewColor, false))
+	if(!CanUndoRedo())
 	{
 		return false;
 	}
-	return FRigVMBaseAction::Redo(InController);
+	if(!GetController()->SetNodeColorByName(*NodePath, NewColor, false))
+	{
+		return false;
+	}
+	return FRigVMBaseAction::Redo();
+}
+
+FRigVMSetNodeCategoryAction::FRigVMSetNodeCategoryAction()
+: FRigVMBaseAction(nullptr)
+, OldCategory()
+, NewCategory()
+{
 }
 
-FRigVMSetNodeCategoryAction::FRigVMSetNodeCategoryAction(URigVMCollapseNode* InNode, const FString& InNewCategory)
-	: NodePath(InNode->GetNodePath())
-	, OldCategory(InNode->GetNodeCategory())
-	, NewCategory(InNewCategory)
+FRigVMSetNodeCategoryAction::FRigVMSetNodeCategoryAction(URigVMController* InController, URigVMCollapseNode* InNode, const FString& InNewCategory)
+: FRigVMBaseAction(InController)
+, NodePath(InNode->GetNodePath())
+, OldCategory(InNode->GetNodeCategory())
+, NewCategory(InNewCategory)
 {
 }
 
@@ -653,28 +759,40 @@ bool FRigVMSetNodeCategoryAction::Merge(const FRigVMBaseAction* Other)
 	return true;
 }
 
-bool FRigVMSetNodeCategoryAction::Undo(URigVMController* InController)
+bool FRigVMSetNodeCategoryAction::Undo()
 {
-	if (!FRigVMBaseAction::Undo(InController))
+	if (!FRigVMBaseAction::Undo())
 	{
 		return false;
 	}
-	return InController->SetNodeCategoryByName(*NodePath, OldCategory, false);
+	return GetController()->SetNodeCategoryByName(*NodePath, OldCategory, false);
 }
 
-bool FRigVMSetNodeCategoryAction::Redo(URigVMController* InController)
+bool FRigVMSetNodeCategoryAction::Redo()
 {
-	if (!InController->SetNodeCategoryByName(*NodePath, NewCategory, false))
+	if(!CanUndoRedo())
+	{
+		return false;
+	}
+	if (!GetController()->SetNodeCategoryByName(*NodePath, NewCategory, false))
 	{
 		return false;
 	}
-	return FRigVMBaseAction::Redo(InController);
+	return FRigVMBaseAction::Redo();
+}
+
+FRigVMSetNodeKeywordsAction::FRigVMSetNodeKeywordsAction()
+: FRigVMBaseAction(nullptr)
+, OldKeywords()
+, NewKeywords()
+{
 }
 
-FRigVMSetNodeKeywordsAction::FRigVMSetNodeKeywordsAction(URigVMCollapseNode* InNode, const FString& InNewKeywords)
-	: NodePath(InNode->GetNodePath())
-	, OldKeywords(InNode->GetNodeKeywords())
-	, NewKeywords(InNewKeywords)
+FRigVMSetNodeKeywordsAction::FRigVMSetNodeKeywordsAction(URigVMController* InController, URigVMCollapseNode* InNode, const FString& InNewKeywords)
+: FRigVMBaseAction(InController)
+, NodePath(InNode->GetNodePath())
+, OldKeywords(InNode->GetNodeKeywords())
+, NewKeywords(InNewKeywords)
 {
 }
 
@@ -695,28 +813,40 @@ bool FRigVMSetNodeKeywordsAction::Merge(const FRigVMBaseAction* Other)
 	return true;
 }
 
-bool FRigVMSetNodeKeywordsAction::Undo(URigVMController* InController)
+bool FRigVMSetNodeKeywordsAction::Undo()
 {
-	if (!FRigVMBaseAction::Undo(InController))
+	if (!FRigVMBaseAction::Undo())
 	{
 		return false;
 	}
-	return InController->SetNodeKeywordsByName(*NodePath, OldKeywords, false);
+	return GetController()->SetNodeKeywordsByName(*NodePath, OldKeywords, false);
 }
 
-bool FRigVMSetNodeKeywordsAction::Redo(URigVMController* InController)
+bool FRigVMSetNodeKeywordsAction::Redo()
 {
-	if (!InController->SetNodeKeywordsByName(*NodePath, NewKeywords, false))
+	if(!CanUndoRedo())
 	{
 		return false;
 	}
-	return FRigVMBaseAction::Redo(InController);
+	if (!GetController()->SetNodeKeywordsByName(*NodePath, NewKeywords, false))
+	{
+		return false;
+	}
+	return FRigVMBaseAction::Redo();
 }
 
-FRigVMSetNodeDescriptionAction::FRigVMSetNodeDescriptionAction(URigVMCollapseNode* InNode, const FString& InNewDescription)
-	: NodePath(InNode->GetNodePath())
-	, OldDescription(InNode->GetNodeDescription())
-	, NewDescription(InNewDescription)
+FRigVMSetNodeDescriptionAction::FRigVMSetNodeDescriptionAction()
+: FRigVMBaseAction(nullptr)
+, OldDescription()
+, NewDescription()
+{
+}
+
+FRigVMSetNodeDescriptionAction::FRigVMSetNodeDescriptionAction(URigVMController* InController, URigVMCollapseNode* InNode, const FString& InNewDescription)
+: FRigVMBaseAction(InController)
+, NodePath(InNode->GetNodePath())
+, OldDescription(InNode->GetNodeDescription())
+, NewDescription(InNewDescription)
 {
 }
 
@@ -737,26 +867,31 @@ bool FRigVMSetNodeDescriptionAction::Merge(const FRigVMBaseAction* Other)
 	return true;
 }
 
-bool FRigVMSetNodeDescriptionAction::Undo(URigVMController* InController)
+bool FRigVMSetNodeDescriptionAction::Undo()
 {
-	if (!FRigVMBaseAction::Undo(InController))
+	if (!FRigVMBaseAction::Undo())
 	{
 		return false;
 	}
-	return InController->SetNodeDescriptionByName(*NodePath, OldDescription, false);
+	return GetController()->SetNodeDescriptionByName(*NodePath, OldDescription, false);
 }
 
-bool FRigVMSetNodeDescriptionAction::Redo(URigVMController* InController)
+bool FRigVMSetNodeDescriptionAction::Redo()
 {
-	if (!InController->SetNodeDescriptionByName(*NodePath, NewDescription, false))
+	if(!CanUndoRedo())
+	{
+		return false;
+	}
+	if (!GetController()->SetNodeDescriptionByName(*NodePath, NewDescription, false))
 	{
 		return false;
 	}
-	return FRigVMBaseAction::Redo(InController);
+	return FRigVMBaseAction::Redo();
 }
 
 FRigVMSetCommentTextAction::FRigVMSetCommentTextAction()
-: NodePath()
+: FRigVMBaseAction(nullptr)
+, NodePath()
 , OldText()
 , NewText()
 , OldFontSize(18)
@@ -770,8 +905,9 @@ FRigVMSetCommentTextAction::FRigVMSetCommentTextAction()
 }
 
 
-FRigVMSetCommentTextAction::FRigVMSetCommentTextAction(URigVMCommentNode* InNode, const FString& InNewText, const int32& InNewFontSize, const bool& bInNewBubbleVisible, const bool& bInNewColorBubble)
-: NodePath(InNode->GetNodePath())
+FRigVMSetCommentTextAction::FRigVMSetCommentTextAction(URigVMController* InController, URigVMCommentNode* InNode, const FString& InNewText, const int32& InNewFontSize, const bool& bInNewBubbleVisible, const bool& bInNewColorBubble)
+: FRigVMBaseAction(InController)
+, NodePath(InNode->GetNodePath())
 , OldText(InNode->GetCommentText())
 , NewText(InNewText)
 , OldFontSize(InNode->GetCommentFontSize())
@@ -783,132 +919,183 @@ FRigVMSetCommentTextAction::FRigVMSetCommentTextAction(URigVMCommentNode* InNode
 {
 }
 
-bool FRigVMSetCommentTextAction::Undo(URigVMController* InController)
+bool FRigVMSetCommentTextAction::Undo()
 {
-	if(!FRigVMBaseAction::Undo(InController))
+	if(!FRigVMBaseAction::Undo())
 	{
 		return false;
 	}
-	return InController->SetCommentTextByName(*NodePath, OldText, OldFontSize, bOldBubbleVisible, bOldColorBubble, false);
+	return GetController()->SetCommentTextByName(*NodePath, OldText, OldFontSize, bOldBubbleVisible, bOldColorBubble, false);
 }
 
-bool FRigVMSetCommentTextAction::Redo(URigVMController* InController)
+bool FRigVMSetCommentTextAction::Redo()
 {
-	if(!InController->SetCommentTextByName(*NodePath, NewText, NewFontSize, bNewBubbleVisible, bNewColorBubble, false))
+	if(!CanUndoRedo())
+	{
+		return false;
+	}
+	if(!GetController()->SetCommentTextByName(*NodePath, NewText, NewFontSize, bNewBubbleVisible, bNewColorBubble, false))
 	{
 		return false;
 	}
-	return FRigVMBaseAction::Redo(InController);
+	return FRigVMBaseAction::Redo();
 }
 
 FRigVMSetRerouteCompactnessAction::FRigVMSetRerouteCompactnessAction()
-	: NodePath()
-	, OldShowAsFullNode(false)
-	, NewShowAsFullNode(false)
+: FRigVMBaseAction(nullptr)
+, NodePath()
+, OldShowAsFullNode(false)
+, NewShowAsFullNode(false)
 {
 
 }
-FRigVMSetRerouteCompactnessAction::FRigVMSetRerouteCompactnessAction(URigVMRerouteNode* InNode, bool InShowAsFullNode)
-	: NodePath(InNode->GetNodePath())
-	, OldShowAsFullNode(InNode->GetShowsAsFullNode())
-	, NewShowAsFullNode(InShowAsFullNode)
+FRigVMSetRerouteCompactnessAction::FRigVMSetRerouteCompactnessAction(URigVMController* InController, URigVMRerouteNode* InNode, bool InShowAsFullNode)
+: FRigVMBaseAction(InController)
+, NodePath(InNode->GetNodePath())
+, OldShowAsFullNode(InNode->GetShowsAsFullNode())
+, NewShowAsFullNode(InShowAsFullNode)
 {
 }
 
-bool FRigVMSetRerouteCompactnessAction::Undo(URigVMController* InController)
+bool FRigVMSetRerouteCompactnessAction::Undo()
 {
-	if(!FRigVMBaseAction::Undo(InController))
+	if(!FRigVMBaseAction::Undo())
 	{
 		return false;
 	}
-	return InController->SetRerouteCompactnessByName(*NodePath, OldShowAsFullNode, false);
+	return GetController()->SetRerouteCompactnessByName(*NodePath, OldShowAsFullNode, false);
 }
 
-bool FRigVMSetRerouteCompactnessAction::Redo(URigVMController* InController)
+bool FRigVMSetRerouteCompactnessAction::Redo()
 {
-	if(!InController->SetRerouteCompactnessByName(*NodePath, NewShowAsFullNode, false))
+	if(!CanUndoRedo())
 	{
 		return false;
 	}
-	return FRigVMBaseAction::Redo(InController);
+	if(!GetController()->SetRerouteCompactnessByName(*NodePath, NewShowAsFullNode, false))
+	{
+		return false;
+	}
+	return FRigVMBaseAction::Redo();
+}
+
+FRigVMRenameVariableAction::FRigVMRenameVariableAction()
+: FRigVMBaseAction(nullptr)
+{
 }
 
-FRigVMRenameVariableAction::FRigVMRenameVariableAction(const FName& InOldVariableName, const FName& InNewVariableName)
-: OldVariableName(InOldVariableName.ToString())
+FRigVMRenameVariableAction::FRigVMRenameVariableAction(URigVMController* InController, const FName& InOldVariableName, const FName& InNewVariableName)
+: FRigVMBaseAction(InController)
+, OldVariableName(InOldVariableName.ToString())
 , NewVariableName(InNewVariableName.ToString())
 {
 }
 
-bool FRigVMRenameVariableAction::Undo(URigVMController* InController)
+bool FRigVMRenameVariableAction::Undo()
 {
-	if(!FRigVMBaseAction::Undo(InController))
+	if(!FRigVMBaseAction::Undo())
 	{
 		return false;
 	}
-	return InController->OnExternalVariableRenamed(*NewVariableName, *OldVariableName, false);
+	return GetController()->OnExternalVariableRenamed(*NewVariableName, *OldVariableName, false);
 }
 
-bool FRigVMRenameVariableAction::Redo(URigVMController* InController)
+bool FRigVMRenameVariableAction::Redo()
 {
-	if(!InController->OnExternalVariableRenamed(*OldVariableName, *NewVariableName, false))
+	if(!CanUndoRedo())
 	{
 		return false;
 	}
-	return FRigVMBaseAction::Redo(InController);
+	if(!GetController()->OnExternalVariableRenamed(*OldVariableName, *NewVariableName, false))
+	{
+		return false;
+	}
+	return FRigVMBaseAction::Redo();
+}
+
+FRigVMSetPinExpansionAction::FRigVMSetPinExpansionAction()
+: FRigVMBaseAction(nullptr)
+, PinPath()
+, OldIsExpanded(false)
+, NewIsExpanded(false)
+{
 }
 
-FRigVMSetPinExpansionAction::FRigVMSetPinExpansionAction(URigVMPin* InPin, bool bNewIsExpanded)
-: PinPath(InPin->GetPinPath())
+FRigVMSetPinExpansionAction::FRigVMSetPinExpansionAction(URigVMController* InController, URigVMPin* InPin, bool bNewIsExpanded)
+: FRigVMBaseAction(InController)
+, PinPath(InPin->GetPinPath())
 , OldIsExpanded(InPin->IsExpanded())
 , NewIsExpanded(bNewIsExpanded)
 {
 }
 
-bool FRigVMSetPinExpansionAction::Undo(URigVMController* InController)
+bool FRigVMSetPinExpansionAction::Undo()
 {
-	if(!FRigVMBaseAction::Undo(InController))
+	if(!FRigVMBaseAction::Undo())
 	{
 		return false;
 	}
-	return InController->SetPinExpansion(PinPath, OldIsExpanded, false);
+	return GetController()->SetPinExpansion(PinPath, OldIsExpanded, false);
 }
 
-bool FRigVMSetPinExpansionAction::Redo(URigVMController* InController)
+bool FRigVMSetPinExpansionAction::Redo()
 {
-	if(!InController->SetPinExpansion(PinPath, NewIsExpanded, false))
+	if(!CanUndoRedo())
 	{
 		return false;
 	}
-	return FRigVMBaseAction::Redo(InController);
+	if(!GetController()->SetPinExpansion(PinPath, NewIsExpanded, false))
+	{
+		return false;
+	}
+	return FRigVMBaseAction::Redo();
 }
 
-FRigVMSetPinWatchAction::FRigVMSetPinWatchAction(URigVMPin* InPin, bool bNewIsWatched)
-: PinPath(InPin->GetPinPath())
+FRigVMSetPinWatchAction::FRigVMSetPinWatchAction()
+: FRigVMBaseAction(nullptr)
+, OldIsWatched(false)
+, NewIsWatched(false)
+{
+}
+
+FRigVMSetPinWatchAction::FRigVMSetPinWatchAction(URigVMController* InController, URigVMPin* InPin, bool bNewIsWatched)
+: FRigVMBaseAction(InController)
+, PinPath(InPin->GetPinPath())
 , OldIsWatched(InPin->RequiresWatch())
 , NewIsWatched(bNewIsWatched)
 {
 }
 
-bool FRigVMSetPinWatchAction::Undo(URigVMController* InController)
+bool FRigVMSetPinWatchAction::Undo()
 {
-	if(!FRigVMBaseAction::Undo(InController))
+	if(!FRigVMBaseAction::Undo())
 	{
 		return false;
 	}
-	return InController->SetPinIsWatched(PinPath, OldIsWatched, false);
+	return GetController()->SetPinIsWatched(PinPath, OldIsWatched, false);
 }
 
-bool FRigVMSetPinWatchAction::Redo(URigVMController* InController)
+bool FRigVMSetPinWatchAction::Redo()
 {
-	if(!InController->SetPinIsWatched(PinPath, NewIsWatched, false))
+	if(!CanUndoRedo())
 	{
 		return false;
 	}
-	return FRigVMBaseAction::Redo(InController);
+	if(!GetController()->SetPinIsWatched(PinPath, NewIsWatched, false))
+	{
+		return false;
+	}
+	return FRigVMBaseAction::Redo();
+}
+
+FRigVMSetPinDefaultValueAction::FRigVMSetPinDefaultValueAction()
+: FRigVMBaseAction(nullptr)
+{
 }
 
-FRigVMSetPinDefaultValueAction::FRigVMSetPinDefaultValueAction(URigVMPin* InPin, const FString& InNewDefaultValue)
-: PinPath(InPin->GetPinPath())
+FRigVMSetPinDefaultValueAction::FRigVMSetPinDefaultValueAction(URigVMController* InController, URigVMPin* InPin, const FString& InNewDefaultValue)
+: FRigVMBaseAction(InController)
+, PinPath(InPin->GetPinPath())
 , OldDefaultValue(InPin->GetDefaultValue())
 , NewDefaultValue(InNewDefaultValue)
 {
@@ -942,9 +1129,9 @@ bool FRigVMSetPinDefaultValueAction::Merge(const FRigVMBaseAction* Other)
 	return true;
 }
 
-bool FRigVMSetPinDefaultValueAction::Undo(URigVMController* InController)
+bool FRigVMSetPinDefaultValueAction::Undo()
 {
-	if(!FRigVMBaseAction::Undo(InController))
+	if(!FRigVMBaseAction::Undo())
 	{
 		return false;
 	}
@@ -952,122 +1139,169 @@ bool FRigVMSetPinDefaultValueAction::Undo(URigVMController* InController)
 	{
 		return true;
 	}
-	return InController->SetPinDefaultValue(PinPath, OldDefaultValue, true, false);
+	return GetController()->SetPinDefaultValue(PinPath, OldDefaultValue, true, false);
 }
 
-bool FRigVMSetPinDefaultValueAction::Redo(URigVMController* InController)
+bool FRigVMSetPinDefaultValueAction::Redo()
 {
+	if(!CanUndoRedo())
+	{
+		return false;
+	}
 	if (!NewDefaultValue.IsEmpty())
 	{
-		if (!InController->SetPinDefaultValue(PinPath, NewDefaultValue, true, false))
+		if (!GetController()->SetPinDefaultValue(PinPath, NewDefaultValue, true, false))
 		{
 			return false;
 		}
 	}
-	return FRigVMBaseAction::Redo(InController);
+	return FRigVMBaseAction::Redo();
 }
 
-FRigVMInsertArrayPinAction::FRigVMInsertArrayPinAction(URigVMPin* InArrayPin, int32 InIndex, const FString& InNewDefaultValue)
-: ArrayPinPath(InArrayPin->GetPinPath())
+FRigVMInsertArrayPinAction::FRigVMInsertArrayPinAction()
+: FRigVMBaseAction(nullptr)
+, Index(0)
+{
+}
+
+FRigVMInsertArrayPinAction::FRigVMInsertArrayPinAction(URigVMController* InController, URigVMPin* InArrayPin, int32 InIndex, const FString& InNewDefaultValue)
+: FRigVMBaseAction(InController)
+, ArrayPinPath(InArrayPin->GetPinPath())
 , Index(InIndex)
 , NewDefaultValue(InNewDefaultValue)
 {
 }
 
-bool FRigVMInsertArrayPinAction::Undo(URigVMController* InController)
+bool FRigVMInsertArrayPinAction::Undo()
 {
-	if(!FRigVMBaseAction::Undo(InController))
+	if(!FRigVMBaseAction::Undo())
 	{
 		return false;
 	}
-	return InController->RemoveArrayPin(FString::Printf(TEXT("%s.%d"), *ArrayPinPath, Index), false);
+	return GetController()->RemoveArrayPin(FString::Printf(TEXT("%s.%d"), *ArrayPinPath, Index), false);
 }
 
-bool FRigVMInsertArrayPinAction::Redo(URigVMController* InController)
+bool FRigVMInsertArrayPinAction::Redo()
 {
-	if(InController->InsertArrayPin(ArrayPinPath, Index, NewDefaultValue, false).IsEmpty())
+	if(!CanUndoRedo())
+	{
+		return false;
+	}
+	if(GetController()->InsertArrayPin(ArrayPinPath, Index, NewDefaultValue, false).IsEmpty())
 	{
 		return false;
 	}
-	return FRigVMBaseAction::Redo(InController);
+	return FRigVMBaseAction::Redo();
+}
+
+FRigVMRemoveArrayPinAction::FRigVMRemoveArrayPinAction()
+: FRigVMBaseAction(nullptr)
+, Index(0)
+{
 }
 
-FRigVMRemoveArrayPinAction::FRigVMRemoveArrayPinAction(URigVMPin* InArrayElementPin)
-: ArrayPinPath(InArrayElementPin->GetParentPin()->GetPinPath())
+FRigVMRemoveArrayPinAction::FRigVMRemoveArrayPinAction(URigVMController* InController, URigVMPin* InArrayElementPin)
+: FRigVMBaseAction(InController)
+, ArrayPinPath(InArrayElementPin->GetParentPin()->GetPinPath())
 , Index(InArrayElementPin->GetPinIndex())
 , DefaultValue(InArrayElementPin->GetDefaultValue())
 {
 }
 
-bool FRigVMRemoveArrayPinAction::Undo(URigVMController* InController)
+bool FRigVMRemoveArrayPinAction::Undo()
 {
-	if (InController->InsertArrayPin(*ArrayPinPath, Index, DefaultValue, false).IsEmpty())
+	if(!CanUndoRedo())
+	{
+		return false;
+	}
+	if (GetController()->InsertArrayPin(*ArrayPinPath, Index, DefaultValue, false).IsEmpty())
 	{
 		return false;
 	}
-	return FRigVMBaseAction::Undo(InController);
+	return FRigVMBaseAction::Undo();
 }
 
-bool FRigVMRemoveArrayPinAction::Redo(URigVMController* InController)
+bool FRigVMRemoveArrayPinAction::Redo()
 {
-	if (!FRigVMBaseAction::Redo(InController))
+	if (!FRigVMBaseAction::Redo())
 	{
 		return false;
 	}
-	return InController->RemoveArrayPin(FString::Printf(TEXT("%s.%d"), *ArrayPinPath, Index), false);
+	return GetController()->RemoveArrayPin(FString::Printf(TEXT("%s.%d"), *ArrayPinPath, Index), false);
 }
 
-FRigVMAddLinkAction::FRigVMAddLinkAction(URigVMPin* InOutputPin, URigVMPin* InInputPin)
-: OutputPinPath(InOutputPin->GetPinPath())
+FRigVMAddLinkAction::FRigVMAddLinkAction()
+: FRigVMBaseAction(nullptr)
+{
+}
+
+FRigVMAddLinkAction::FRigVMAddLinkAction(URigVMController* InController, URigVMPin* InOutputPin, URigVMPin* InInputPin)
+: FRigVMBaseAction(InController)
+, OutputPinPath(InOutputPin->GetPinPath())
 , InputPinPath(InInputPin->GetPinPath())
 {
 }
 
-bool FRigVMAddLinkAction::Undo(URigVMController* InController)
+bool FRigVMAddLinkAction::Undo()
 {
-	if(!FRigVMBaseAction::Undo(InController))
+	if(!FRigVMBaseAction::Undo())
 	{
 		return false;
 	}
-	return InController->BreakLink(OutputPinPath, InputPinPath, false);
+	return GetController()->BreakLink(OutputPinPath, InputPinPath, false);
 }
 
-bool FRigVMAddLinkAction::Redo(URigVMController* InController)
+bool FRigVMAddLinkAction::Redo()
 {
-	if(!InController->AddLink(OutputPinPath, InputPinPath, false))
+	if(!CanUndoRedo())
+	{
+		return false;
+	}
+	if(!GetController()->AddLink(OutputPinPath, InputPinPath, false))
 	{
 		return false;
 	}
-	return FRigVMBaseAction::Redo(InController);
+	return FRigVMBaseAction::Redo();
+}
+
+FRigVMBreakLinkAction::FRigVMBreakLinkAction()
+: FRigVMBaseAction(nullptr)
+{
 }
 
-FRigVMBreakLinkAction::FRigVMBreakLinkAction(URigVMPin* InOutputPin, URigVMPin* InInputPin)
-: OutputPinPath(InOutputPin->GetPinPath())
+FRigVMBreakLinkAction::FRigVMBreakLinkAction(URigVMController* InController, URigVMPin* InOutputPin, URigVMPin* InInputPin)
+: FRigVMBaseAction(InController)
+, OutputPinPath(InOutputPin->GetPinPath())
 , InputPinPath(InInputPin->GetPinPath())
 {
 	GraphPath = TSoftObjectPtr<URigVMGraph>(Cast<URigVMGraph>(InOutputPin->GetGraph())).GetUniqueID();
 }
 
-bool FRigVMBreakLinkAction::Undo(URigVMController* InController)
+bool FRigVMBreakLinkAction::Undo()
 {
-	if(!FRigVMBaseAction::Undo(InController))
+	if(!FRigVMBaseAction::Undo())
 	{
 		return false;
 	}
-	return InController->AddLink(OutputPinPath, InputPinPath, false);
+	return GetController()->AddLink(OutputPinPath, InputPinPath, false);
 }
 
-bool FRigVMBreakLinkAction::Redo(URigVMController* InController)
+bool FRigVMBreakLinkAction::Redo()
 {
-	if(!InController->BreakLink(OutputPinPath, InputPinPath, false))
+	if(!CanUndoRedo())
+	{
+		return false;
+	}
+	if(!GetController()->BreakLink(OutputPinPath, InputPinPath, false))
 	{
 		return false;
 	}
-	return FRigVMBaseAction::Redo(InController);
+	return FRigVMBaseAction::Redo();
 }
 
 FRigVMChangePinTypeAction::FRigVMChangePinTypeAction()
-: PinPath()
+: FRigVMBaseAction(nullptr)
+, PinPath()
 , OldTypeIndex(INDEX_NONE)
 , NewTypeIndex(INDEX_NONE)
 , bSetupOrphanPins(true)
@@ -1075,8 +1309,9 @@ FRigVMChangePinTypeAction::FRigVMChangePinTypeAction()
 , bRemoveSubPins(true)
 {}
 
-FRigVMChangePinTypeAction::FRigVMChangePinTypeAction(URigVMPin* InPin, int32 InTypeIndex, bool InSetupOrphanPins, bool InBreakLinks, bool InRemoveSubPins)
-: PinPath(InPin->GetPinPath())
+FRigVMChangePinTypeAction::FRigVMChangePinTypeAction(URigVMController* InController, URigVMPin* InPin, int32 InTypeIndex, bool InSetupOrphanPins, bool InBreakLinks, bool InRemoveSubPins)
+: FRigVMBaseAction(InController)
+, PinPath(InPin->GetPinPath())
 , OldTypeIndex(InPin->GetTypeIndex())
 , NewTypeIndex(InTypeIndex)
 , bSetupOrphanPins(InSetupOrphanPins)
@@ -1085,30 +1320,34 @@ FRigVMChangePinTypeAction::FRigVMChangePinTypeAction(URigVMPin* InPin, int32 InT
 {
 }
 
-bool FRigVMChangePinTypeAction::Undo(URigVMController* InController)
+bool FRigVMChangePinTypeAction::Undo()
 {
-	if(!FRigVMBaseAction::Undo(InController))
+	if(!FRigVMBaseAction::Undo())
 	{
 		return false;
 	}
 
-	if(const URigVMGraph* Graph = InController->GetGraph())
+	if(const URigVMGraph* Graph = GetController()->GetGraph())
 	{
 		if(URigVMPin* Pin = Graph->FindPin(PinPath))
 		{
-			return InController->ChangePinType(Pin, OldTypeIndex, false, bSetupOrphanPins, bBreakLinks, bRemoveSubPins);
+			return GetController()->ChangePinType(Pin, OldTypeIndex, false, bSetupOrphanPins, bBreakLinks, bRemoveSubPins);
 		}
 	}
 	return false;
 }
 
-bool FRigVMChangePinTypeAction::Redo(URigVMController* InController)
+bool FRigVMChangePinTypeAction::Redo()
 {
-	if(const URigVMGraph* Graph = InController->GetGraph())
+	if(!CanUndoRedo())
+	{
+		return false;
+	}
+	if(const URigVMGraph* Graph = GetController()->GetGraph())
 	{
 		if(URigVMPin* Pin = Graph->FindPin(PinPath))
 		{
-			if(!InController->ChangePinType(Pin, NewTypeIndex, false, bSetupOrphanPins, bBreakLinks, bRemoveSubPins))
+			if(!GetController()->ChangePinType(Pin, NewTypeIndex, false, bSetupOrphanPins, bBreakLinks, bRemoveSubPins))
 			{
 				return false;
 			}
@@ -1122,16 +1361,20 @@ bool FRigVMChangePinTypeAction::Redo(URigVMController* InController)
 	{
 		return false;
 	}
-	return FRigVMBaseAction::Redo(InController);
+	return FRigVMBaseAction::Redo();
 }
 
 FRigVMCollapseNodesAction::FRigVMCollapseNodesAction()
-	: LibraryNodePath(), bIsAggregate(false)
+: FRigVMBaseAction(nullptr)
+, LibraryNodePath()
+, bIsAggregate(false)
 {
 }
 
 FRigVMCollapseNodesAction::FRigVMCollapseNodesAction(URigVMController* InController, const TArray<URigVMNode*>& InNodes, const FString& InNodePath, const bool bIsAggregate)
-	: LibraryNodePath(InNodePath), bIsAggregate(bIsAggregate)
+: FRigVMBaseAction(InController)
+, LibraryNodePath(InNodePath)
+, bIsAggregate(bIsAggregate)
 {
 	TArray<FName> NodesToExport;
 	for (const URigVMNode* InNode : InNodes)
@@ -1143,22 +1386,22 @@ FRigVMCollapseNodesAction::FRigVMCollapseNodesAction(URigVMController* InControl
 	CollapsedNodesContent = InController->ExportNodesToText(NodesToExport, true);
 }
 
-bool FRigVMCollapseNodesAction::Undo(URigVMController* InController)
+bool FRigVMCollapseNodesAction::Undo()
 {
-	if (!FRigVMBaseAction::Undo(InController))
+	if (!FRigVMBaseAction::Undo())
 	{
 		return false;
 	}
 
 	// remove the library node
 	{
-		if(!InController->RemoveNodeByName(*LibraryNodePath, false, false))
+		if(!GetController()->RemoveNodeByName(*LibraryNodePath, false, false))
 		{
 			return false;
 		}
 	}
 
-	const TArray<FName> RecoveredNodes = InController->ImportNodesFromText(CollapsedNodesContent, false, false);
+	const TArray<FName> RecoveredNodes = GetController()->ImportNodesFromText(CollapsedNodesContent, false, false);
 	if(RecoveredNodes.Num() != CollapsedNodesPaths.Num())
 	{
 		return false;
@@ -1167,8 +1410,12 @@ bool FRigVMCollapseNodesAction::Undo(URigVMController* InController)
 	return true;
 }
 
-bool FRigVMCollapseNodesAction::Redo(URigVMController* InController)
+bool FRigVMCollapseNodesAction::Redo()
 {
+	if(!CanUndoRedo())
+	{
+		return false;
+	}
 #if WITH_EDITOR
 	TArray<FName> NodeNames;
 	for (const FString& NodePath : CollapsedNodesPaths)
@@ -1176,31 +1423,33 @@ bool FRigVMCollapseNodesAction::Redo(URigVMController* InController)
 		NodeNames.Add(*NodePath);
 	}
 
-	URigVMLibraryNode* LibraryNode = InController->CollapseNodes(NodeNames, LibraryNodePath, false, false, bIsAggregate);
+	URigVMLibraryNode* LibraryNode = GetController()->CollapseNodes(NodeNames, LibraryNodePath, false, false, bIsAggregate);
 	if (LibraryNode)
 	{
-		return FRigVMBaseAction::Redo(InController);
+		return FRigVMBaseAction::Redo();
 	}
 #endif
 	return false;
 }
 
 FRigVMExpandNodeAction::FRigVMExpandNodeAction()
-	: LibraryNodePath()
+: FRigVMBaseAction(nullptr)
+, LibraryNodePath()
 {
 }
 
 FRigVMExpandNodeAction::FRigVMExpandNodeAction(URigVMController* InController, URigVMLibraryNode* InLibraryNode)
-	: LibraryNodePath(InLibraryNode->GetName())
+: FRigVMBaseAction(InController)
+, LibraryNodePath(InLibraryNode->GetName())
 {
 	TArray<FName> NodesToExport;
 	NodesToExport.Add(InLibraryNode->GetFName());
 	LibraryNodeContent = InController->ExportNodesToText(NodesToExport, true);
 }
 
-bool FRigVMExpandNodeAction::Undo(URigVMController* InController)
+bool FRigVMExpandNodeAction::Undo()
 {
-	if (!FRigVMBaseAction::Undo(InController))
+	if (!FRigVMBaseAction::Undo())
 	{
 		return false;
 	}
@@ -1211,12 +1460,12 @@ bool FRigVMExpandNodeAction::Undo(URigVMController* InController)
 	{
 		return FName(*NodePath); 
 	});
-	if(!InController->RemoveNodesByName(NodeNames, false, false))
+	if(!GetController()->RemoveNodesByName(NodeNames, false, false))
 	{
 		return false;
 	}
 
-	const TArray<FName> RecoveredNodes = InController->ImportNodesFromText(LibraryNodeContent, false, false);
+	const TArray<FName> RecoveredNodes = GetController()->ImportNodesFromText(LibraryNodeContent, false, false);
 	if(RecoveredNodes.Num() != 1)
 	{
 		return false;
@@ -1225,63 +1474,79 @@ bool FRigVMExpandNodeAction::Undo(URigVMController* InController)
 	return true;
 }
 
-bool FRigVMExpandNodeAction::Redo(URigVMController* InController)
+bool FRigVMExpandNodeAction::Redo()
 {
+	if(!CanUndoRedo())
+	{
+		return false;
+	}
 #if WITH_EDITOR
-	const TArray<URigVMNode*> ExpandedNodes = InController->ExpandLibraryNode(*LibraryNodePath, false);
+	const TArray<URigVMNode*> ExpandedNodes = GetController()->ExpandLibraryNode(*LibraryNodePath, false);
 	if (ExpandedNodes.Num() == ExpandedNodePaths.Num())
 	{
-		return FRigVMBaseAction::Redo(InController);
+		return FRigVMBaseAction::Redo();
 	}
 #endif
 	return false;
 }
 
-FRigVMRenameNodeAction::FRigVMRenameNodeAction(const FName& InOldNodeName, const FName& InNewNodeName)
-	: OldNodeName(InOldNodeName.ToString())
-	, NewNodeName(InNewNodeName.ToString())
+FRigVMRenameNodeAction::FRigVMRenameNodeAction()
+: FRigVMBaseAction(nullptr)
+{
+}
+
+FRigVMRenameNodeAction::FRigVMRenameNodeAction(URigVMController* InController, const FName& InOldNodeName, const FName& InNewNodeName)
+: FRigVMBaseAction(InController)
+, OldNodeName(InOldNodeName.ToString())
+, NewNodeName(InNewNodeName.ToString())
 {
 }
 
-bool FRigVMRenameNodeAction::Undo(URigVMController* InController)
+bool FRigVMRenameNodeAction::Undo()
 {
-	if (!FRigVMBaseAction::Undo(InController))
+	if (!FRigVMBaseAction::Undo())
 	{
 		return false;
 	}
-	if (URigVMNode* Node = InController->GetGraph()->FindNode(NewNodeName))
+	if (URigVMNode* Node = GetController()->GetGraph()->FindNode(NewNodeName))
 	{
-		return InController->RenameNode(Node, *OldNodeName, false);
+		return GetController()->RenameNode(Node, *OldNodeName, false);
 	}
 	return false;
 }
 
-bool FRigVMRenameNodeAction::Redo(URigVMController* InController)
+bool FRigVMRenameNodeAction::Redo()
 {
-	if (URigVMNode* Node = InController->GetGraph()->FindNode(OldNodeName))
+	if(!CanUndoRedo())
 	{
-		return InController->RenameNode(Node, *NewNodeName, false);
+		return false;
 	}
-	else
+	if (URigVMNode* Node = GetController()->GetGraph()->FindNode(OldNodeName))
 	{
-		return false;
+		return GetController()->RenameNode(Node, *NewNodeName, false);
 	}
-	return FRigVMBaseAction::Redo(InController);
+	return false;
 }
 
-FRigVMPushGraphAction::FRigVMPushGraphAction(UObject* InGraph)
+FRigVMPushGraphAction::FRigVMPushGraphAction()
+: FRigVMBaseAction(nullptr)
+{
+}
+
+FRigVMPushGraphAction::FRigVMPushGraphAction(URigVMController* InController, UObject* InGraph)
+: FRigVMBaseAction(InController)
 {
 	GraphPath = TSoftObjectPtr<URigVMGraph>(Cast<URigVMGraph>(InGraph)).GetUniqueID();
 }
 
-bool FRigVMPushGraphAction::Undo(URigVMController* InController)
+bool FRigVMPushGraphAction::Undo()
 {
-	if (!FRigVMBaseAction::Undo(InController))
+	if (!FRigVMBaseAction::Undo())
 	{
 		return false;
 	}
 	
-	if(URigVMGraph* PoppedGraph = InController->PopGraph(false))
+	if(URigVMGraph* PoppedGraph = GetController()->PopGraph(false))
 	{
 		const FSoftObjectPath PoppedGraphPath = TSoftObjectPtr<URigVMGraph>(PoppedGraph).GetUniqueID();
 		return PoppedGraphPath == GraphPath;
@@ -1290,13 +1555,17 @@ bool FRigVMPushGraphAction::Undo(URigVMController* InController)
 	return false;
 }
 
-bool FRigVMPushGraphAction::Redo(URigVMController* InController)
+bool FRigVMPushGraphAction::Redo()
 {
+	if(!CanUndoRedo())
+	{
+		return false;
+	}
 	TSoftObjectPtr<URigVMGraph> GraphPtr(GraphPath);
 	if(URigVMGraph* Graph = GraphPtr.Get())
 	{
-		InController->PushGraph(Graph, false);
-		return FRigVMBaseAction::Redo(InController);
+		GetController()->PushGraph(Graph, false);
+		return FRigVMBaseAction::Redo();
 	}
 	return false;
 }
@@ -1314,14 +1583,20 @@ bool FRigVMPushGraphAction::MakesObsolete(const FRigVMBaseAction* Other) const
 	return false;
 }
 
-FRigVMPopGraphAction::FRigVMPopGraphAction(UObject* InGraph)
+FRigVMPopGraphAction::FRigVMPopGraphAction()
+: FRigVMBaseAction(nullptr)
+{
+}
+
+FRigVMPopGraphAction::FRigVMPopGraphAction(URigVMController* InController, UObject* InGraph)
+: FRigVMBaseAction(InController)
 {
 	GraphPath = TSoftObjectPtr<URigVMGraph>(Cast<URigVMGraph>(InGraph)).GetUniqueID();
 }
 
-bool FRigVMPopGraphAction::Undo(URigVMController* InController)
+bool FRigVMPopGraphAction::Undo()
 {
-	if (!FRigVMBaseAction::Undo(InController))
+	if (!FRigVMBaseAction::Undo())
 	{
 		return false;
 	}
@@ -1329,20 +1604,24 @@ bool FRigVMPopGraphAction::Undo(URigVMController* InController)
 	TSoftObjectPtr<URigVMGraph> GraphPtr(GraphPath);
 	if (URigVMGraph* Graph = GraphPtr.Get())
 	{
-		InController->PushGraph(Graph, false);
+		GetController()->PushGraph(Graph, false);
 		return true;
 	}
 	return false;
 }
 
-bool FRigVMPopGraphAction::Redo(URigVMController* InController)
+bool FRigVMPopGraphAction::Redo()
 {
-	if(URigVMGraph* PoppedGraph = InController->PopGraph(false))
+	if(!CanUndoRedo())
+	{
+		return false;
+	}
+	if(URigVMGraph* PoppedGraph = GetController()->PopGraph(false))
 	{
 		const FSoftObjectPath PoppedGraphPath = TSoftObjectPtr<URigVMGraph>(PoppedGraph).GetUniqueID();
 		if(PoppedGraphPath == GraphPath)
 		{
-			return FRigVMBaseAction::Redo(InController);
+			return FRigVMBaseAction::Redo();
 		}
 	}
 	return false;
@@ -1361,12 +1640,19 @@ bool FRigVMPopGraphAction::MakesObsolete(const FRigVMBaseAction* Other) const
 	return false;
 }
 
-FRigVMAddExposedPinAction::FRigVMAddExposedPinAction(URigVMPin* InPin)
-	: PinName(InPin->GetName())
-	, Direction(InPin->GetDirection())
-	, CPPType(InPin->GetCPPType())
-	, CPPTypeObjectPath()
-	, DefaultValue(InPin->GetDefaultValue())
+FRigVMAddExposedPinAction::FRigVMAddExposedPinAction()
+: FRigVMBaseAction(nullptr)
+, Direction(ERigVMPinDirection::Input)
+{
+}
+
+FRigVMAddExposedPinAction::FRigVMAddExposedPinAction(URigVMController* InController, URigVMPin* InPin)
+: FRigVMBaseAction(InController)
+, PinName(InPin->GetName())
+, Direction(InPin->GetDirection())
+, CPPType(InPin->GetCPPType())
+, CPPTypeObjectPath()
+, DefaultValue(InPin->GetDefaultValue())
 {
 	if (UObject* CPPTypeObject = InPin->GetCPPTypeObject())
 	{
@@ -1374,36 +1660,43 @@ FRigVMAddExposedPinAction::FRigVMAddExposedPinAction(URigVMPin* InPin)
 	}
 }
 
-bool FRigVMAddExposedPinAction::Undo(URigVMController* InController)
+bool FRigVMAddExposedPinAction::Undo()
 {
-	if (!FRigVMBaseAction::Undo(InController))
+	if (!FRigVMBaseAction::Undo())
 	{
 		return false;
 	}
-	return InController->RemoveExposedPin(*PinName, false);
+	return GetController()->RemoveExposedPin(*PinName, false);
 }
 
-bool FRigVMAddExposedPinAction::Redo(URigVMController* InController)
+bool FRigVMAddExposedPinAction::Redo()
 {
-	if (!InController->AddExposedPin(*PinName, Direction, CPPType, *CPPTypeObjectPath, DefaultValue, false).IsNone())
+	if(!CanUndoRedo())
 	{
-		return FRigVMBaseAction::Redo(InController);
+		return false;
+	}
+	if (!GetController()->AddExposedPin(*PinName, Direction, CPPType, *CPPTypeObjectPath, DefaultValue, false).IsNone())
+	{
+		return FRigVMBaseAction::Redo();
 	}
 	return false;
 }
 
 FRigVMRemoveExposedPinAction::FRigVMRemoveExposedPinAction()
-	: PinIndex(INDEX_NONE)
+: FRigVMBaseAction(nullptr)
+, Direction(ERigVMPinDirection::Input)
+, PinIndex(INDEX_NONE)
 {
 }
 
-FRigVMRemoveExposedPinAction::FRigVMRemoveExposedPinAction(URigVMPin* InPin)
-	: PinName(InPin->GetName())
-	, Direction(InPin->GetDirection())
-	, CPPType(InPin->GetCPPType())
-	, CPPTypeObjectPath()
-	, DefaultValue(InPin->GetDefaultValue())
-	, PinIndex(InPin->GetPinIndex())
+FRigVMRemoveExposedPinAction::FRigVMRemoveExposedPinAction(URigVMController* InController, URigVMPin* InPin)
+: FRigVMBaseAction(InController)
+, PinName(InPin->GetName())
+, Direction(InPin->GetDirection())
+, CPPType(InPin->GetCPPType())
+, CPPTypeObjectPath()
+, DefaultValue(InPin->GetDefaultValue())
+, PinIndex(InPin->GetPinIndex())
 {
 	if (UObject* CPPTypeObject = InPin->GetCPPTypeObject())
 	{
@@ -1411,89 +1704,126 @@ FRigVMRemoveExposedPinAction::FRigVMRemoveExposedPinAction(URigVMPin* InPin)
 	}
 }
 
-bool FRigVMRemoveExposedPinAction::Undo(URigVMController* InController)
+bool FRigVMRemoveExposedPinAction::Undo()
 {
-	if(!InController->AddExposedPin(*PinName, Direction, CPPType, *CPPTypeObjectPath, DefaultValue, false).IsNone())
+	if(!CanUndoRedo())
+	{
+		return false;
+	}
+	if(!GetController()->AddExposedPin(*PinName, Direction, CPPType, *CPPTypeObjectPath, DefaultValue, false).IsNone())
 	{
-		if (InController->SetExposedPinIndex(*PinName, PinIndex, false))
+		if (GetController()->SetExposedPinIndex(*PinName, PinIndex, false))
 		{
-			return FRigVMBaseAction::Undo(InController);
+			return FRigVMBaseAction::Undo();
 		}
 	}
 	return false;
 }
 
-bool FRigVMRemoveExposedPinAction::Redo(URigVMController* InController)
+bool FRigVMRemoveExposedPinAction::Redo()
 {
-	if(FRigVMBaseAction::Redo(InController))
+	if(!CanUndoRedo())
+	{
+		return false;
+	}
+	if(FRigVMBaseAction::Redo())
 	{
-		return InController->RemoveExposedPin(*PinName, false);
+		return GetController()->RemoveExposedPin(*PinName, false);
 	}
 	return false;
 }
 
-FRigVMRenameExposedPinAction::FRigVMRenameExposedPinAction(const FName& InOldPinName, const FName& InNewPinName)
-	: OldPinName(InOldPinName.ToString())
-	, NewPinName(InNewPinName.ToString())
+FRigVMRenameExposedPinAction::FRigVMRenameExposedPinAction()
+: FRigVMBaseAction(nullptr)
+{
+}
+
+FRigVMRenameExposedPinAction::FRigVMRenameExposedPinAction(URigVMController* InController, const FName& InOldPinName, const FName& InNewPinName)
+: FRigVMBaseAction(InController)
+, OldPinName(InOldPinName.ToString())
+, NewPinName(InNewPinName.ToString())
 {
 }
 
-bool FRigVMRenameExposedPinAction::Undo(URigVMController* InController)
+bool FRigVMRenameExposedPinAction::Undo()
 {
-	if (!FRigVMBaseAction::Undo(InController))
+	if (!FRigVMBaseAction::Undo())
 	{
 		return false;
 	}
-	return InController->RenameExposedPin(*NewPinName, *OldPinName, false);
+	return GetController()->RenameExposedPin(*NewPinName, *OldPinName, false);
 }
 
-bool FRigVMRenameExposedPinAction::Redo(URigVMController* InController)
+bool FRigVMRenameExposedPinAction::Redo()
 {
-	if(!InController->RenameExposedPin(*OldPinName, *NewPinName, false))
+	if(!CanUndoRedo())
 	{
 		return false;
 	}
-	return FRigVMBaseAction::Redo(InController);
+	if(!GetController()->RenameExposedPin(*OldPinName, *NewPinName, false))
+	{
+		return false;
+	}
+	return FRigVMBaseAction::Redo();
 }
 
 FRigVMSetPinIndexAction::FRigVMSetPinIndexAction()
-	: PinPath()
-	, OldIndex(INDEX_NONE)
-	, NewIndex(INDEX_NONE)
+: FRigVMBaseAction(nullptr)
+, PinPath()
+, OldIndex(INDEX_NONE)
+, NewIndex(INDEX_NONE)
 {
 }
 
-FRigVMSetPinIndexAction::FRigVMSetPinIndexAction(URigVMPin* InPin, int32 InNewIndex)
-	: PinPath(InPin->GetName())
-	, OldIndex(InPin->GetPinIndex())
-	, NewIndex(InNewIndex)
+FRigVMSetPinIndexAction::FRigVMSetPinIndexAction(URigVMController* InController, URigVMPin* InPin, int32 InNewIndex)
+: FRigVMBaseAction(InController)
+, PinPath(InPin->GetName())
+, OldIndex(InPin->GetPinIndex())
+, NewIndex(InNewIndex)
 {
 }
 
-bool FRigVMSetPinIndexAction::Undo(URigVMController* InController)
+bool FRigVMSetPinIndexAction::Undo()
 {
-	if (!FRigVMBaseAction::Undo(InController))
+	if (!FRigVMBaseAction::Undo())
 	{
 		return false;
 	}
-	return InController->SetExposedPinIndex(*PinPath, OldIndex, false);
+	return GetController()->SetExposedPinIndex(*PinPath, OldIndex, false);
 }
 
-bool FRigVMSetPinIndexAction::Redo(URigVMController* InController)
+bool FRigVMSetPinIndexAction::Redo()
 {
-	if (!InController->SetExposedPinIndex(*PinPath, NewIndex, false))
+	if(!CanUndoRedo())
+	{
+		return false;
+	}
+	if (!GetController()->SetExposedPinIndex(*PinPath, NewIndex, false))
 	{
 		return false;
 	}
-	return FRigVMBaseAction::Redo(InController);
+	return FRigVMBaseAction::Redo();
+}
+
+FRigVMSetRemappedVariableAction::FRigVMSetRemappedVariableAction()
+: FRigVMBaseAction(nullptr)
+, InnerVariableName(NAME_None)
+, OldOuterVariableName(NAME_None)
+, NewOuterVariableName(NAME_None)
+{
 }
 
-FRigVMSetRemappedVariableAction::FRigVMSetRemappedVariableAction(URigVMFunctionReferenceNode* InFunctionRefNode,
-	const FName& InInnerVariableName, const FName& InOldOuterVariableName, const FName& InNewOuterVariableName)
-	: NodePath()
-	, InnerVariableName(InInnerVariableName)
-	, OldOuterVariableName(InOldOuterVariableName)
-	, NewOuterVariableName(InNewOuterVariableName)
+FRigVMSetRemappedVariableAction::FRigVMSetRemappedVariableAction(
+	URigVMController* InController,
+	URigVMFunctionReferenceNode* InFunctionRefNode,
+	const FName& InInnerVariableName,
+	const FName& InOldOuterVariableName,
+	const FName& InNewOuterVariableName)
+: FRigVMBaseAction(InController)
+, NodePath()
+, InnerVariableName(InInnerVariableName)
+, OldOuterVariableName(InOldOuterVariableName)
+, NewOuterVariableName(InNewOuterVariableName)
 {
 	if(InFunctionRefNode)
 	{
@@ -1501,296 +1831,379 @@ FRigVMSetRemappedVariableAction::FRigVMSetRemappedVariableAction(URigVMFunctionR
 	}
 }
 
-bool FRigVMSetRemappedVariableAction::Undo(URigVMController* InController)
+bool FRigVMSetRemappedVariableAction::Undo()
 {
-	if (!FRigVMBaseAction::Undo(InController))
+	if (!FRigVMBaseAction::Undo())
 	{
 		return false;
 	}
-	if (URigVMFunctionReferenceNode* Node = Cast<URigVMFunctionReferenceNode>(InController->GetGraph()->FindNode(NodePath)))
+	if (URigVMFunctionReferenceNode* Node = Cast<URigVMFunctionReferenceNode>(GetController()->GetGraph()->FindNode(NodePath)))
 	{
-		return InController->SetRemappedVariable(Node, InnerVariableName, OldOuterVariableName, false);
+		return GetController()->SetRemappedVariable(Node, InnerVariableName, OldOuterVariableName, false);
 	}
 	return false;
 }
 
-bool FRigVMSetRemappedVariableAction::Redo(URigVMController* InController)
+bool FRigVMSetRemappedVariableAction::Redo()
 {
-	if (URigVMFunctionReferenceNode* Node = Cast<URigVMFunctionReferenceNode>(InController->GetGraph()->FindNode(NodePath)))
+	if(!CanUndoRedo())
 	{
-		return InController->SetRemappedVariable(Node, InnerVariableName, NewOuterVariableName, false);
+		return false;
 	}
-	else
+	if (URigVMFunctionReferenceNode* Node = Cast<URigVMFunctionReferenceNode>(GetController()->GetGraph()->FindNode(NodePath)))
 	{
-		return false;
+		return GetController()->SetRemappedVariable(Node, InnerVariableName, NewOuterVariableName, false);
 	}
-	return FRigVMBaseAction::Redo(InController);
+	return false;
+}
+
+FRigVMAddLocalVariableAction::FRigVMAddLocalVariableAction()
+: FRigVMBaseAction(nullptr)
+{
 }
 
-FRigVMAddLocalVariableAction::FRigVMAddLocalVariableAction(const FRigVMGraphVariableDescription& InLocalVariable)
-	: LocalVariable(InLocalVariable)
+FRigVMAddLocalVariableAction::FRigVMAddLocalVariableAction(URigVMController* InController, const FRigVMGraphVariableDescription& InLocalVariable)
+: FRigVMBaseAction(InController)
+, LocalVariable(InLocalVariable)
 {
 	
 }
 
-bool FRigVMAddLocalVariableAction::Undo(URigVMController* InController)
+bool FRigVMAddLocalVariableAction::Undo()
 {
-	if (!FRigVMBaseAction::Undo(InController))
+	if (!FRigVMBaseAction::Undo())
 	{
 		return false;
 	}
-	return InController->RemoveLocalVariable(LocalVariable.Name, false);
+	return GetController()->RemoveLocalVariable(LocalVariable.Name, false);
 }
 
-bool FRigVMAddLocalVariableAction::Redo(URigVMController* InController)
+bool FRigVMAddLocalVariableAction::Redo()
 {
+	if(!CanUndoRedo())
+	{
+		return false;
+	}
 	if (!LocalVariable.Name.IsNone())
 	{
-		return InController->AddLocalVariable(LocalVariable.Name, LocalVariable.CPPType, LocalVariable.CPPTypeObject, LocalVariable.DefaultValue, false).Name.IsNone() == false;
+		return GetController()->AddLocalVariable(LocalVariable.Name, LocalVariable.CPPType, LocalVariable.CPPTypeObject, LocalVariable.DefaultValue, false).Name.IsNone() == false;
 	}
-	return FRigVMBaseAction::Redo(InController);
+	return FRigVMBaseAction::Redo();
+}
+
+FRigVMRemoveLocalVariableAction::FRigVMRemoveLocalVariableAction()
+: FRigVMBaseAction(nullptr)
+{
 }
 
-FRigVMRemoveLocalVariableAction::FRigVMRemoveLocalVariableAction(const FRigVMGraphVariableDescription& InLocalVariable)
-	: LocalVariable(InLocalVariable)
+FRigVMRemoveLocalVariableAction::FRigVMRemoveLocalVariableAction(URigVMController* InController, const FRigVMGraphVariableDescription& InLocalVariable)
+: FRigVMBaseAction(InController)
+, LocalVariable(InLocalVariable)
 {
 }
 
-bool FRigVMRemoveLocalVariableAction::Undo(URigVMController* InController)
+bool FRigVMRemoveLocalVariableAction::Undo()
 {
-	if (!FRigVMBaseAction::Undo(InController))
+	if (!FRigVMBaseAction::Undo())
 	{
 		return false;
 	}
-	return !InController->AddLocalVariable(LocalVariable.Name, LocalVariable.CPPType, LocalVariable.CPPTypeObject, LocalVariable.DefaultValue, false).Name.IsNone();
+	return !GetController()->AddLocalVariable(LocalVariable.Name, LocalVariable.CPPType, LocalVariable.CPPTypeObject, LocalVariable.DefaultValue, false).Name.IsNone();
 }
 
-bool FRigVMRemoveLocalVariableAction::Redo(URigVMController* InController)
+bool FRigVMRemoveLocalVariableAction::Redo()
 {
+	if(!CanUndoRedo())
+	{
+		return false;
+	}
 	if (!LocalVariable.Name.IsNone())
 	{
-		return InController->RemoveLocalVariable(LocalVariable.Name, false);
+		return GetController()->RemoveLocalVariable(LocalVariable.Name, false);
 	}
-	return FRigVMBaseAction::Redo(InController);
+	return FRigVMBaseAction::Redo();
+}
+
+FRigVMRenameLocalVariableAction::FRigVMRenameLocalVariableAction()
+: FRigVMBaseAction(nullptr)
+, OldVariableName(NAME_None)
+, NewVariableName(NAME_None)
+{
 }
 
-FRigVMRenameLocalVariableAction::FRigVMRenameLocalVariableAction(const FName& InOldName, const FName& InNewName)
-	: OldVariableName(InOldName), NewVariableName(InNewName)
+FRigVMRenameLocalVariableAction::FRigVMRenameLocalVariableAction(URigVMController* InController, const FName& InOldName, const FName& InNewName)
+: FRigVMBaseAction(InController)
+, OldVariableName(InOldName)
+, NewVariableName(InNewName)
 {
 	
 }
 
-bool FRigVMRenameLocalVariableAction::Undo(URigVMController* InController)
+bool FRigVMRenameLocalVariableAction::Undo()
 {
-	if (!FRigVMBaseAction::Undo(InController))
+	if (!FRigVMBaseAction::Undo())
 	{
 		return false;
 	}
-	return InController->RenameLocalVariable(NewVariableName, OldVariableName, false);
+	return GetController()->RenameLocalVariable(NewVariableName, OldVariableName, false);
 }
 
-bool FRigVMRenameLocalVariableAction::Redo(URigVMController* InController)
+bool FRigVMRenameLocalVariableAction::Redo()
 {
-	if (!InController->RenameLocalVariable(OldVariableName, NewVariableName, false))
+	if(!CanUndoRedo())
 	{
 		return false;
 	}
-	return FRigVMBaseAction::Redo(InController);
+	if (!GetController()->RenameLocalVariable(OldVariableName, NewVariableName, false))
+	{
+		return false;
+	}
+	return FRigVMBaseAction::Redo();
 }
 
 FRigVMChangeLocalVariableTypeAction::FRigVMChangeLocalVariableTypeAction()
-	: LocalVariable()
-	, CPPType()
-	, CPPTypeObject(nullptr)
+: FRigVMBaseAction(nullptr)
+, LocalVariable()
+, CPPType()
+, CPPTypeObject(nullptr)
 {
 }
 
 FRigVMChangeLocalVariableTypeAction::FRigVMChangeLocalVariableTypeAction(
-	const FRigVMGraphVariableDescription& InLocalVariable, const FString& InCPPType, UObject* InCPPTypeObject)
-		: LocalVariable(InLocalVariable), CPPType(InCPPType), CPPTypeObject(InCPPTypeObject)
+	URigVMController* InController, 
+	const FRigVMGraphVariableDescription& InLocalVariable,
+	const FString& InCPPType,
+	UObject* InCPPTypeObject)
+: FRigVMBaseAction(InController)
+, LocalVariable(InLocalVariable)
+, CPPType(InCPPType)
+, CPPTypeObject(InCPPTypeObject)
 {
 }
 
-bool FRigVMChangeLocalVariableTypeAction::Undo(URigVMController* InController)
+bool FRigVMChangeLocalVariableTypeAction::Undo()
 {
-	if (!FRigVMBaseAction::Undo(InController))
+	if (!FRigVMBaseAction::Undo())
 	{
 		return false;
 	}
-	return InController->SetLocalVariableType(LocalVariable.Name, LocalVariable.CPPType, LocalVariable.CPPTypeObject, false);
+	return GetController()->SetLocalVariableType(LocalVariable.Name, LocalVariable.CPPType, LocalVariable.CPPTypeObject, false);
 }
 
-bool FRigVMChangeLocalVariableTypeAction::Redo(URigVMController* InController)
+bool FRigVMChangeLocalVariableTypeAction::Redo()
 {
-	if (!InController->SetLocalVariableType(LocalVariable.Name, CPPType, CPPTypeObject, false))
+	if(!CanUndoRedo())
+	{
+		return false;
+	}
+	if (!GetController()->SetLocalVariableType(LocalVariable.Name, CPPType, CPPTypeObject, false))
 	{
 		return false;
 	}
-	return FRigVMBaseAction::Redo(InController);
+	return FRigVMBaseAction::Redo();
 }
 
 FRigVMChangeLocalVariableDefaultValueAction::FRigVMChangeLocalVariableDefaultValueAction()
-	: LocalVariable()
-	, DefaultValue()
+: FRigVMBaseAction(nullptr)
+, LocalVariable()
+, DefaultValue()
 {
 }
 
 FRigVMChangeLocalVariableDefaultValueAction::FRigVMChangeLocalVariableDefaultValueAction(
-	const FRigVMGraphVariableDescription& InLocalVariable, const FString& InDefaultValue)
-		: LocalVariable(InLocalVariable), DefaultValue(InDefaultValue)
+	URigVMController* InController, 
+	const FRigVMGraphVariableDescription& InLocalVariable,
+	const FString& InDefaultValue)
+: FRigVMBaseAction(InController)
+, LocalVariable(InLocalVariable)
+, DefaultValue(InDefaultValue)
 {
 }
 
-bool FRigVMChangeLocalVariableDefaultValueAction::Undo(URigVMController* InController)
+bool FRigVMChangeLocalVariableDefaultValueAction::Undo()
 {
-	if (!FRigVMBaseAction::Undo(InController))
+	if (!FRigVMBaseAction::Undo())
 	{
 		return false;
 	}
-	return InController->SetLocalVariableDefaultValue(LocalVariable.Name, LocalVariable.DefaultValue, false);
+	return GetController()->SetLocalVariableDefaultValue(LocalVariable.Name, LocalVariable.DefaultValue, false);
 }
 
-bool FRigVMChangeLocalVariableDefaultValueAction::Redo(URigVMController* InController)
+bool FRigVMChangeLocalVariableDefaultValueAction::Redo()
 {
-	if (!InController->SetLocalVariableDefaultValue(LocalVariable.Name, DefaultValue, false))
+	if(!CanUndoRedo())
+	{
+		return false;
+	}
+	if (!GetController()->SetLocalVariableDefaultValue(LocalVariable.Name, DefaultValue, false))
 	{
 		return false;		
 	}
-	return FRigVMBaseAction::Redo(InController);
+	return FRigVMBaseAction::Redo();
 }
 
 FRigVMPromoteNodeAction::FRigVMPromoteNodeAction()
-	: LibraryNodePath()
-	, FunctionDefinitionPath()
-	, bFromFunctionToCollapseNode(false)
+: FRigVMBaseAction(nullptr)
+, LibraryNodePath()
+, FunctionDefinitionPath()
+, bFromFunctionToCollapseNode(false)
 {
 }
 
-FRigVMPromoteNodeAction::FRigVMPromoteNodeAction(const URigVMNode* InNodeToPromote, const FString& InNodePath, const FString& InFunctionDefinitionPath)
-	: LibraryNodePath(InNodePath)
-	, FunctionDefinitionPath(InFunctionDefinitionPath)
-	, bFromFunctionToCollapseNode(InNodeToPromote->IsA<URigVMFunctionReferenceNode>())
+FRigVMPromoteNodeAction::FRigVMPromoteNodeAction(URigVMController* InController, const URigVMNode* InNodeToPromote, const FString& InNodePath, const FString& InFunctionDefinitionPath)
+: FRigVMBaseAction(InController)
+, LibraryNodePath(InNodePath)
+, FunctionDefinitionPath(InFunctionDefinitionPath)
+, bFromFunctionToCollapseNode(InNodeToPromote->IsA<URigVMFunctionReferenceNode>())
 {
 }
 
-bool FRigVMPromoteNodeAction::Undo(URigVMController* InController)
+bool FRigVMPromoteNodeAction::Undo()
 {
 	if(bFromFunctionToCollapseNode)
 	{
-		const FName FunctionRefNodeName = InController->PromoteCollapseNodeToFunctionReferenceNode(*LibraryNodePath, false, false, FunctionDefinitionPath);
+		const FName FunctionRefNodeName = GetController()->PromoteCollapseNodeToFunctionReferenceNode(*LibraryNodePath, false, false, FunctionDefinitionPath);
 		return FunctionRefNodeName.ToString() == LibraryNodePath;
 	}
 
-	const FName CollapseNodeName = InController->PromoteFunctionReferenceNodeToCollapseNode(*LibraryNodePath, false, false, true);
+	const FName CollapseNodeName = GetController()->PromoteFunctionReferenceNodeToCollapseNode(*LibraryNodePath, false, false, true);
 	return CollapseNodeName.ToString() == LibraryNodePath;
 }
 
-bool FRigVMPromoteNodeAction::Redo(URigVMController* InController)
+bool FRigVMPromoteNodeAction::Redo()
 {
+	if(!CanUndoRedo())
+	{
+		return false;
+	}
 	if(bFromFunctionToCollapseNode)
 	{
-		const FName CollapseNodeName = InController->PromoteFunctionReferenceNodeToCollapseNode(*LibraryNodePath, false, false, false);
+		const FName CollapseNodeName = GetController()->PromoteFunctionReferenceNodeToCollapseNode(*LibraryNodePath, false, false, false);
 		return CollapseNodeName.ToString() == LibraryNodePath;
 	}
-	const FName FunctionRefNodeName = InController->PromoteCollapseNodeToFunctionReferenceNode(*LibraryNodePath, false, false);
+	const FName FunctionRefNodeName = GetController()->PromoteCollapseNodeToFunctionReferenceNode(*LibraryNodePath, false, false);
 	return FunctionRefNodeName.ToString() == LibraryNodePath;
 }
 
 FRigVMMarkFunctionPublicAction::FRigVMMarkFunctionPublicAction()
-	: FunctionName(NAME_None)
-	, bIsPublic(false)
+: FRigVMBaseAction(nullptr)
+, FunctionName(NAME_None)
+, bIsPublic(false)
 {
 }
 
-FRigVMMarkFunctionPublicAction::FRigVMMarkFunctionPublicAction(const FName& InFunctionName, bool bInIsPublic)
-	: FunctionName(InFunctionName)
-	, bIsPublic(bInIsPublic)
+FRigVMMarkFunctionPublicAction::FRigVMMarkFunctionPublicAction(URigVMController* InController, const FName& InFunctionName, bool bInIsPublic)
+: FRigVMBaseAction(InController)
+, FunctionName(InFunctionName)
+, bIsPublic(bInIsPublic)
 {
 }
 
-bool FRigVMMarkFunctionPublicAction::Undo(URigVMController* InController)
+bool FRigVMMarkFunctionPublicAction::Undo()
 {
-	if (!FRigVMBaseAction::Undo(InController))
+	if (!FRigVMBaseAction::Undo())
 	{
 		return false;
 	}
-	return InController->MarkFunctionAsPublic(FunctionName, !bIsPublic, false);
+	return GetController()->MarkFunctionAsPublic(FunctionName, !bIsPublic, false);
 }
 
-bool FRigVMMarkFunctionPublicAction::Redo(URigVMController* InController)
+bool FRigVMMarkFunctionPublicAction::Redo()
 {
+	if(!CanUndoRedo())
+	{
+		return false;
+	}
 #if WITH_EDITOR
-	if (InController->MarkFunctionAsPublic(FunctionName, bIsPublic, false))
+	if (GetController()->MarkFunctionAsPublic(FunctionName, bIsPublic, false))
 	{
-		return FRigVMBaseAction::Redo(InController);
+		return FRigVMBaseAction::Redo();
 	}
 #endif
 	return false;
 }
 
-FRigVMImportFromTextAction::FRigVMImportFromTextAction(const FString& InContent, const TArray<FName>& InTopLevelNodeNames)
-	: Content(InContent)
-	, TopLevelNodeNames(InTopLevelNodeNames)
+FRigVMImportFromTextAction::FRigVMImportFromTextAction()
+: FRigVMBaseAction(nullptr)
+{
+}
+
+FRigVMImportFromTextAction::FRigVMImportFromTextAction(URigVMController* InController, const FString& InContent, const TArray<FName>& InTopLevelNodeNames)
+: FRigVMBaseAction(InController)
+, Content(InContent)
+, TopLevelNodeNames(InTopLevelNodeNames)
 {
 }
 
-FRigVMImportFromTextAction::FRigVMImportFromTextAction(URigVMNode* InNode, URigVMController* InController, bool bIncludeExteriorLinks)
+FRigVMImportFromTextAction::FRigVMImportFromTextAction(URigVMController* InController, URigVMNode* InNode, bool bIncludeExteriorLinks)
+: FRigVMBaseAction(InController)
 {
-	SetContent({InNode}, InController, bIncludeExteriorLinks);
+	SetContent({InNode}, bIncludeExteriorLinks);
 }
 
-FRigVMImportFromTextAction::FRigVMImportFromTextAction(const TArray<URigVMNode*>& InNodes, URigVMController* InController, bool bIncludeExteriorLinks)
+FRigVMImportFromTextAction::FRigVMImportFromTextAction(URigVMController* InController, const TArray<URigVMNode*>& InNodes, bool bIncludeExteriorLinks)
+: FRigVMBaseAction(InController)
 {
-	SetContent(InNodes, InController, bIncludeExteriorLinks);
+	SetContent(InNodes, bIncludeExteriorLinks);
 }
 
-void FRigVMImportFromTextAction::SetContent(const TArray<URigVMNode*>& InNodes, URigVMController* InController, bool bIncludeExteriorLinks)
+void FRigVMImportFromTextAction::SetContent(const TArray<URigVMNode*>& InNodes, bool bIncludeExteriorLinks)
 {
 	TopLevelNodeNames.Reset();
 	Algo::Transform(InNodes, TopLevelNodeNames, [](const URigVMNode* Node)
 	{
 		return Node->GetFName();
 	});
-	Content = InController->ExportNodesToText(TopLevelNodeNames, bIncludeExteriorLinks);
+	Content = GetController()->ExportNodesToText(TopLevelNodeNames, bIncludeExteriorLinks);
 }
 
-bool FRigVMImportFromTextAction::Undo(URigVMController* InController)
+bool FRigVMImportFromTextAction::Undo()
 {
-	if(!FRigVMBaseAction::Undo(InController))
+	if(!FRigVMBaseAction::Undo())
 	{
 		return false;
 	}
-	return InController->RemoveNodesByName(TopLevelNodeNames, false, false);
+	return GetController()->RemoveNodesByName(TopLevelNodeNames, false, false);
 }
 
-bool FRigVMImportFromTextAction::Redo(URigVMController* InController)
+bool FRigVMImportFromTextAction::Redo()
 {
-	const TArray<FName> ImportedNames = InController->ImportNodesFromText(Content, false, false);
+	if(!CanUndoRedo())
+	{
+		return false;
+	}
+	const TArray<FName> ImportedNames = GetController()->ImportNodesFromText(Content, false, false);
 	if(ImportedNames.Num() != TopLevelNodeNames.Num())
 	{
 		return false;
 	}
-	return FRigVMBaseAction::Redo(InController);
+	return FRigVMBaseAction::Redo();
+}
+
+FRigVMReplaceNodesAction::FRigVMReplaceNodesAction()
+: FRigVMBaseAction(nullptr)
+{
 }
 
-FRigVMReplaceNodesAction::FRigVMReplaceNodesAction(const TArray<URigVMNode*>& InNodes, URigVMController* InController)
+FRigVMReplaceNodesAction::FRigVMReplaceNodesAction(URigVMController* InController, const TArray<URigVMNode*>& InNodes)
+: FRigVMBaseAction(InController)
 {
 	for(const URigVMNode* Node : InNodes)
 	{
-		StoreNode(Node, InController);
+		StoreNode(Node, true);
 	}
 }
 
-bool FRigVMReplaceNodesAction::Undo(URigVMController* InController)
+bool FRigVMReplaceNodesAction::Undo()
 {
-	if(!FRigVMBaseAction::Undo(InController))
+	if(!FRigVMBaseAction::Undo())
 	{
 		return false;
 	}
 
 	for(const TPair<FName, FRigVMActionNodeContent>& Pair : ExportedNodes)
 	{
-		if(!RestoreNode(Pair.Key, InController, true))
+		if(!RestoreNode(Pair.Key, true))
 		{
 			return false;
 		}
@@ -1799,16 +2212,20 @@ bool FRigVMReplaceNodesAction::Undo(URigVMController* InController)
 	return true;
 }
 
-bool FRigVMReplaceNodesAction::Redo(URigVMController* InController)
+bool FRigVMReplaceNodesAction::Redo()
 {
+	if(!CanUndoRedo())
+	{
+		return false;
+	}
 	for(const TPair<FName, FRigVMActionNodeContent>& Pair : ExportedNodes)
 	{
-		if(!RestoreNode(Pair.Key, InController, false))
+		if(!RestoreNode(Pair.Key, false))
 		{
 			return false;
 		}
 	}
 	
-	return FRigVMBaseAction::Redo(InController);
+	return FRigVMBaseAction::Redo();
 }
 
diff --git a/Engine/Source/Developer/RigVMDeveloper/Public/RigVMModel/RigVMClient.h b/Engine/Source/Developer/RigVMDeveloper/Public/RigVMModel/RigVMClient.h
--- a/Engine/Source/Developer/RigVMDeveloper/Public/RigVMModel/RigVMClient.h
+++ b/Engine/Source/Developer/RigVMDeveloper/Public/RigVMModel/RigVMClient.h
@@ -166,6 +166,8 @@ public:
 	bool UpdateGraphFunctionSerializedGraph(URigVMLibraryNode* InLibraryNode);
 	bool IsFunctionPublic(URigVMLibraryNode* InLibraryNode);
 
+	static constexpr TCHAR RigVMModelPrefix[] = TEXT("RigVMModel");
+
 private:
 
 	enum ERigVMClientAction
@@ -183,6 +185,7 @@ private:
 	};
 
 	URigVMController* CreateController(const URigVMGraph* InModel);
+	URigVMActionStack* GetOrCreateActionStack();
 
 	UPROPERTY(transient)
 	TObjectPtr<UScriptStruct> ExecuteContextStruct;
@@ -196,6 +199,9 @@ private:
 	UPROPERTY(transient)
 	TMap<FSoftObjectPath, TObjectPtr<URigVMController>> Controllers;
 
+	UPROPERTY(transient)
+	TObjectPtr<URigVMActionStack> ActionStack;
+
 	UPROPERTY(transient)
 	int32 UndoRedoIndex = 0;
 	
diff --git a/Engine/Source/Developer/RigVMDeveloper/Public/RigVMModel/RigVMController.h b/Engine/Source/Developer/RigVMDeveloper/Public/RigVMModel/RigVMController.h
--- a/Engine/Source/Developer/RigVMDeveloper/Public/RigVMModel/RigVMController.h
+++ b/Engine/Source/Developer/RigVMDeveloper/Public/RigVMModel/RigVMController.h
@@ -1271,6 +1271,9 @@ public:
 	FRigVMTemplate::FTypeMap GetCommonlyUsedTypesForTemplate(const URigVMTemplateNode* InNode) const;
 #endif
 
+	URigVMActionStack* GetActionStack() const;
+	void SetActionStack(URigVMActionStack* InActionStack);
+
 protected:
 
 	// backwards compatibility code
@@ -1346,8 +1349,8 @@ private:
 	UPROPERTY(transient)
 	TArray<TObjectPtr<URigVMGraph>> Graphs;
 
-	UPROPERTY(transient)
-	TObjectPtr<URigVMActionStack> ActionStack;
+	TWeakObjectPtr<URigVMActionStack> WeakActionStack;
+	FDelegateHandle ActionStackHandle;
 
 	bool bSuspendNotifications;
 	bool bReportWarningsAndErrors;
diff --git a/Engine/Source/Developer/RigVMDeveloper/Public/RigVMModel/RigVMControllerActions.h b/Engine/Source/Developer/RigVMDeveloper/Public/RigVMModel/RigVMControllerActions.h
--- a/Engine/Source/Developer/RigVMDeveloper/Public/RigVMModel/RigVMControllerActions.h
+++ b/Engine/Source/Developer/RigVMDeveloper/Public/RigVMModel/RigVMControllerActions.h
@@ -79,7 +79,7 @@ public:
 	~FRigVMActionWrapper();
 
 	const UScriptStruct* GetScriptStruct() const;
-	FRigVMBaseAction* GetAction() const;
+	FRigVMBaseAction* GetAction(URigVMController* InLastController = nullptr) const;
 	FString ExportText() const;
 
 private:
@@ -116,6 +116,8 @@ struct FRigVMBaseAction
 {
 	GENERATED_BODY()
 
+public:
+	
 	inline static const FString RedoPrefix = TEXT("Redo");
 	inline static const FString UndoPrefix = TEXT("Undo");
 	inline static const FString AddActionPrefix = TEXT("Add Action");
@@ -124,14 +126,22 @@ struct FRigVMBaseAction
 	inline static const FString CancelActionPrefix = TEXT("Cancel Action");
 
 	// Default constructor
-	FRigVMBaseAction()
-		 : Title(TEXT("Action"))
+	FRigVMBaseAction(URigVMController* InController)
+		: WeakController(InController)
+		, Title(TEXT("Action"))
 	{
+		if(InController)
+		{
+			ControllerPath = FSoftObjectPath(InController);
+		}
 	}
 
 	// Default destructor
 	virtual ~FRigVMBaseAction() {};
 
+	// Returns the controller of this action
+	URigVMController* GetController() const;
+
 	// Access to the actions script struct. this needs to be overloaded
 	virtual UScriptStruct* GetScriptStruct() const { return FRigVMBaseAction::StaticStruct(); }
 
@@ -141,6 +151,9 @@ struct FRigVMBaseAction
 	// Returns the title of the action - used for the Edit menu's undo / redo
 	virtual FString GetTitle() const { return Title; }
 
+	// Sets the title of the action - used for the Edit menu's undo / redo
+	virtual void SetTitle(const FString& InTitle) { Title = InTitle; }
+
 	// Trys to merge the action with another action and 
 	// returns true if successfull.
 	virtual bool Merge(const FRigVMBaseAction* Other);
@@ -148,31 +161,49 @@ struct FRigVMBaseAction
 	// Returns true if this action makes the other action obsolete.
 	virtual bool MakesObsolete(const FRigVMBaseAction* Other) const;
 
+	// Returns true if this action can undo / redo
+	bool CanUndoRedo() const { return GetController() != nullptr; }
+
 	// Un-does the action and returns true if successfull.
-	virtual bool Undo(URigVMController* InController);
+	virtual bool Undo();
 
 	// Re-does the action and returns true if successfull.
-	virtual bool Redo(URigVMController* InController);
+	virtual bool Redo();
 
 	// Adds a child / sub action to this one
 	template<class ActionType>
-	void AddAction(const ActionType& InAction, URigVMController* InController)
+	void AddAction(const ActionType& InAction)
 	{
+		check(InAction.CanUndoRedo());
 		FRigVMActionKey Key;
 		Key.Set<ActionType>(InAction);
 		SubActions.Add(Key);
 	}
 
-	bool StoreNode(const URigVMNode* InNode, URigVMController* InController, bool bIsPriorChange = true);
-	bool RestoreNode(const FName& InNodeName, URigVMController* InController, bool bIsUndoing);
+	bool StoreNode(const URigVMNode* InNode, bool bIsPriorChange = true);
+	bool RestoreNode(const FName& InNodeName, bool bIsUndoing);
 
 #if RIGVM_ACTIONSTACK_VERBOSE_LOG
 	
 	// Logs the action to the console / output log
-	void LogAction(URigVMController* InController, const FString& InPrefix) const;
+	void LogAction(const FString& InPrefix) const;
 
 #endif
 
+protected:
+
+	// Empty constructor
+	FRigVMBaseAction()
+		: WeakController(nullptr)
+		, Title(TEXT("Action"))
+	{
+	}
+
+	mutable TWeakObjectPtr<URigVMController> WeakController;
+
+	UPROPERTY()
+	FSoftObjectPath ControllerPath;
+
 	UPROPERTY()
 	FString Title;
 
@@ -181,6 +212,10 @@ struct FRigVMBaseAction
 
 	UPROPERTY()
 	TMap<FName, FRigVMActionNodeContent> ExportedNodes;
+
+	friend struct FRigVMActionWrapper;
+	friend class URigVMActionStack;
+	friend class UScriptStruct;
 };
 
 /**
@@ -192,7 +227,7 @@ struct FRigVMBaseAction
  * actions occuring during the bracket's lifetime.
  */
 UCLASS()
-class URigVMActionStack : public UObject
+class RIGVMDEVELOPER_API URigVMActionStack : public UObject
 {
 	GENERATED_BODY()
 
@@ -247,13 +282,13 @@ public:
 	// Cancels an action, closes a bracket / scope and discards all 
 	// actions to this point.
 	template<class ActionType>
-	void CancelAction(ActionType& InAction, URigVMController* InController)
+	void CancelAction(ActionType& InAction)
 	{
 		ensure(CurrentActions.Num() > 0);
 		ensure((FRigVMBaseAction*)&InAction == CurrentActions.Last());
 		CurrentActions.Pop();
 
-		InAction.Undo(InController);
+		InAction.Undo();
 
 #if RIGVM_ACTIONSTACK_VERBOSE_LOG		
 		TGuardValue<int32> TabDepthGuard(LogActionDepth, CurrentActions.Num());
@@ -338,7 +373,7 @@ public:
 	// Opens an undo bracket / scope to record actions into.
 	// This is primary useful for Python.
 	UFUNCTION()
-	bool OpenUndoBracket(const FString& InTitle);
+	bool OpenUndoBracket(URigVMController* InController, const FString& InTitle);
 
 	// Closes an undo bracket / scope.
 	// This is primary useful for Python.
@@ -415,11 +450,11 @@ struct FRigVMInjectNodeIntoPinAction : public FRigVMBaseAction
 public:
 
 	FRigVMInjectNodeIntoPinAction();
-	FRigVMInjectNodeIntoPinAction(URigVMInjectionInfo* InInjectionInfo);
+	FRigVMInjectNodeIntoPinAction(URigVMController* InController, URigVMInjectionInfo* InInjectionInfo);
 	virtual ~FRigVMInjectNodeIntoPinAction() {};
 	virtual UScriptStruct* GetScriptStruct() const override { return FRigVMInjectNodeIntoPinAction::StaticStruct(); }
-	virtual bool Undo(URigVMController* InController) override;
-	virtual bool Redo(URigVMController* InController) override;
+	virtual bool Undo() override;
+	virtual bool Redo() override;
 
 	UPROPERTY()
 	FString PinPath;
@@ -448,11 +483,11 @@ struct FRigVMEjectNodeFromPinAction : public FRigVMInjectNodeIntoPinAction
 public:
 
 	FRigVMEjectNodeFromPinAction();
-	FRigVMEjectNodeFromPinAction(URigVMInjectionInfo* InInjectionInfo);
+	FRigVMEjectNodeFromPinAction(URigVMController* InController, URigVMInjectionInfo* InInjectionInfo);
 	virtual ~FRigVMEjectNodeFromPinAction() {};
 	virtual UScriptStruct* GetScriptStruct() const override { return FRigVMEjectNodeFromPinAction::StaticStruct(); }
-	virtual bool Undo(URigVMController* InController) override;
-	virtual bool Redo(URigVMController* InController) override;
+	virtual bool Undo() override;
+	virtual bool Redo() override;
 };
 
 /**
@@ -465,12 +500,12 @@ struct FRigVMRemoveNodesAction : public FRigVMBaseAction
 
 public:
 
-	FRigVMRemoveNodesAction() {}
-	FRigVMRemoveNodesAction(TArray<URigVMNode*> InNodes, URigVMController* InController);
+	FRigVMRemoveNodesAction();
+	FRigVMRemoveNodesAction(URigVMController* InController, TArray<URigVMNode*> InNodes);
 	virtual ~FRigVMRemoveNodesAction() {};
 	virtual UScriptStruct* GetScriptStruct() const override { return FRigVMRemoveNodesAction::StaticStruct(); }
-	virtual bool Undo(URigVMController* InController) override;
-	virtual bool Redo(URigVMController* InController) override;
+	virtual bool Undo() override;
+	virtual bool Redo() override;
 
 	UPROPERTY()
 	TArray<FName> NodeNames;
@@ -490,11 +525,11 @@ struct FRigVMSetNodeSelectionAction : public FRigVMBaseAction
 public:
 
 	FRigVMSetNodeSelectionAction();
-	FRigVMSetNodeSelectionAction(URigVMGraph* InGraph, TArray<FName> InNewSelection);
+	FRigVMSetNodeSelectionAction(URigVMController* InController, URigVMGraph* InGraph, TArray<FName> InNewSelection);
 	virtual ~FRigVMSetNodeSelectionAction() {};
 	virtual UScriptStruct* GetScriptStruct() const override { return FRigVMSetNodeSelectionAction::StaticStruct(); }
-	virtual bool Undo(URigVMController* InController) override;
-	virtual bool Redo(URigVMController* InController) override;
+	virtual bool Undo() override;
+	virtual bool Redo() override;
 
 	UPROPERTY()
 	TArray<FName> NewSelection;
@@ -513,16 +548,13 @@ struct FRigVMSetNodePositionAction : public FRigVMBaseAction
 
 public:
 
-	FRigVMSetNodePositionAction()
-	{
-		OldPosition = NewPosition = FVector2D::ZeroVector;
-	}
-	FRigVMSetNodePositionAction(URigVMNode* InNode, const FVector2D& InNewPosition);
+	FRigVMSetNodePositionAction();
+	FRigVMSetNodePositionAction(URigVMController* InController, URigVMNode* InNode, const FVector2D& InNewPosition);
 	virtual ~FRigVMSetNodePositionAction() {};
 	virtual UScriptStruct* GetScriptStruct() const override { return FRigVMSetNodePositionAction::StaticStruct(); }
 	virtual bool Merge(const FRigVMBaseAction* Other);
-	virtual bool Undo(URigVMController* InController) override;
-	virtual bool Redo(URigVMController* InController) override;
+	virtual bool Undo() override;
+	virtual bool Redo() override;
 
 	UPROPERTY()
 	FString NodePath;
@@ -544,16 +576,13 @@ struct FRigVMSetNodeSizeAction : public FRigVMBaseAction
 
 public:
 
-	FRigVMSetNodeSizeAction()
-	{
-		OldSize = NewSize = FVector2D::ZeroVector;
-	}
-	FRigVMSetNodeSizeAction(URigVMNode* InNode, const FVector2D& InNewSize);
+	FRigVMSetNodeSizeAction();
+	FRigVMSetNodeSizeAction(URigVMController* InController, URigVMNode* InNode, const FVector2D& InNewSize);
 	virtual ~FRigVMSetNodeSizeAction() {};
 	virtual UScriptStruct* GetScriptStruct() const override { return FRigVMSetNodeSizeAction::StaticStruct(); }
 	virtual bool Merge(const FRigVMBaseAction* Other);
-	virtual bool Undo(URigVMController* InController) override;
-	virtual bool Redo(URigVMController* InController) override;
+	virtual bool Undo() override;
+	virtual bool Redo() override;
 
 	UPROPERTY()
 	FString NodePath;
@@ -575,16 +604,13 @@ struct FRigVMSetNodeColorAction : public FRigVMBaseAction
 
 public:
 
-	FRigVMSetNodeColorAction()
-	{
-		OldColor = NewColor = FLinearColor::Black;
-	}
-	FRigVMSetNodeColorAction(URigVMNode* InNode, const FLinearColor& InNewColor);
+	FRigVMSetNodeColorAction();
+	FRigVMSetNodeColorAction(URigVMController* InController, URigVMNode* InNode, const FLinearColor& InNewColor);
 	virtual ~FRigVMSetNodeColorAction() {};
 	virtual UScriptStruct* GetScriptStruct() const override { return FRigVMSetNodeColorAction::StaticStruct(); }
 	virtual bool Merge(const FRigVMBaseAction* Other);
-	virtual bool Undo(URigVMController* InController) override;
-	virtual bool Redo(URigVMController* InController) override;
+	virtual bool Undo() override;
+	virtual bool Redo() override;
 
 	UPROPERTY()
 	FString NodePath;
@@ -606,16 +632,13 @@ struct FRigVMSetNodeCategoryAction : public FRigVMBaseAction
 
 public:
 
-	FRigVMSetNodeCategoryAction()
-	{
-		OldCategory = NewCategory = FString();
-	}
-	FRigVMSetNodeCategoryAction(URigVMCollapseNode* InNode, const FString& InNewCategory);
+	FRigVMSetNodeCategoryAction();
+	FRigVMSetNodeCategoryAction(URigVMController* InController, URigVMCollapseNode* InNode, const FString& InNewCategory);
 	virtual ~FRigVMSetNodeCategoryAction() {};
 	virtual UScriptStruct* GetScriptStruct() const override { return FRigVMSetNodeCategoryAction::StaticStruct(); }
 	virtual bool Merge(const FRigVMBaseAction* Other);
-	virtual bool Undo(URigVMController* InController) override;
-	virtual bool Redo(URigVMController* InController) override;
+	virtual bool Undo() override;
+	virtual bool Redo() override;
 
 	UPROPERTY()
 	FString NodePath;
@@ -638,16 +661,13 @@ struct FRigVMSetNodeKeywordsAction : public FRigVMBaseAction
 
 public:
 
-	FRigVMSetNodeKeywordsAction()
-	{
-		OldKeywords = NewKeywords = FString();
-	}
-	FRigVMSetNodeKeywordsAction(URigVMCollapseNode* InNode, const FString& InNewKeywords);
+	FRigVMSetNodeKeywordsAction();
+	FRigVMSetNodeKeywordsAction(URigVMController* InController, URigVMCollapseNode* InNode, const FString& InNewKeywords);
 	virtual ~FRigVMSetNodeKeywordsAction() {};
 	virtual UScriptStruct* GetScriptStruct() const override { return FRigVMSetNodeKeywordsAction::StaticStruct(); }
 	virtual bool Merge(const FRigVMBaseAction* Other);
-	virtual bool Undo(URigVMController* InController) override;
-	virtual bool Redo(URigVMController* InController) override;
+	virtual bool Undo() override;
+	virtual bool Redo() override;
 
 	UPROPERTY()
 	FString NodePath;
@@ -669,16 +689,13 @@ struct FRigVMSetNodeDescriptionAction : public FRigVMBaseAction
 
 public:
 
-	FRigVMSetNodeDescriptionAction()
-	{
-		OldDescription = NewDescription = FString();
-	}
-	FRigVMSetNodeDescriptionAction(URigVMCollapseNode* InNode, const FString& InNewDescription);
+	FRigVMSetNodeDescriptionAction();
+	FRigVMSetNodeDescriptionAction(URigVMController* InController, URigVMCollapseNode* InNode, const FString& InNewDescription);
 	virtual ~FRigVMSetNodeDescriptionAction() {};
 	virtual UScriptStruct* GetScriptStruct() const override { return FRigVMSetNodeDescriptionAction::StaticStruct(); }
 	virtual bool Merge(const FRigVMBaseAction* Other);
-	virtual bool Undo(URigVMController* InController) override;
-	virtual bool Redo(URigVMController* InController) override;
+	virtual bool Undo() override;
+	virtual bool Redo() override;
 
 	UPROPERTY()
 	FString NodePath;
@@ -701,11 +718,11 @@ struct FRigVMSetCommentTextAction : public FRigVMBaseAction
 public:
 
 	FRigVMSetCommentTextAction();
-	FRigVMSetCommentTextAction(URigVMCommentNode* InNode, const FString& InNewText, const int32& InNewFontSize, const bool& bInNewBubbleVisible, const bool& bInNewColorBubble);
+	FRigVMSetCommentTextAction(URigVMController* InController, URigVMCommentNode* InNode, const FString& InNewText, const int32& InNewFontSize, const bool& bInNewBubbleVisible, const bool& bInNewColorBubble);
 	virtual ~FRigVMSetCommentTextAction() {};
 	virtual UScriptStruct* GetScriptStruct() const override { return FRigVMSetCommentTextAction::StaticStruct(); }
-	virtual bool Undo(URigVMController* InController) override;
-	virtual bool Redo(URigVMController* InController) override;
+	virtual bool Undo() override;
+	virtual bool Redo() override;
 
 	UPROPERTY()
 	FString NodePath;
@@ -746,11 +763,11 @@ struct FRigVMSetRerouteCompactnessAction : public FRigVMBaseAction
 public:
 
 	FRigVMSetRerouteCompactnessAction();
-	FRigVMSetRerouteCompactnessAction(URigVMRerouteNode* InNode, bool InShowAsFullNode);
+	FRigVMSetRerouteCompactnessAction(URigVMController* InController, URigVMRerouteNode* InNode, bool InShowAsFullNode);
 	virtual ~FRigVMSetRerouteCompactnessAction() {};
 	virtual UScriptStruct* GetScriptStruct() const override { return FRigVMSetRerouteCompactnessAction::StaticStruct(); }
-	virtual bool Undo(URigVMController* InController) override;
-	virtual bool Redo(URigVMController* InController) override;
+	virtual bool Undo() override;
+	virtual bool Redo() override;
 
 	UPROPERTY()
 	FString NodePath;
@@ -772,12 +789,12 @@ struct FRigVMRenameVariableAction : public FRigVMBaseAction
 
 public:
 
-	FRigVMRenameVariableAction() {}
-	FRigVMRenameVariableAction(const FName& InOldVariableName, const FName& InNewVariableName);
+	FRigVMRenameVariableAction();
+	FRigVMRenameVariableAction(URigVMController* InController, const FName& InOldVariableName, const FName& InNewVariableName);
 	virtual ~FRigVMRenameVariableAction() {};
 	virtual UScriptStruct* GetScriptStruct() const override { return FRigVMRenameVariableAction::StaticStruct(); }
-	virtual bool Undo(URigVMController* InController) override;
-	virtual bool Redo(URigVMController* InController) override;
+	virtual bool Undo() override;
+	virtual bool Redo() override;
 
 	UPROPERTY()
 	FString OldVariableName;
@@ -796,17 +813,12 @@ struct FRigVMSetPinExpansionAction : public FRigVMBaseAction
 
 public:
 
-	FRigVMSetPinExpansionAction()
-	{
-		PinPath = FString();
-		OldIsExpanded = NewIsExpanded = false;
-	}
-
-	FRigVMSetPinExpansionAction(URigVMPin* InPin, bool bNewIsExpanded);
+	FRigVMSetPinExpansionAction();
+	FRigVMSetPinExpansionAction(URigVMController* InController, URigVMPin* InPin, bool bNewIsExpanded);
 	virtual ~FRigVMSetPinExpansionAction() {};
 	virtual UScriptStruct* GetScriptStruct() const override { return FRigVMSetPinExpansionAction::StaticStruct(); }
-	virtual bool Undo(URigVMController* InController) override;
-	virtual bool Redo(URigVMController* InController) override;
+	virtual bool Undo() override;
+	virtual bool Redo() override;
 
 	UPROPERTY()
 	FString PinPath;
@@ -828,16 +840,12 @@ struct FRigVMSetPinWatchAction : public FRigVMBaseAction
 
 public:
 
-	FRigVMSetPinWatchAction()
-	{
-		OldIsWatched = NewIsWatched = false;
-	}
-
-	FRigVMSetPinWatchAction(URigVMPin* InPin, bool bNewIsWatched);
+	FRigVMSetPinWatchAction();
+	FRigVMSetPinWatchAction(URigVMController* InController, URigVMPin* InPin, bool bNewIsWatched);
 	virtual ~FRigVMSetPinWatchAction() {};
 	virtual UScriptStruct* GetScriptStruct() const override { return FRigVMSetPinWatchAction::StaticStruct(); }
-	virtual bool Undo(URigVMController* InController) override;
-	virtual bool Redo(URigVMController* InController) override;
+	virtual bool Undo() override;
+	virtual bool Redo() override;
 
 	UPROPERTY()
 	FString PinPath;
@@ -859,13 +867,13 @@ struct FRigVMSetPinDefaultValueAction : public FRigVMBaseAction
 
 public:
 
-	FRigVMSetPinDefaultValueAction() {}
-	FRigVMSetPinDefaultValueAction(URigVMPin* InPin, const FString& InNewDefaultValue);
+	FRigVMSetPinDefaultValueAction();
+	FRigVMSetPinDefaultValueAction(URigVMController* InController, URigVMPin* InPin, const FString& InNewDefaultValue);
 	virtual ~FRigVMSetPinDefaultValueAction() {};
 	virtual UScriptStruct* GetScriptStruct() const override { return FRigVMSetPinDefaultValueAction::StaticStruct(); }
 	virtual bool Merge(const FRigVMBaseAction* Other);
-	virtual bool Undo(URigVMController* InController) override;
-	virtual bool Redo(URigVMController* InController) override;
+	virtual bool Undo() override;
+	virtual bool Redo() override;
 
 	UPROPERTY()
 	FString PinPath;
@@ -887,15 +895,12 @@ struct FRigVMInsertArrayPinAction : public FRigVMBaseAction
 
 public:
 
-	FRigVMInsertArrayPinAction()
-		: Index(0)
-	{
-	}
-	FRigVMInsertArrayPinAction(URigVMPin* InArrayPin, int32 InIndex, const FString& InNewDefaultValue);
+	FRigVMInsertArrayPinAction();
+	FRigVMInsertArrayPinAction(URigVMController* InController, URigVMPin* InArrayPin, int32 InIndex, const FString& InNewDefaultValue);
 	virtual ~FRigVMInsertArrayPinAction() {};
 	virtual UScriptStruct* GetScriptStruct() const override { return FRigVMInsertArrayPinAction::StaticStruct(); }
-	virtual bool Undo(URigVMController* InController) override;
-	virtual bool Redo(URigVMController* InController) override;
+	virtual bool Undo() override;
+	virtual bool Redo() override;
 
 	UPROPERTY()
 	FString ArrayPinPath;
@@ -917,15 +922,12 @@ struct FRigVMRemoveArrayPinAction : public FRigVMBaseAction
 
 public:
 
-	FRigVMRemoveArrayPinAction()
-		: Index(0)
-	{
-	}
-	FRigVMRemoveArrayPinAction(URigVMPin* InArrayElementPin);
+	FRigVMRemoveArrayPinAction();
+	FRigVMRemoveArrayPinAction(URigVMController* InController, URigVMPin* InArrayElementPin);
 	virtual ~FRigVMRemoveArrayPinAction() {};
 	virtual UScriptStruct* GetScriptStruct() const override { return FRigVMRemoveArrayPinAction::StaticStruct(); }
-	virtual bool Undo(URigVMController* InController) override;
-	virtual bool Redo(URigVMController* InController) override;
+	virtual bool Undo() override;
+	virtual bool Redo() override;
 
 	UPROPERTY()
 	FString ArrayPinPath;
@@ -947,12 +949,12 @@ struct FRigVMAddLinkAction : public FRigVMBaseAction
 
 public:
 
-	FRigVMAddLinkAction() {}
-	FRigVMAddLinkAction(URigVMPin* InOutputPin, URigVMPin* InInputPin);
+	FRigVMAddLinkAction();
+	FRigVMAddLinkAction(URigVMController* InController, URigVMPin* InOutputPin, URigVMPin* InInputPin);
 	virtual ~FRigVMAddLinkAction() {};
 	virtual UScriptStruct* GetScriptStruct() const override { return FRigVMAddLinkAction::StaticStruct(); }
-	virtual bool Undo(URigVMController* InController) override;
-	virtual bool Redo(URigVMController* InController) override;
+	virtual bool Undo() override;
+	virtual bool Redo() override;
 
 	UPROPERTY()
 	FString OutputPinPath;
@@ -971,12 +973,12 @@ struct FRigVMBreakLinkAction : public FRigVMBaseAction
 
 public:
 
-	FRigVMBreakLinkAction() {}
-	FRigVMBreakLinkAction(URigVMPin* InOutputPin, URigVMPin* InInputPin);
+	FRigVMBreakLinkAction();
+	FRigVMBreakLinkAction(URigVMController* InController, URigVMPin* InOutputPin, URigVMPin* InInputPin);
 	virtual ~FRigVMBreakLinkAction() {};
 	virtual UScriptStruct* GetScriptStruct() const override { return FRigVMBreakLinkAction::StaticStruct(); }
-	virtual bool Undo(URigVMController* InController) override;
-	virtual bool Redo(URigVMController* InController) override;
+	virtual bool Undo() override;
+	virtual bool Redo() override;
 
 	UPROPERTY()
 	FSoftObjectPath GraphPath;
@@ -999,11 +1001,11 @@ struct FRigVMChangePinTypeAction : public FRigVMBaseAction
 public:
 
 	FRigVMChangePinTypeAction();
-	FRigVMChangePinTypeAction(URigVMPin* InPin, int32 InTypeIndex, bool InSetupOrphanPins, bool InBreakLinks, bool InRemoveSubPins);
+	FRigVMChangePinTypeAction(URigVMController* InController, URigVMPin* InPin, int32 InTypeIndex, bool InSetupOrphanPins, bool InBreakLinks, bool InRemoveSubPins);
 	virtual ~FRigVMChangePinTypeAction() {};
 	virtual UScriptStruct* GetScriptStruct() const override { return FRigVMChangePinTypeAction::StaticStruct(); }
-	virtual bool Undo(URigVMController* InController) override;
-	virtual bool Redo(URigVMController* InController) override;
+	virtual bool Undo() override;
+	virtual bool Redo() override;
 
 	UPROPERTY()
 	FString PinPath;
@@ -1038,8 +1040,8 @@ public:
 	FRigVMCollapseNodesAction(URigVMController* InController, const TArray<URigVMNode*>& InNodes, const FString& InNodePath, const bool bIsAggregate);
 	virtual ~FRigVMCollapseNodesAction() {};
 	virtual UScriptStruct* GetScriptStruct() const override { return FRigVMCollapseNodesAction::StaticStruct(); }
-	virtual bool Undo(URigVMController* InController) override;
-	virtual bool Redo(URigVMController* InController) override;
+	virtual bool Undo() override;
+	virtual bool Redo() override;
 
 	UPROPERTY()
 	FString LibraryNodePath;
@@ -1068,8 +1070,8 @@ public:
 	FRigVMExpandNodeAction(URigVMController* InController, URigVMLibraryNode* InLibraryNode);
 	virtual ~FRigVMExpandNodeAction() {};
 	virtual UScriptStruct* GetScriptStruct() const override { return FRigVMExpandNodeAction::StaticStruct(); }
-	virtual bool Undo(URigVMController* InController) override;
-	virtual bool Redo(URigVMController* InController) override;
+	virtual bool Undo() override;
+	virtual bool Redo() override;
 
 	UPROPERTY()
 	FString LibraryNodePath;
@@ -1091,12 +1093,12 @@ struct FRigVMRenameNodeAction : public FRigVMBaseAction
 
 public:
 
-	FRigVMRenameNodeAction() {}
-	FRigVMRenameNodeAction(const FName& InOldNodeName, const FName& InNewNodeName);
+	FRigVMRenameNodeAction();
+	FRigVMRenameNodeAction(URigVMController* InController, const FName& InOldNodeName, const FName& InNewNodeName);
 	virtual ~FRigVMRenameNodeAction() {};
 	virtual UScriptStruct* GetScriptStruct() const override { return FRigVMRenameNodeAction::StaticStruct(); }
-	virtual bool Undo(URigVMController* InController) override;
-	virtual bool Redo(URigVMController* InController) override;
+	virtual bool Undo() override;
+	virtual bool Redo() override;
 
 	UPROPERTY()
 	FString OldNodeName;
@@ -1115,12 +1117,12 @@ struct FRigVMPushGraphAction : public FRigVMBaseAction
 
 public:
 
-	FRigVMPushGraphAction() {}
-	FRigVMPushGraphAction(UObject* InGraph);
+	FRigVMPushGraphAction();
+	FRigVMPushGraphAction(URigVMController* InController, UObject* InGraph);
 	virtual ~FRigVMPushGraphAction() {};
 	virtual UScriptStruct* GetScriptStruct() const override { return FRigVMPushGraphAction::StaticStruct(); }
-	virtual bool Undo(URigVMController* InController) override;
-	virtual bool Redo(URigVMController* InController) override;
+	virtual bool Undo() override;
+	virtual bool Redo() override;
 	virtual bool MakesObsolete(const FRigVMBaseAction* Other) const override;
 
 	UPROPERTY()
@@ -1137,12 +1139,12 @@ struct FRigVMPopGraphAction : public FRigVMBaseAction
 
 public:
 
-	FRigVMPopGraphAction() {}
-	FRigVMPopGraphAction(UObject* InGraph);
+	FRigVMPopGraphAction();
+	FRigVMPopGraphAction(URigVMController* InController, UObject* InGraph);
 	virtual ~FRigVMPopGraphAction() {};
 	virtual UScriptStruct* GetScriptStruct() const override { return FRigVMPopGraphAction::StaticStruct(); }
-	virtual bool Undo(URigVMController* InController) override;
-	virtual bool Redo(URigVMController* InController) override;
+	virtual bool Undo() override;
+	virtual bool Redo() override;
 	virtual bool MakesObsolete(const FRigVMBaseAction* Other) const override;
 
 	UPROPERTY()
@@ -1160,18 +1162,18 @@ struct FRigVMAddExposedPinAction : public FRigVMBaseAction
 
 public:
 
-	FRigVMAddExposedPinAction() {}
-	FRigVMAddExposedPinAction(URigVMPin* InPin);
+	FRigVMAddExposedPinAction();
+	FRigVMAddExposedPinAction(URigVMController* InController, URigVMPin* InPin);
 	virtual ~FRigVMAddExposedPinAction() {};
 	virtual UScriptStruct* GetScriptStruct() const override { return FRigVMAddExposedPinAction::StaticStruct(); }
-	virtual bool Undo(URigVMController* InController) override;
-	virtual bool Redo(URigVMController* InController) override;
+	virtual bool Undo() override;
+	virtual bool Redo() override;
 
 	UPROPERTY()
 	FString PinName;
 	
 	UPROPERTY()
-	ERigVMPinDirection Direction = ERigVMPinDirection::Input;
+	ERigVMPinDirection Direction;
 
 	UPROPERTY()
 	FString CPPType;
@@ -1194,17 +1196,17 @@ struct FRigVMRemoveExposedPinAction : public FRigVMBaseAction
 public:
 
 	FRigVMRemoveExposedPinAction();
-	FRigVMRemoveExposedPinAction(URigVMPin* InPin);
+	FRigVMRemoveExposedPinAction(URigVMController* InController, URigVMPin* InPin);
 	virtual ~FRigVMRemoveExposedPinAction() {};
 	virtual UScriptStruct* GetScriptStruct() const override { return FRigVMRemoveExposedPinAction::StaticStruct(); }
-	virtual bool Undo(URigVMController* InController) override;
-	virtual bool Redo(URigVMController* InController) override;
+	virtual bool Undo() override;
+	virtual bool Redo() override;
 
 	UPROPERTY()
 	FString PinName;
 	
 	UPROPERTY()
-	ERigVMPinDirection Direction = ERigVMPinDirection::Input;
+	ERigVMPinDirection Direction;
 
 	UPROPERTY()
 	FString CPPType;
@@ -1229,12 +1231,12 @@ struct FRigVMRenameExposedPinAction : public FRigVMBaseAction
 
 public:
 
-	FRigVMRenameExposedPinAction() {}
-	FRigVMRenameExposedPinAction(const FName& InOldPinName, const FName& InNewPinName);
+	FRigVMRenameExposedPinAction();
+	FRigVMRenameExposedPinAction(URigVMController* InController, const FName& InOldPinName, const FName& InNewPinName);
 	virtual ~FRigVMRenameExposedPinAction() {};
 	virtual UScriptStruct* GetScriptStruct() const override { return FRigVMRenameExposedPinAction::StaticStruct(); }
-	virtual bool Undo(URigVMController* InController) override;
-	virtual bool Redo(URigVMController* InController) override;
+	virtual bool Undo() override;
+	virtual bool Redo() override;
 
 	UPROPERTY()
 	FString OldPinName;
@@ -1254,11 +1256,11 @@ struct FRigVMSetPinIndexAction : public FRigVMBaseAction
 public:
 
 	FRigVMSetPinIndexAction();
-	FRigVMSetPinIndexAction(URigVMPin* InPin, int32 InNewIndex);
+	FRigVMSetPinIndexAction(URigVMController* InController, URigVMPin* InPin, int32 InNewIndex);
 	virtual ~FRigVMSetPinIndexAction() {};
 	virtual UScriptStruct* GetScriptStruct() const override { return FRigVMSetPinIndexAction::StaticStruct(); }
-	virtual bool Undo(URigVMController* InController) override;
-	virtual bool Redo(URigVMController* InController) override;
+	virtual bool Undo() override;
+	virtual bool Redo() override;
 
 	UPROPERTY()
 	FString PinPath;
@@ -1280,13 +1282,13 @@ struct FRigVMSetRemappedVariableAction : public FRigVMBaseAction
 
 public:
 
-	FRigVMSetRemappedVariableAction() {}
-	FRigVMSetRemappedVariableAction(URigVMFunctionReferenceNode* InFunctionRefNode, const FName& InInnerVariableName,
+	FRigVMSetRemappedVariableAction();
+	FRigVMSetRemappedVariableAction(URigVMController* InController, URigVMFunctionReferenceNode* InFunctionRefNode, const FName& InInnerVariableName,
 		const FName& InOldOuterVariableName, const FName& InNewOuterVariableName);
 	virtual ~FRigVMSetRemappedVariableAction() {};
 	virtual UScriptStruct* GetScriptStruct() const override { return FRigVMSetRemappedVariableAction::StaticStruct(); }
-	virtual bool Undo(URigVMController* InController) override;
-	virtual bool Redo(URigVMController* InController) override;
+	virtual bool Undo() override;
+	virtual bool Redo() override;
 
 	UPROPERTY()
 	FString NodePath;
@@ -1311,12 +1313,12 @@ struct FRigVMAddLocalVariableAction : public FRigVMBaseAction
 
 public:
 
-	FRigVMAddLocalVariableAction() {}
-	FRigVMAddLocalVariableAction(const FRigVMGraphVariableDescription& InLocalVariable);
+	FRigVMAddLocalVariableAction();
+	FRigVMAddLocalVariableAction(URigVMController* InController, const FRigVMGraphVariableDescription& InLocalVariable);
 	virtual ~FRigVMAddLocalVariableAction() {};
 	virtual UScriptStruct* GetScriptStruct() const override { return FRigVMAddLocalVariableAction::StaticStruct(); }
-	virtual bool Undo(URigVMController* InController) override;
-	virtual bool Redo(URigVMController* InController) override;
+	virtual bool Undo() override;
+	virtual bool Redo() override;
 
 	UPROPERTY()
 	FRigVMGraphVariableDescription LocalVariable;
@@ -1332,12 +1334,12 @@ struct FRigVMRemoveLocalVariableAction : public FRigVMBaseAction
 
 public:
 
-	FRigVMRemoveLocalVariableAction() {}
-	FRigVMRemoveLocalVariableAction(const FRigVMGraphVariableDescription& InLocalVariable);
+	FRigVMRemoveLocalVariableAction();
+	FRigVMRemoveLocalVariableAction(URigVMController* InController, const FRigVMGraphVariableDescription& InLocalVariable);
 	virtual ~FRigVMRemoveLocalVariableAction() {};
 	virtual UScriptStruct* GetScriptStruct() const override { return FRigVMRemoveLocalVariableAction::StaticStruct(); }
-	virtual bool Undo(URigVMController* InController) override;
-	virtual bool Redo(URigVMController* InController) override;
+	virtual bool Undo() override;
+	virtual bool Redo() override;
 
 	UPROPERTY()
 	FRigVMGraphVariableDescription LocalVariable;
@@ -1353,12 +1355,12 @@ struct FRigVMRenameLocalVariableAction : public FRigVMBaseAction
 
 public:
 
-	FRigVMRenameLocalVariableAction() {}
-	FRigVMRenameLocalVariableAction(const FName& InOldName, const FName& InNewName);
+	FRigVMRenameLocalVariableAction();
+	FRigVMRenameLocalVariableAction(URigVMController* InController, const FName& InOldName, const FName& InNewName);
 	virtual ~FRigVMRenameLocalVariableAction() {};
 	virtual UScriptStruct* GetScriptStruct() const override { return FRigVMRenameLocalVariableAction::StaticStruct(); }
-	virtual bool Undo(URigVMController* InController) override;
-	virtual bool Redo(URigVMController* InController) override;
+	virtual bool Undo() override;
+	virtual bool Redo() override;
 
 	UPROPERTY()
 	FName OldVariableName;
@@ -1377,11 +1379,11 @@ struct FRigVMChangeLocalVariableTypeAction : public FRigVMBaseAction
 
 public:
 	FRigVMChangeLocalVariableTypeAction();
-	FRigVMChangeLocalVariableTypeAction(const FRigVMGraphVariableDescription& InLocalVariable, const FString& InCPPType, UObject* InCPPTypeObject);
+	FRigVMChangeLocalVariableTypeAction(URigVMController* InController, const FRigVMGraphVariableDescription& InLocalVariable, const FString& InCPPType, UObject* InCPPTypeObject);
 	virtual ~FRigVMChangeLocalVariableTypeAction() {};
 	virtual UScriptStruct* GetScriptStruct() const override { return FRigVMChangeLocalVariableTypeAction::StaticStruct(); }
-	virtual bool Undo(URigVMController* InController) override;
-	virtual bool Redo(URigVMController* InController) override;
+	virtual bool Undo() override;
+	virtual bool Redo() override;
 
 	UPROPERTY()
 	FRigVMGraphVariableDescription LocalVariable;
@@ -1403,11 +1405,11 @@ struct FRigVMChangeLocalVariableDefaultValueAction : public FRigVMBaseAction
 
 public:
 	FRigVMChangeLocalVariableDefaultValueAction();
-	FRigVMChangeLocalVariableDefaultValueAction(const FRigVMGraphVariableDescription& InLocalVariable, const FString& InDefaultValue);
+	FRigVMChangeLocalVariableDefaultValueAction(URigVMController* InController, const FRigVMGraphVariableDescription& InLocalVariable, const FString& InDefaultValue);
 	virtual ~FRigVMChangeLocalVariableDefaultValueAction() {};
 	virtual UScriptStruct* GetScriptStruct() const override { return FRigVMChangeLocalVariableDefaultValueAction::StaticStruct(); }
-	virtual bool Undo(URigVMController* InController) override;
-	virtual bool Redo(URigVMController* InController) override;
+	virtual bool Undo() override;
+	virtual bool Redo() override;
 
 	UPROPERTY()
 	FRigVMGraphVariableDescription LocalVariable;
@@ -1427,11 +1429,11 @@ struct FRigVMPromoteNodeAction : public FRigVMBaseAction
 public:
 
 	FRigVMPromoteNodeAction();
-	FRigVMPromoteNodeAction(const URigVMNode* InNodeToPromote, const FString& InNodePath, const FString& InFunctionDefinitionPath);
+	FRigVMPromoteNodeAction(URigVMController* InController, const URigVMNode* InNodeToPromote, const FString& InNodePath, const FString& InFunctionDefinitionPath);
 	virtual ~FRigVMPromoteNodeAction() {};
 	virtual UScriptStruct* GetScriptStruct() const override { return FRigVMPromoteNodeAction::StaticStruct(); }
-	virtual bool Undo(URigVMController* InController) override;
-	virtual bool Redo(URigVMController* InController) override;
+	virtual bool Undo() override;
+	virtual bool Redo() override;
 
 	UPROPERTY()
 	FString LibraryNodePath;
@@ -1454,11 +1456,11 @@ struct FRigVMMarkFunctionPublicAction : public FRigVMBaseAction
 public:
 
 	FRigVMMarkFunctionPublicAction();
-	FRigVMMarkFunctionPublicAction(const FName& InFunctionName, bool bInIsPublic);
+	FRigVMMarkFunctionPublicAction(URigVMController* InController, const FName& InFunctionName, bool bInIsPublic);
 	virtual ~FRigVMMarkFunctionPublicAction() {};
 	virtual UScriptStruct* GetScriptStruct() const override { return FRigVMMarkFunctionPublicAction::StaticStruct(); }
-	virtual bool Undo(URigVMController* InController) override;
-	virtual bool Redo(URigVMController* InController) override;
+	virtual bool Undo() override;
+	virtual bool Redo() override;
 
 	UPROPERTY()
 	FName FunctionName;
@@ -1477,15 +1479,15 @@ struct FRigVMImportFromTextAction : public FRigVMBaseAction
 
 public:
 
-	FRigVMImportFromTextAction() {}
-	FRigVMImportFromTextAction(const FString& InContent, const TArray<FName>& InTopLevelNodeNames);
-	FRigVMImportFromTextAction(URigVMNode* InNode, URigVMController* InController, bool bIncludeExteriorLinks = true);
-	FRigVMImportFromTextAction(const TArray<URigVMNode*>& InNodes, URigVMController* InController, bool bIncludeExteriorLinks = true);
-	void SetContent(const TArray<URigVMNode*>& InNodes, URigVMController* InController, bool bIncludeExteriorLinks = true);
+	FRigVMImportFromTextAction();
+	FRigVMImportFromTextAction(URigVMController* InController, const FString& InContent, const TArray<FName>& InTopLevelNodeNames);
+	FRigVMImportFromTextAction(URigVMController* InController, URigVMNode* InNode, bool bIncludeExteriorLinks = true);
+	FRigVMImportFromTextAction(URigVMController* InController, const TArray<URigVMNode*>& InNodes, bool bIncludeExteriorLinks = true);
+	void SetContent(const TArray<URigVMNode*>& InNodes, bool bIncludeExteriorLinks = true);
 	virtual ~FRigVMImportFromTextAction() {};
 	virtual UScriptStruct* GetScriptStruct() const override { return FRigVMImportFromTextAction::StaticStruct(); }
-	virtual bool Undo(URigVMController* InController) override;
-	virtual bool Redo(URigVMController* InController) override;
+	virtual bool Undo() override;
+	virtual bool Redo() override;
 
 	UPROPERTY()
 	FString Content;
@@ -1504,10 +1506,10 @@ struct FRigVMReplaceNodesAction : public FRigVMBaseAction
 
 public:
 
-	FRigVMReplaceNodesAction() {}
-	FRigVMReplaceNodesAction(const TArray<URigVMNode*>& InNodes, URigVMController* InController);
+	FRigVMReplaceNodesAction();
+	FRigVMReplaceNodesAction(URigVMController* InController, const TArray<URigVMNode*>& InNodes);
 	virtual ~FRigVMReplaceNodesAction() {};
 	virtual UScriptStruct* GetScriptStruct() const override { return FRigVMReplaceNodesAction::StaticStruct(); }
-	virtual bool Undo(URigVMController* InController) override;
-	virtual bool Redo(URigVMController* InController) override;
+	virtual bool Undo() override;
+	virtual bool Redo() override;
 };
