RigVM: Disallow the controller to change graphs

diff --git a/Engine/Plugins/Animation/ControlRig/Source/ControlRigDeveloper/Private/ControlRigBlueprint.cpp b/Engine/Plugins/Animation/ControlRig/Source/ControlRigDeveloper/Private/ControlRigBlueprint.cpp
--- a/Engine/Plugins/Animation/ControlRig/Source/ControlRigDeveloper/Private/ControlRigBlueprint.cpp
+++ b/Engine/Plugins/Animation/ControlRig/Source/ControlRigDeveloper/Private/ControlRigBlueprint.cpp
@@ -3801,16 +3801,6 @@ void UControlRigBlueprint::HandleModifiedEvent(ERigVMGraphNotifType InNotifType,
 			{
 				if (URigVMCollapseNode* CollapseNode = Cast<URigVMCollapseNode>(InSubject))
 				{
-					// Fix the path to the contained graph controller
-					{
-						FSoftObjectPath ContainedGraphPath = CollapseNode->GetContainedGraph();
-						FString SubPathString = ContainedGraphPath.GetSubPathString();
-						FString OldSubPathString = SubPathString.Replace(*CollapseNode->GetName(), *CollapseNode->GetPreviousFName().ToString());
-						FSoftObjectPath OldContainedGraphPath = ContainedGraphPath;
-						OldContainedGraphPath.SetSubPathString(OldSubPathString);
-						RigVMClient.OnSubGraphRenamed(OldContainedGraphPath, ContainedGraphPath);
-					}
-					
 					FString NewNodePath = CollapseNode->GetNodePath(true /* recursive */);
 					FString Left, Right = NewNodePath;
 					URigVMNode::SplitNodePathAtEnd(NewNodePath, Left, Right);
diff --git a/Engine/Source/Developer/RigVMDeveloper/Private/RigVMCompiler/RigVMCompiler.cpp b/Engine/Source/Developer/RigVMDeveloper/Private/RigVMCompiler/RigVMCompiler.cpp
--- a/Engine/Source/Developer/RigVMDeveloper/Private/RigVMCompiler/RigVMCompiler.cpp
+++ b/Engine/Source/Developer/RigVMDeveloper/Private/RigVMCompiler/RigVMCompiler.cpp
@@ -452,17 +452,21 @@ bool URigVMCompiler::Compile(TArray<URigVMGraph*> InGraphs, URigVMController* In
 	for(int32 GraphIndex=0; GraphIndex<VisitedGraphs.Num(); GraphIndex++)
 	{
 		URigVMGraph* VisitedGraph = VisitedGraphs[GraphIndex];
-		
+
+		if(URigVMController* VisitedGraphController = InController->GetControllerForGraph(VisitedGraph))
 		{
-			FRigVMControllerGraphGuard Guard(InController, VisitedGraph, false);
 			// make sure variables are up to date before validating other things.
 			// that is, make sure their cpp type and type object agree with each other
-			InController->EnsureLocalVariableValidity();
+			VisitedGraphController->EnsureLocalVariableValidity();
 		}
 		
 		for(URigVMNode* ModelNode : VisitedGraph->GetNodes())
 		{
-			FRigVMControllerGraphGuard Guard(InController, VisitedGraph, false);
+			URigVMController* VisitedGraphController = InController->GetControllerForGraph(VisitedGraph);
+			if(VisitedGraphController == nullptr)
+			{
+				return false;
+			}
 
 			// make sure pins are up to date before validating other things.
 			// that is, make sure their cpp type and type object agree with each other
@@ -502,7 +506,7 @@ bool URigVMCompiler::Compile(TArray<URigVMGraph*> InGraphs, URigVMController* In
 				bEncounteredGraphError = true;
 			}
 
-			if(!InController->RemoveUnusedOrphanedPins(ModelNode))
+			if(!VisitedGraphController->RemoveUnusedOrphanedPins(ModelNode))
 			{
 				static const FString LinkedMessage = TEXT("Node @@ uses pins that no longer exist. Please rewire the links and re-compile.");
 				Settings.ASTSettings.Report(EMessageSeverity::Error, ModelNode, LinkedMessage);
@@ -652,7 +656,7 @@ bool URigVMCompiler::Compile(TArray<URigVMGraph*> InGraphs, URigVMController* In
 					UScriptStruct* ScriptStruct = UnitNode->GetScriptStruct(); 
 					if(ScriptStruct == nullptr)
 					{
-						InController->FullyResolveTemplateNode(UnitNode, INDEX_NONE, false);
+						VisitedGraphController->FullyResolveTemplateNode(UnitNode, INDEX_NONE, false);
 					}
 
 					if (UnitNode->GetScriptStruct() == nullptr || UnitNode->ResolvedFunctionName.IsEmpty())
@@ -1126,7 +1130,12 @@ bool URigVMCompiler::Compile(TArray<URigVMGraph*> InGraphs, URigVMController* In
 bool URigVMCompiler::CompileFunction(const URigVMLibraryNode* InLibraryNode, URigVMController* InController, FRigVMFunctionCompilationData* OutFunctionCompilationData)
 {
 	TGuardValue<const URigVMLibraryNode*> CompilationGuard(CurrentCompilationFunction, InLibraryNode);
-	FRigVMControllerGraphGuard ControllerGraphGuard(InController, InLibraryNode->GetContainedGraph(), false);
+
+	URigVMController* LibraryController = InController->GetControllerForGraph(InLibraryNode->GetContainedGraph());
+	if(LibraryController == nullptr)
+	{
+		return false;
+	}
 
 	double CompilationTime = 0;
 	FDurationTimer CompileTimer(CompilationTime);
@@ -1140,14 +1149,14 @@ bool URigVMCompiler::CompileFunction(const URigVMLibraryNode* InLibraryNode, URi
 	OutFunctionCompilationData->ByteCode.Reset();
 
 	TArray<FRigVMExternalVariable> ExternalVariables;
-	if (InController->GetExternalVariablesDelegate.IsBound())
+	if (LibraryController->GetExternalVariablesDelegate.IsBound())
 	{
-		ExternalVariables = InController->GetExternalVariablesDelegate.Execute(InLibraryNode->GetContainedGraph());
+		ExternalVariables = LibraryController->GetExternalVariablesDelegate.Execute(InLibraryNode->GetContainedGraph());
 	}
 	TMap<FString, FRigVMOperand> Operands;
 	
 	URigVM* TempVM = NewObject<URigVM>(InLibraryNode->GetContainedGraph());
-	const bool bSuccess = Compile({InLibraryNode->GetContainedGraph()}, InController, TempVM, ExternalVariables, &Operands, nullptr, OutFunctionCompilationData);
+	const bool bSuccess = Compile({InLibraryNode->GetContainedGraph()}, LibraryController, TempVM, ExternalVariables, &Operands, nullptr, OutFunctionCompilationData);
 	TempVM->ClearMemory();
 	TempVM->Rename(nullptr, GetTransientPackage(), REN_ForceNoResetLoaders | REN_DoNotDirty | REN_DontCreateRedirectors | REN_NonTransactional);
 	TempVM->MarkAsGarbage();
diff --git a/Engine/Source/Developer/RigVMDeveloper/Private/RigVMModel/RigVMClient.cpp b/Engine/Source/Developer/RigVMDeveloper/Private/RigVMModel/RigVMClient.cpp
--- a/Engine/Source/Developer/RigVMDeveloper/Private/RigVMModel/RigVMClient.cpp
+++ b/Engine/Source/Developer/RigVMDeveloper/Private/RigVMModel/RigVMClient.cpp
@@ -176,6 +176,7 @@ URigVMController* FRigVMClient::GetController(const URigVMGraph* InModel) const
 		const FSoftObjectPath Key(InModel);
 		if(const TObjectPtr<URigVMController>* Controller = Controllers.Find(Key))
 		{
+			checkf(Controller->Get()->GetGraph() == InModel, TEXT("Controller %s contains unexpected graph."), *Controller->GetPathName());
 			return Controller->Get();
 		}
 	}
@@ -236,6 +237,20 @@ URigVMController* FRigVMClient::GetOrCreateController(const UObject* InEditorSid
 	return nullptr;
 }
 
+bool FRigVMClient::RemoveController(const URigVMGraph* InModel)
+{
+	const FSoftObjectPath Key(InModel);
+	URigVMController* Controller = GetController(InModel);
+	const bool bSuccess = Controllers.Remove(Key) > 0;
+	if(Controller)
+	{
+		Controller->Rename(nullptr, GetTransientPackage(), REN_ForceNoResetLoaders | REN_DoNotDirty | REN_DontCreateRedirectors | REN_NonTransactional);
+		Controller->RemoveFromRoot();
+		Controller->MarkAsGarbage();
+	}
+	return bSuccess;
+}
+
 URigVMGraph* FRigVMClient::AddModel(const FName& InName, bool bSetupUndoRedo, const FObjectInitializer* ObjectInitializer, bool bCreateController)
 {
 #if WITH_EDITOR
@@ -408,10 +423,7 @@ bool FRigVMClient::RemoveModel(const FString& InNodePathOrName, bool bSetupUndoR
 #endif
 
 		// remove the controller
-		if(URigVMController* Controller = GetController(Model))
-		{
-			verify(Controller == Controllers.FindAndRemoveChecked(Model));
-		}
+		verify(RemoveController(Model));
 
 		if (GetOuter()->Implements<URigVMClientHost>())
 		{
@@ -566,14 +578,45 @@ void FRigVMClient::PostTransacted(const FTransactionObjectEvent& TransactionEven
 	}
 }
 
-void FRigVMClient::OnSubGraphRenamed(const FSoftObjectPath& OldObjectPath, const FSoftObjectPath& NewObjectPath)
+void FRigVMClient::OnCollapseNodeRenamed(const URigVMCollapseNode* InCollapseNode)
 {
-	if(OldObjectPath != NewObjectPath)
+	TMap<FSoftObjectPath, TObjectPtr<URigVMController>> OldControllers = Controllers;
+	for(const TPair<FSoftObjectPath, TObjectPtr<URigVMController>>& Pair : OldControllers)
 	{
-		if (TObjectPtr<URigVMController>* Controller = Controllers.Find(OldObjectPath))
+		if(URigVMController* Controller = Pair.Value)
 		{
-			Controllers.Add(NewObjectPath, *Controller);
-			Controllers.Remove(OldObjectPath);
+			if(const URigVMGraph* Graph = Controller->GetGraph())
+			{
+				if(Graph->IsInOuter(InCollapseNode))
+				{
+					const FSoftObjectPath& OldObjectPath = Pair.Key;
+					const FSoftObjectPath NewObjectPath(Graph);
+
+					if(OldObjectPath != NewObjectPath)
+					{
+						Controllers.Add(NewObjectPath, Controller);
+						Controllers.Remove(OldObjectPath);
+					}
+				}
+			}
+		}
+	}
+}
+
+void FRigVMClient::OnCollapseNodeRemoved(const URigVMCollapseNode* InCollapseNode)
+{
+	TMap<FSoftObjectPath, TObjectPtr<URigVMController>> OldControllers = Controllers;
+	for(const TPair<FSoftObjectPath, TObjectPtr<URigVMController>>& Pair : OldControllers)
+	{
+		if(URigVMController* Controller = Pair.Value)
+		{
+			if(const URigVMGraph* Graph = Controller->GetGraph())
+			{
+				if(Graph->IsInOuter(InCollapseNode))
+				{
+					RemoveController(Graph);
+				}
+			}
 		}
 	}
 }
diff --git a/Engine/Source/Developer/RigVMDeveloper/Private/RigVMModel/RigVMController.cpp b/Engine/Source/Developer/RigVMDeveloper/Private/RigVMModel/RigVMController.cpp
--- a/Engine/Source/Developer/RigVMDeveloper/Private/RigVMModel/RigVMController.cpp
+++ b/Engine/Source/Developer/RigVMDeveloper/Private/RigVMModel/RigVMController.cpp
@@ -151,6 +151,9 @@ void URigVMController::SetGraph(URigVMGraph* InGraph)
 			return;
 		}
 		LastGraph->OnModified().RemoveAll(this);
+
+		static constexpr TCHAR Message[] = TEXT("Usage of URigVMController::SetGraph to switch between graphs has been deprecated. Please rely on URigVMController::GetControllerForGraph instead.");
+		ReportWarning(Message);
 	}
 
 	Graphs.Reset();
@@ -171,6 +174,9 @@ bool URigVMController::PushGraph(URigVMGraph* InGraph, bool bSetupUndoRedo)
 			return false;
 		}
 		LastGraph->OnModified().RemoveAll(this);
+		
+		static constexpr TCHAR Message[] = TEXT("Usage of URigVMController::PushGraph to switch between graphs has been deprecated. Please rely on URigVMController::GetControllerForGraph instead.");
+		ReportWarning(Message);
 	}
 
 	check(InGraph);
@@ -178,11 +184,6 @@ bool URigVMController::PushGraph(URigVMGraph* InGraph, bool bSetupUndoRedo)
 
 	InGraph->OnModified().AddUObject(this, &URigVMController::HandleModifiedEvent);
 
-	if (bSetupUndoRedo)
-	{
-		GetActionStack()->AddAction(FRigVMPushGraphAction(this, InGraph));
-	}
-
 	return true;
 }
 
@@ -207,11 +208,6 @@ URigVMGraph* URigVMController::PopGraph(bool bSetupUndoRedo)
 		CurrentGraph->OnModified().AddUObject(this, &URigVMController::HandleModifiedEvent);
 	}
 
-	if (bSetupUndoRedo)
-	{
-		GetActionStack()->AddAction(FRigVMPopGraphAction(this, LastGraph));
-	}
-
 	return LastGraph;
 }
 
@@ -239,6 +235,36 @@ URigVMGraph* URigVMController::GetTopLevelGraph() const
 	return Graph;
 }
 
+URigVMController* URigVMController::GetControllerForGraph(const URigVMGraph* InGraph) const
+{
+	if(InGraph == GetGraph())
+	{
+		return const_cast<URigVMController*>(this);
+	}
+	
+	if(IRigVMClientHost* ClientHost = GetImplementingOuter<IRigVMClientHost>())
+	{
+		if(FRigVMClient* Client = ClientHost->GetRigVMClient())
+		{
+			URigVMController* Controller = Client->GetOrCreateController(InGraph);
+			if(bIsRunningUnitTest)
+			{
+				Controller->bIsRunningUnitTest = true;
+				if(!Controller->GetExternalVariablesDelegate.IsBound() && !InGraph->IsRootGraph())
+				{
+					// get the controller for the root graph
+					if(const URigVMController* RootController =  Client->GetOrCreateController(InGraph->GetRootGraph()))
+					{
+						Controller->GetExternalVariablesDelegate = RootController->GetExternalVariablesDelegate;
+					}
+				}
+			}
+			return Controller;
+		}
+	}
+	return nullptr;
+}
+
 FRigVMGraphModifiedEvent& URigVMController::OnModified()
 {
 	return ModifiedEventStatic;
@@ -776,12 +802,14 @@ void URigVMController::HandleModifiedEvent(ERigVMGraphNotifType InNotifType, URi
 				URigVMNode* Node = CastChecked<URigVMNode>(InSubject);
 				check(Node);
 
-				if(URigVMLibraryNode* Function = FunctionLibrary->FindFunctionForNode(Node))
+				if(const URigVMLibraryNode* Function = FunctionLibrary->FindFunctionForNode(Node))
 				{
 					FunctionLibrary->ForEachReference(Function->GetFName(), [this](URigVMFunctionReferenceNode* Reference)
                     {
-                        FRigVMControllerGraphGuard GraphGuard(this, Reference->GetGraph(), false);
-                        Reference->GetGraph()->Notify(ERigVMGraphNotifType::VariableRemappingChanged, Reference);
+						if(const URigVMController* ReferenceController = GetControllerForGraph(Reference->GetGraph()))
+						{
+							ReferenceController->Notify(ERigVMGraphNotifType::VariableRemappingChanged, Reference);
+						}
                     });
 				}
 			}
@@ -1864,26 +1892,32 @@ void URigVMController::OnExternalVariableRemoved(const FName& InVarName, bool bS
 				}
 			}
 		}
-		else if(URigVMCollapseNode* CollapseNode = Cast<URigVMCollapseNode>(Node))
+		else if(const URigVMCollapseNode* CollapseNode = Cast<URigVMCollapseNode>(Node))
 		{
-			FRigVMControllerGraphGuard GraphGuard(this, CollapseNode->GetContainedGraph(), bSetupUndoRedo);
-			TGuardValue<bool> GuardEditGraph(CollapseNode->ContainedGraph->bEditable, true);
-
-			// call this function for the contained graph recursively 
-			OnExternalVariableRemoved(InVarName, bSetupUndoRedo);
+			if(URigVMController* CollapseController = GetControllerForGraph(CollapseNode->GetContainedGraph()))
+			{
+				FRigVMControllerCompileBracketScope CollapseCompileScope(CollapseController);
+				TGuardValue<bool> GuardEditGraph(CollapseNode->ContainedGraph->bEditable, true);
 
+				// call this function for the contained graph recursively 
+				CollapseController->OnExternalVariableRemoved(InVarName, bSetupUndoRedo);
+			}
+			
 			// if we are a function we need to notify all references!
-			if(URigVMFunctionLibrary* FunctionLibrary = Cast<URigVMFunctionLibrary>(Graph))
+			if(const URigVMFunctionLibrary* FunctionLibrary = Cast<URigVMFunctionLibrary>(Graph))
 			{
 				FunctionLibrary->ForEachReference(CollapseNode->GetFName(), [this, InVarName](URigVMFunctionReferenceNode* Reference)
 				{
 					if(Reference->VariableMap.Contains(InVarName))
 					{
 						Reference->Modify();
-                        Reference->VariableMap.Remove(InVarName);
+						Reference->VariableMap.Remove(InVarName);
 
-                        FRigVMControllerGraphGuard GraphGuard(this, Reference->GetGraph(), false);
-                        Notify(ERigVMGraphNotifType::VariableRemappingChanged, Reference);
+						if(URigVMController* ReferenceController = GetControllerForGraph(Reference->GetGraph()))
+						{
+							FRigVMControllerCompileBracketScope ReferenceCompileScope(ReferenceController);
+							ReferenceController->Notify(ERigVMGraphNotifType::VariableRemappingChanged, Reference);
+						}
 					}
 				});
 			}
@@ -1957,14 +1991,17 @@ bool URigVMController::OnExternalVariableRenamed(const FName& InOldVarName, cons
 				}
 			}
 		}
-		else if(URigVMCollapseNode* CollapseNode = Cast<URigVMCollapseNode>(Node))
+		else if(const URigVMCollapseNode* CollapseNode = Cast<URigVMCollapseNode>(Node))
 		{
-			FRigVMControllerGraphGuard GraphGuard(this, CollapseNode->GetContainedGraph(), bSetupUndoRedo);
-			TGuardValue<bool> GuardEditGraph(CollapseNode->ContainedGraph->bEditable, true);
-			OnExternalVariableRenamed(InOldVarName, InNewVarName, bSetupUndoRedo);
+			if(URigVMController* CollapseController = GetControllerForGraph(CollapseNode->GetContainedGraph()))
+			{
+				FRigVMControllerCompileBracketScope CollapseCompileScope(CollapseController);
+				TGuardValue<bool> GuardEditGraph(CollapseNode->ContainedGraph->bEditable, true);
+				CollapseController->OnExternalVariableRenamed(InOldVarName, InNewVarName, bSetupUndoRedo);
+			}
 
 			// if we are a function we need to notify all references!
-			if(URigVMFunctionLibrary* FunctionLibrary = Cast<URigVMFunctionLibrary>(Graph))
+			if(const URigVMFunctionLibrary* FunctionLibrary = Cast<URigVMFunctionLibrary>(Graph))
 			{
 				FunctionLibrary->ForEachReference(CollapseNode->GetFName(), [this, InOldVarName, InNewVarName](URigVMFunctionReferenceNode* Reference)
                 {
@@ -1972,12 +2009,15 @@ bool URigVMController::OnExternalVariableRenamed(const FName& InOldVarName, cons
 					{
 						Reference->Modify();
 
-						FName MappedVariable = Reference->VariableMap.FindChecked(InOldVarName);
+						const FName& MappedVariable = Reference->VariableMap.FindChecked(InOldVarName);
 						Reference->VariableMap.Remove(InOldVarName);
 						Reference->VariableMap.FindOrAdd(InNewVarName) = MappedVariable; 
 
-						FRigVMControllerGraphGuard GraphGuard(this, Reference->GetGraph(), false);
-                        Notify(ERigVMGraphNotifType::VariableRemappingChanged, Reference);
+						if(URigVMController* ReferenceController = GetControllerForGraph(Reference->GetGraph()))
+						{
+							FRigVMControllerCompileBracketScope ReferenceCompileScope(ReferenceController);
+							ReferenceController->Notify(ERigVMGraphNotifType::VariableRemappingChanged, Reference);
+						}
                     }
                 });
 			}
@@ -2042,9 +2082,9 @@ void URigVMController::OnExternalVariableTypeChanged(const FName& InVarName, con
 	TArray<URigVMNode*> Nodes = Graph->GetNodes();
 	for (URigVMNode* Node : Nodes)
 	{
-		if (URigVMVariableNode* VariableNode = Cast<URigVMVariableNode>(Node))
+		if (const URigVMVariableNode* VariableNode = Cast<URigVMVariableNode>(Node))
 		{
-			if (URigVMPin* VariablePin = VariableNode->FindPin(URigVMVariableNode::VariableName))
+			if (const URigVMPin* VariablePin = VariableNode->FindPin(URigVMVariableNode::VariableName))
 			{
 				if (VariablePin->GetDefaultValue() == VarNameStr)
 				{
@@ -2053,14 +2093,17 @@ void URigVMController::OnExternalVariableTypeChanged(const FName& InVarName, con
 				}
 			}
 		}
-		else if(URigVMCollapseNode* CollapseNode = Cast<URigVMCollapseNode>(Node))
+		else if(const URigVMCollapseNode* CollapseNode = Cast<URigVMCollapseNode>(Node))
 		{
-			FRigVMControllerGraphGuard GraphGuard(this, CollapseNode->GetContainedGraph(), bSetupUndoRedo);
-			TGuardValue<bool> GuardEditGraph(CollapseNode->ContainedGraph->bEditable, true);
-			OnExternalVariableTypeChanged(InVarName, InCPPType, InCPPTypeObject, bSetupUndoRedo);
+			if(URigVMController* CollapseController = GetControllerForGraph(CollapseNode->GetContainedGraph()))
+			{
+				FRigVMControllerCompileBracketScope CollapseCompileScope(CollapseController);
+				TGuardValue<bool> GuardEditGraph(CollapseNode->ContainedGraph->bEditable, true);
+				CollapseController->OnExternalVariableTypeChanged(InVarName, InCPPType, InCPPTypeObject, bSetupUndoRedo);
+			}
 
 			// if we are a function we need to notify all references!
-			if(URigVMFunctionLibrary* FunctionLibrary = Cast<URigVMFunctionLibrary>(Graph))
+			if(const URigVMFunctionLibrary* FunctionLibrary = Cast<URigVMFunctionLibrary>(Graph))
 			{
 				FunctionLibrary->ForEachReference(CollapseNode->GetFName(), [this, InVarName](URigVMFunctionReferenceNode* Reference)
                 {
@@ -2069,8 +2112,11 @@ void URigVMController::OnExternalVariableTypeChanged(const FName& InVarName, con
                         Reference->Modify();
                         Reference->VariableMap.Remove(InVarName); 
 
-                        FRigVMControllerGraphGuard GraphGuard(this, Reference->GetGraph(), false);
-                        Notify(ERigVMGraphNotifType::VariableRemappingChanged, Reference);
+                    	if(URigVMController* ReferenceController = GetControllerForGraph(Reference->GetGraph()))
+                    	{
+                    		FRigVMControllerCompileBracketScope ReferenceCompileScope(ReferenceController);
+							ReferenceController->Notify(ERigVMGraphNotifType::VariableRemappingChanged, Reference);
+						}
                     }
                 });
 			}
@@ -3242,7 +3288,7 @@ bool URigVMController::RemoveInjectedNode(const FString& InPinPath, bool bAsInpu
 	}
 	
 	// 3.- Remove node
-	if (!RemoveNode(NodeEjected))
+	if (!RemoveNode(NodeEjected, bSetupUndoRedo, false))
 	{
 		GetActionStack()->CancelAction(Action);
 		return false;
@@ -3663,7 +3709,20 @@ FString URigVMController::ExportNodesToText(const TArray<FName>& InNodeNames, bo
 	FStringOutputDevice Archive;
 	const FExportObjectInnerContext Context;
 
+	// find all of the nodes
+	TArray<URigVMNode*> Nodes;
+	Nodes.Reserve(InNodeNames.Num());
+	for (const FName& NodeName : InNodeNames)
+	{
+		if (URigVMNode* Node = Graph->FindNodeByName(NodeName))
+		{
+			Nodes.Add(Node);
+		}
+	}
+
 	TArray<FName> AllNodeNames = InNodeNames;
+	TArray<FName> FilteredNodeNames;
+	FilteredNodeNames.Reserve(InNodeNames.Num());
 	for (const FName& NodeName : InNodeNames)
 	{
 		if (URigVMNode* Node = Graph->FindNodeByName(NodeName))
@@ -3675,11 +3734,25 @@ FString URigVMController::ExportNodesToText(const TArray<FName>& InNodeNames, bo
 					AllNodeNames.AddUnique(Injection->Node->GetFName());
 				}
 			}
+
+			// skip injected nodes that would show up twice in the export
+			if(const URigVMInjectionInfo* InjectionInfo = Node->GetInjectionInfo())
+			{
+				if(const URigVMNode* OuterNode = InjectionInfo->GetTypedOuter<URigVMNode>())
+				{
+					if(Nodes.Contains(OuterNode))
+					{
+						continue;
+					}
+				}
+			}
 		}
+
+		FilteredNodeNames.Add(NodeName);
 	}
 
 	// Export each of the selected nodes
-	for (const FName& NodeName : InNodeNames)
+	for (const FName& NodeName : FilteredNodeNames)
 	{
 		if (URigVMNode* Node = Graph->FindNodeByName(NodeName))
 		{
@@ -3910,16 +3983,18 @@ TArray<FName> URigVMController::ImportNodesFromText(const FString& InText, bool
 				continue;
 			}
 
-			FRigVMControllerGraphGuard GraphGuard(this, Graph, false);
-			if(!CanAddNode(CreatedNode, true))
+			if(URigVMController* ControllerForGraph = GetControllerForGraph(Graph))
 			{
-				Graph->Nodes.Remove(CreatedNode);
-				DestroyObject(CreatedNode);
-				continue;
-			}
+				if(!ControllerForGraph->CanAddNode(CreatedNode, true))
+				{
+					Graph->Nodes.Remove(CreatedNode);
+					DestroyObject(CreatedNode);
+					continue;
+				}
 
-			Graph->Nodes.AddUnique(CreatedNode);
-			FilteredNodes.Add(CreatedNode);
+				Graph->Nodes.AddUnique(CreatedNode);
+				FilteredNodes.Add(CreatedNode);
+			}
 		}
 
 		Swap(Factory.CreatedNodes, FilteredNodes);
@@ -3942,65 +4017,67 @@ TArray<FName> URigVMController::ImportNodesFromText(const FString& InText, bool
 	{
 		URigVMGraph* Graph = CreatedNode->GetTypedOuter<URigVMGraph>();
 		check(Graph);
-		FRigVMControllerGraphGuard GraphGuard(this, Graph, false);
 
-		// Refresh the unit node to account for changes in node color, pin additions, pin order, etc
-		if (URigVMUnitNode* UnitNode = Cast<URigVMUnitNode>(CreatedNode))
+		if(URigVMController* ControllerForGraph = GetControllerForGraph(Graph))
 		{
-			if(UnitNode->ResolvedFunctionName.IsEmpty())
+			// Refresh the unit node to account for changes in node color, pin additions, pin order, etc
+			if (URigVMUnitNode* UnitNode = Cast<URigVMUnitNode>(CreatedNode))
 			{
-				if(UScriptStruct* ScriptStruct = UnitNode->GetScriptStruct())
+				if(UnitNode->ResolvedFunctionName.IsEmpty())
 				{
-					const FName MethodName = UnitNode->GetMethodName();
-					if(!MethodName.IsNone())
+					if(UScriptStruct* ScriptStruct = UnitNode->GetScriptStruct())
 					{
-						if(const FRigVMFunction* Function = FRigVMRegistry::Get().FindFunction(ScriptStruct, *MethodName.ToString()))
+						const FName MethodName = UnitNode->GetMethodName();
+						if(!MethodName.IsNone())
 						{
-							UnitNode->ResolvedFunctionName = Function->GetName();
-
-							if(const FRigVMTemplate* Template = Function->GetTemplate())
+							if(const FRigVMFunction* Function = FRigVMRegistry::Get().FindFunction(ScriptStruct, *MethodName.ToString()))
 							{
-								UnitNode->TemplateNotation = Template->GetNotation();
+								UnitNode->ResolvedFunctionName = Function->GetName();
+
+								if(const FRigVMTemplate* Template = Function->GetTemplate())
+								{
+									UnitNode->TemplateNotation = Template->GetNotation();
+								}
 							}
 						}
 					}
 				}
+				
+				TGuardValue<bool> SuspendNotifications(ControllerForGraph->bSuspendNotifications, true);
+				ControllerForGraph->RepopulatePinsOnNode(UnitNode);
 			}
-			
-			TGuardValue<bool> SuspendNotifications(bSuspendNotifications, true);
-			RepopulatePinsOnNode(UnitNode);
-		}
 
-		if (URigVMUnitNode* UnitNode = Cast<URigVMUnitNode>(CreatedNode))
-		{
-			if (UnitNodeCreatedContext.IsValid())
+			if (URigVMUnitNode* UnitNode = Cast<URigVMUnitNode>(CreatedNode))
 			{
-				if (TSharedPtr<FStructOnScope> StructScope = UnitNode->ConstructStructInstance())
+				if (ControllerForGraph->UnitNodeCreatedContext.IsValid())
 				{
-					TGuardValue<FName> NodeNameScope(UnitNodeCreatedContext.NodeName, UnitNode->GetFName());
-					FRigVMStruct* StructInstance = (FRigVMStruct*)StructScope->GetStructMemory();
-					StructInstance->OnUnitNodeCreated(UnitNodeCreatedContext);
+					if (TSharedPtr<FStructOnScope> StructScope = UnitNode->ConstructStructInstance())
+					{
+						TGuardValue<FName> NodeNameScope(ControllerForGraph->UnitNodeCreatedContext.NodeName, UnitNode->GetFName());
+						FRigVMStruct* StructInstance = (FRigVMStruct*)StructScope->GetStructMemory();
+						StructInstance->OnUnitNodeCreated(ControllerForGraph->UnitNodeCreatedContext);
+					}
 				}
 			}
-		}
 
-		if (URigVMFunctionReferenceNode* FunctionRefNode = Cast<URigVMFunctionReferenceNode>(CreatedNode))
-		{
-			if(URigVMBuildData* BuildData = URigVMBuildData::Get())
+			if (URigVMFunctionReferenceNode* FunctionRefNode = Cast<URigVMFunctionReferenceNode>(CreatedNode))
 			{
-				BuildData->RegisterFunctionReference(FunctionRefNode->GetReferencedFunctionHeader().LibraryPointer, FunctionRefNode);
+				if(URigVMBuildData* BuildData = URigVMBuildData::Get())
+				{
+					BuildData->RegisterFunctionReference(FunctionRefNode->GetReferencedFunctionHeader().LibraryPointer, FunctionRefNode);
+				}
 			}
-		}
 
-		for(URigVMPin* Pin : CreatedNode->Pins)
-		{
-			EnsurePinValidity(Pin, true);
-		}
+			for(URigVMPin* Pin : CreatedNode->Pins)
+			{
+				ControllerForGraph->EnsurePinValidity(Pin, true);
+			}
 
-		Notify(ERigVMGraphNotifType::NodeAdded, CreatedNode);
+			ControllerForGraph->Notify(ERigVMGraphNotifType::NodeAdded, CreatedNode);
 
-		TArray<FName>& NodeNames = NodeNamesPerGraph.FindOrAdd(Graph);
-		NodeNames.Add(CreatedNode->GetFName());
+			TArray<FName>& NodeNames = NodeNamesPerGraph.FindOrAdd(Graph);
+			NodeNames.Add(CreatedNode->GetFName());
+		}
 	}
 
 	if (Factory.CreatedLinks.Num() > 0)
@@ -4021,54 +4098,55 @@ TArray<FName> URigVMController::ImportNodesFromText(const FString& InText, bool
 				LastGraph = Graph;
 			}
 
-			FRigVMControllerGraphGuard GraphGuard(this, Graph, false);
-
-			// this takes care of mapping from the old to the new nodes etc
-			const TArray<FLinkedPath> LinkedPaths = RemapLinkedPaths(
-				GetLinkedPaths({CreatedLink}), 
-				Factory.RestoreLinksSettings,
-				false);
-			const FLinkedPath& LinkedPath = LinkedPaths[0]; 
-
-			URigVMPin* SourcePin = LinkedPath.GetSourcePin();
-			URigVMPin* TargetPin = LinkedPath.GetTargetPin();
-			
-			if (SourcePin && TargetPin)
+			if(URigVMController* ControllerForGraph = GetControllerForGraph(Graph))
 			{
-				// update the paths on the link to allow reuse of the UObject
-				CreatedLink->SetSourceAndTargetPinPaths(LinkedPath.SourcePinPath, LinkedPath.TargetPinPath);
+				// this takes care of mapping from the old to the new nodes etc
+				const TArray<FLinkedPath> LinkedPaths = ControllerForGraph->RemapLinkedPaths(
+					ControllerForGraph->GetLinkedPaths({CreatedLink}), 
+					Factory.RestoreLinksSettings,
+					false);
+				const FLinkedPath& LinkedPath = LinkedPaths[0]; 
+
+				URigVMPin* SourcePin = LinkedPath.GetSourcePin();
+				URigVMPin* TargetPin = LinkedPath.GetTargetPin();
 				
-				if (TargetPin->IsBoundToVariable())
+				if (SourcePin && TargetPin)
 				{
-					const FString VariableNodeName = TargetPin->GetBoundVariableNode()->GetName();
-					const FString BindingPath = TargetPin->GetBoundVariablePath();
-
-					// The current situation is that the outer pin has an injection info, and the injected node exists
-					// but the injected node is not linked to the outer pin. BreakAllLinks will try to unbind the outer pin,
-					// for that to be successful, the binding needs to be complete
-					// Connect it so that the unbound is successful
-					check(!SourcePin->IsLinkedTo(TargetPin));
-					Graph->DetachedLinks.Add(CreatedLink);
-					AddLink(SourcePin, TargetPin, false);
-
-					// recreate binding
-					UnbindPinFromVariable(TargetPin, false);
-					BindPinToVariable(TargetPin, BindingPath, false, VariableNodeName);
+					// update the paths on the link to allow reuse of the UObject
+					CreatedLink->SetSourceAndTargetPinPaths(LinkedPath.SourcePinPath, LinkedPath.TargetPinPath);
+					
+					if (TargetPin->IsBoundToVariable())
+					{
+						const FString VariableNodeName = TargetPin->GetBoundVariableNode()->GetName();
+						const FString BindingPath = TargetPin->GetBoundVariablePath();
+
+						// The current situation is that the outer pin has an injection info, and the injected node exists
+						// but the injected node is not linked to the outer pin. BreakAllLinks will try to unbind the outer pin,
+						// for that to be successful, the binding needs to be complete
+						// Connect it so that the unbound is successful
+						check(!SourcePin->IsLinkedTo(TargetPin));
+						Graph->DetachedLinks.Add(CreatedLink);
+						ControllerForGraph->AddLink(SourcePin, TargetPin, false);
+
+						// recreate binding
+						ControllerForGraph->UnbindPinFromVariable(TargetPin, false);
+						ControllerForGraph->BindPinToVariable(TargetPin, BindingPath, false, VariableNodeName);
+					}
+					else
+					{
+						Graph->DetachedLinks.Add(CreatedLink);
+						ControllerForGraph->AddLink(SourcePin, TargetPin, false);
+					}
 				}
-				else
+
+				// if the link is still part of the detached link array
+				// the restore was not successful.
+				if(Graph->DetachedLinks.Remove(CreatedLink))
 				{
-					Graph->DetachedLinks.Add(CreatedLink);
-					AddLink(SourcePin, TargetPin, false);
+					ControllerForGraph->ReportErrorf(TEXT("Cannot import link '%s'."), *URigVMLink::GetPinPathRepresentation(CreatedLink->GetSourcePinPath(), CreatedLink->GetTargetPinPath()));
+					DestroyObject(CreatedLink);
 				}
 			}
-
-			// if the link is still part of the detached link array
-			// the restore was not successful.
-			if(Graph->DetachedLinks.Remove(CreatedLink))
-			{
-				ReportErrorf(TEXT("Cannot import link '%s'."), *URigVMLink::GetPinPathRepresentation(CreatedLink->GetSourcePinPath(), CreatedLink->GetTargetPinPath()));
-				DestroyObject(CreatedLink);
-			}
 		}
 	}
 
@@ -4285,8 +4363,8 @@ TMap<FRigVMGraphFunctionIdentifier, URigVMLibraryNode*> URigVMController::Locali
 	}
 
 	// import the functions to our local function library
+	if(URigVMController* LibraryController = GetControllerForGraph(ThisLibrary))
 	{
-		FRigVMControllerGraphGuard GraphGuard(this, ThisLibrary, bSetupUndoRedo);
 		for(FRigVMGraphFunctionData* FunctionToLocalize : FunctionsToLocalize)
 		{
 			if (URigVMLibraryNode* ReferencedFunction = Cast<URigVMLibraryNode>(FunctionToLocalize->Header.LibraryPointer.LibraryNode.TryLoad()))
@@ -4297,7 +4375,7 @@ TMap<FRigVMGraphFunctionIdentifier, URigVMLibraryNode*> URigVMController::Locali
 				}
 			}
 			
-			TArray<FName> NodeNames = ImportNodesFromText(FunctionToLocalize->SerializedCollapsedNode, false);
+			TArray<FName> NodeNames = LibraryController->ImportNodesFromText(FunctionToLocalize->SerializedCollapsedNode, false);
 			if (NodeNames.Num() > 0)
 			{
 				URigVMLibraryNode* LocalizedFunction = ThisLibrary->FindFunction(NodeNames[0]);
@@ -4749,18 +4827,19 @@ URigVMCollapseNode* URigVMController::CollapseNodes(const TArray<URigVMNode*>& I
 
 	URigVMFunctionEntryNode* EntryNode = nullptr;
 	URigVMFunctionReturnNode* ReturnNode = nullptr;
-	{
-		FRigVMControllerGraphGuard GraphGuard(this, CollapseNode->GetContainedGraph(), false);
 
+	// import the functions to our local function library
+	if(URigVMController* CollapseController = GetControllerForGraph(CollapseNode->GetContainedGraph()))
+	{
 		EntryNode = NewObject<URigVMFunctionEntryNode>(CollapseNode->ContainedGraph, TEXT("Entry"));
 		CollapseNode->ContainedGraph->Nodes.Add(EntryNode);
 		EntryNode->Position = -Diagonal * 0.5f - FVector2D(250.f, 0.f);
 		{
-			TGuardValue<bool> SuspendNotifications(bSuspendNotifications, true);
-			RefreshFunctionPins(EntryNode);
+			TGuardValue<bool> SuspendNotifications(CollapseController->bSuspendNotifications, true);
+			CollapseController->RefreshFunctionPins(EntryNode);
 		}
 
-		Notify(ERigVMGraphNotifType::NodeAdded, EntryNode);
+		CollapseController->Notify(ERigVMGraphNotifType::NodeAdded, EntryNode);
 
 		if (bContainsOutputs)
 		{
@@ -4768,21 +4847,21 @@ URigVMCollapseNode* URigVMController::CollapseNodes(const TArray<URigVMNode*>& I
 			CollapseNode->ContainedGraph->Nodes.Add(ReturnNode);
 			ReturnNode->Position = FVector2D(Diagonal.X, -Diagonal.Y) * 0.5f + FVector2D(300.f, 0.f);
 			{
-				TGuardValue<bool> SuspendNotifications(bSuspendNotifications, true);
-				RefreshFunctionPins(ReturnNode);
+				TGuardValue<bool> SuspendNotifications(CollapseController->bSuspendNotifications, true);
+				CollapseController->RefreshFunctionPins(ReturnNode);
 			}
 
-			Notify(ERigVMGraphNotifType::NodeAdded, ReturnNode);
+			CollapseController->Notify(ERigVMGraphNotifType::NodeAdded, ReturnNode);
 		}
 	}
 
 	// create the new nodes within the collapse node
 	TArray<FName> ContainedNodeNames;
+	if(URigVMController* CollapseController = GetControllerForGraph(CollapseNode->GetContainedGraph()))
 	{
 		FString TextContent = ExportNodesToText(NodeNames);
 
-		FRigVMControllerGraphGuard GraphGuard(this, CollapseNode->GetContainedGraph(), false);
-		ContainedNodeNames = ImportNodesFromText(TextContent, false);
+		ContainedNodeNames = CollapseController->ImportNodesFromText(TextContent, false);
 
 		// move the nodes to the right place
 		for (const FName& ContainedNodeName : ContainedNodeNames)
@@ -4791,7 +4870,7 @@ URigVMCollapseNode* URigVMController::CollapseNodes(const TArray<URigVMNode*>& I
 			{
 				if(!ContainedNode->IsInjected())
 				{
-					SetNodePosition(ContainedNode, ContainedNode->Position - Center, false, false);
+					CollapseController->SetNodePosition(ContainedNode, ContainedNode->Position - Center, false, false);
 				}
 			}
 		}
@@ -4825,7 +4904,7 @@ URigVMCollapseNode* URigVMController::CollapseNodes(const TArray<URigVMNode*>& I
 			{
 				if (!SourcePin->IsLinkedTo(TargetPin))
 				{
-					AddLink(SourcePin, TargetPin, false);
+					CollapseController->AddLink(SourcePin, TargetPin, false);
 				}
 			}
 		}
@@ -4992,8 +5071,10 @@ URigVMCollapseNode* URigVMController::CollapseNodes(const TArray<URigVMNode*>& I
 					continue;
 				}
 
-				FRigVMControllerGraphGuard GraphGuard(this, CollapseNode->GetContainedGraph(), false);
-				AddLink(ExecuteContextPin, ReturnNode->FindPin(FRigVMStruct::ExecuteContextName.ToString()), false);
+				if(URigVMController* CollapseController = GetControllerForGraph(CollapseNode->GetContainedGraph()))
+				{
+					CollapseController->AddLink(ExecuteContextPin, ReturnNode->FindPin(FRigVMStruct::ExecuteContextName.ToString()), false);
+				}
 				break;
 			}
 		}
@@ -5181,9 +5262,9 @@ TArray<URigVMNode*> URigVMController::ExpandLibraryNode(URigVMLibraryNode* InNod
 	FVector2D Center = (Bounds.Min + Bounds.Max) * 0.5f;
 
 	FString TextContent;
+	if(URigVMController* InnerController = GetControllerForGraph(InnerGraph))
 	{
-		FRigVMControllerGraphGuard GraphGuard(this, InnerGraph, false);
-		TextContent = ExportNodesToText(NodeNames);
+		TextContent = InnerController->ExportNodesToText(NodeNames);
 	}
 
 	TArray<FName> ExpandedNodeNames = ImportNodesFromText(TextContent, false);
@@ -5836,10 +5917,12 @@ URigVMFunctionReferenceNode* URigVMController::PromoteCollapseNodeToFunctionRefe
 
 	if (FunctionDefinition == nullptr)
 	{
+		if(URigVMController* LibraryController = GetControllerForGraph(FunctionLibrary))
 		{
-			FRigVMControllerGraphGuard GraphGuard(this, FunctionLibrary, false);
-			const FString FunctionName = GetValidNodeName(InCollapseNode->GetName());
-			FunctionDefinition = AddFunctionToLibrary(
+			FRigVMControllerCompileBracketScope LibraryCompileScope(LibraryController);
+
+			const FString FunctionName = LibraryController->GetValidNodeName(InCollapseNode->GetName());
+			FunctionDefinition = LibraryController->AddFunctionToLibrary(
 				*FunctionName,
 				InCollapseNode->GetPins().ContainsByPredicate([](const URigVMPin* Pin) -> bool
 				{
@@ -5851,10 +5934,12 @@ URigVMFunctionReferenceNode* URigVMController::PromoteCollapseNodeToFunctionRefe
 		// Add interface pins in function
 		if (FunctionDefinition)
 		{
-			FRigVMControllerGraphGuard GraphGuard(this, FunctionDefinition->GetContainedGraph(), false);
-			for(const URigVMPin* Pin : InCollapseNode->GetPins())
+			if(URigVMController* DefinitionController = GetControllerForGraph(FunctionDefinition->GetContainedGraph()))
 			{
-				AddExposedPin(Pin->GetFName(), Pin->GetDirection(), Pin->GetCPPType(), (Pin->GetCPPTypeObject() ? *Pin->GetCPPTypeObject()->GetPathName() : TEXT("")), Pin->GetDefaultValue(), false);
+				for(const URigVMPin* Pin : InCollapseNode->GetPins())
+				{
+					DefinitionController->AddExposedPin(Pin->GetFName(), Pin->GetDirection(), Pin->GetCPPType(), (Pin->GetCPPTypeObject() ? *Pin->GetCPPTypeObject()->GetPathName() : TEXT("")), Pin->GetDefaultValue(), false);
+				}
 			}
 		}
 	}
@@ -5863,8 +5948,8 @@ URigVMFunctionReferenceNode* URigVMController::PromoteCollapseNodeToFunctionRefe
 	if (FunctionDefinition)
 	{
 		FString TextContent;
+		if(URigVMController* CollapseController = GetControllerForGraph(InCollapseNode->GetContainedGraph()))
 		{
-			FRigVMControllerGraphGuard GraphGuard(this, InCollapseNode->GetContainedGraph(), false);
 			TArray<FName> NodeNames;
 			for (const URigVMNode* Node : InCollapseNode->GetContainedNodes())
 			{
@@ -5875,19 +5960,20 @@ URigVMFunctionReferenceNode* URigVMController::PromoteCollapseNodeToFunctionRefe
 				
 				NodeNames.Add(Node->GetFName());
 			}
-			TextContent = ExportNodesToText(NodeNames);
+			TextContent = CollapseController->ExportNodesToText(NodeNames);
 		}
+
+		if(URigVMController* DefinitionController = GetControllerForGraph(FunctionDefinition->GetContainedGraph()))
 		{
-			FRigVMControllerGraphGuard GraphGuard(this, FunctionDefinition->GetContainedGraph(), false);
-			ImportNodesFromText(TextContent, false);
+			DefinitionController->ImportNodesFromText(TextContent, false);
 			if (FunctionDefinition->GetContainedGraph()->GetEntryNode() && InCollapseNode->GetContainedGraph()->GetEntryNode())
 			{ 
-				SetNodePosition(FunctionDefinition->GetContainedGraph()->GetEntryNode(), InCollapseNode->GetContainedGraph()->GetEntryNode()->GetPosition(), false);
+				DefinitionController->SetNodePosition(FunctionDefinition->GetContainedGraph()->GetEntryNode(), InCollapseNode->GetContainedGraph()->GetEntryNode()->GetPosition(), false);
 			}
 			
 			if (FunctionDefinition->GetContainedGraph()->GetReturnNode() && InCollapseNode->GetContainedGraph()->GetReturnNode())
 			{ 
-				SetNodePosition(FunctionDefinition->GetContainedGraph()->GetReturnNode(), InCollapseNode->GetContainedGraph()->GetReturnNode()->GetPosition(), false);
+				DefinitionController->SetNodePosition(FunctionDefinition->GetContainedGraph()->GetReturnNode(), InCollapseNode->GetContainedGraph()->GetReturnNode()->GetPosition(), false);
 			}
 
 			for (const URigVMLink* InnerLink : InCollapseNode->GetContainedGraph()->GetLinks())
@@ -5898,7 +5984,7 @@ URigVMFunctionReferenceNode* URigVMController::PromoteCollapseNodeToFunctionRefe
 				{
 					if (!SourcePin->IsLinkedTo(TargetPin))
 					{
-						AddLink(InnerLink->GetSourcePinPath(), InnerLink->GetTargetPinPath(), false);	
+						DefinitionController->AddLink(InnerLink->GetSourcePinPath(), InnerLink->GetTargetPinPath(), false);	
 					}
 				}				
 			}
@@ -6030,18 +6116,18 @@ URigVMCollapseNode* URigVMController::PromoteFunctionReferenceNodeToCollapseNode
 	{
 		FRestoreLinkedPathSettings RestoreSettings;
 
+		if(URigVMController* CollapseController = GetControllerForGraph(CollapseNode->GetContainedGraph()))
 		{
-			FRigVMControllerGraphGuard Guard(this, CollapseNode->GetContainedGraph(), false);
-			TGuardValue<bool> SuspendNotifications(bSuspendNotifications, true);
+			TGuardValue<bool> SuspendNotifications(CollapseController->bSuspendNotifications, true);
 
 			for (URigVMNode* Node : CollapseNode->GetContainedGraph()->GetNodes())
 			{
 				if (URigVMVariableNode* VariableNode = Cast<URigVMVariableNode>(Node))
 				{
 					const TArray<FLinkedPath> VariableLinkedPaths = GetLinkedPaths(VariableNode);
-					FastBreakLinkedPaths(VariableLinkedPaths);
-					RepopulatePinsOnNode(VariableNode, true, false, true);
-					RestoreLinkedPaths(VariableLinkedPaths, RestoreSettings);
+					CollapseController->FastBreakLinkedPaths(VariableLinkedPaths);
+					CollapseController->RepopulatePinsOnNode(VariableNode, true, false, true);
+					CollapseController->RestoreLinkedPaths(VariableLinkedPaths, RestoreSettings);
 				}
 			}
 
@@ -6059,8 +6145,10 @@ URigVMCollapseNode* URigVMController::PromoteFunctionReferenceNodeToCollapseNode
 
 	if(bRemoveFunctionDefinition)
 	{
-		FRigVMControllerGraphGuard Guard(this, FunctionDefinition->GetRootGraph(), false);
-		RemoveFunctionFromLibrary(FunctionDefinition->GetFName(), false);
+		if(URigVMController* LibraryController = GetControllerForGraph(FunctionDefinition->GetRootGraph()))
+		{
+			LibraryController->RemoveFunctionFromLibrary(FunctionDefinition->GetFName(), false);
+		}
 	}
 
 	return CollapseNode;
@@ -6090,9 +6178,11 @@ void URigVMController::SetReferencedFunction(URigVMFunctionReferenceNode* InFunc
 		}
 	}
 
-	
-	FRigVMControllerGraphGuard GraphGuard(this, InFunctionRefNode->GetGraph(), false);
-	GetGraph()->Notify(ERigVMGraphNotifType::NodeReferenceChanged, InFunctionRefNode);
+
+	if(URigVMController* ReferenceController = GetControllerForGraph(InFunctionRefNode->GetGraph()))
+	{
+		ReferenceController->Notify(ERigVMGraphNotifType::NodeReferenceChanged, InFunctionRefNode);
+	}
 }
 
 void URigVMController::RefreshFunctionPins(URigVMNode* InNode)
@@ -6228,16 +6318,18 @@ bool URigVMController::RemoveNodes(TArray<URigVMNode*> InNodes, bool bSetupUndoR
 
 								if(!bFoundExternalVariable)
 								{
-									FRigVMControllerGraphGuard Guard(this, OuterFunction->GetContainedGraph(), false);
-									if(RequestBulkEditDialogDelegate.IsBound())
+									if(URigVMController* FunctionController = GetControllerForGraph(OuterFunction->GetContainedGraph()))
 									{
-										FRigVMController_BulkEditResult Result = RequestBulkEditDialogDelegate.Execute(OuterFunction, ERigVMControllerBulkEditType::RemoveVariable);
-										if(Result.bCanceled)
+										if(FunctionController->RequestBulkEditDialogDelegate.IsBound())
 										{
-											return false;
+											FRigVMController_BulkEditResult Result = FunctionController->RequestBulkEditDialogDelegate.Execute(OuterFunction, ERigVMControllerBulkEditType::RemoveVariable);
+											if(Result.bCanceled)
+											{
+												return false;
+											}
+											bSetupUndoRedo = Result.bSetupUndoRedo;
+											bBulkEditDialogShown = true;
 										}
-										bSetupUndoRedo = Result.bSetupUndoRedo;
-										bBulkEditDialogShown = true;
 									}
 								}
 							}
@@ -6250,13 +6342,31 @@ bool URigVMController::RemoveNodes(TArray<URigVMNode*> InNodes, bool bSetupUndoR
 	
 	TGuardValue<bool> GuardCompactness(bIgnoreRerouteCompactnessChanges, true);
 
+	// remove injected nodes from the list as needed 
+	TArray<URigVMNode*> FilteredNodes;
+	FilteredNodes.Reserve(InNodes.Num());
+	for(URigVMNode* InNode : InNodes)
+	{
+		if(const URigVMInjectionInfo* InjectionInfo = InNode->GetInjectionInfo())
+		{
+			if(const URigVMNode* OuterNode = InjectionInfo->GetTypedOuter<URigVMNode>())
+			{
+				if(InNodes.Contains(OuterNode))
+				{
+					continue;
+				}
+			}
+		}
+		FilteredNodes.Add(InNode);
+	}
+
 	FRigVMControllerCompileBracketScope CompileScope(this);
 	FRigVMBaseAction Action(this);
 	if (bSetupUndoRedo)
 	{
-		if(InNodes.Num() == 1)
+		if(FilteredNodes.Num() == 1)
 		{
-			Action.SetTitle(FString::Printf(TEXT("Remove %s Node"), *InNodes[0]->GetNodeTitle()));
+			Action.SetTitle(FString::Printf(TEXT("Remove %s Node"), *FilteredNodes[0]->GetNodeTitle()));
 		}
 		else
 		{
@@ -6264,10 +6374,10 @@ bool URigVMController::RemoveNodes(TArray<URigVMNode*> InNodes, bool bSetupUndoR
 			Action.SetTitle(RemoveNodesTitle);
 		}
 		GetActionStack()->BeginAction(Action);
-		GetActionStack()->AddAction(FRigVMRemoveNodesAction(this, InNodes));
+		GetActionStack()->AddAction(FRigVMRemoveNodesAction(this, FilteredNodes));
 	}
 
-	for(URigVMNode* InNode : InNodes)
+	for(URigVMNode* InNode : FilteredNodes)
 	{
 		if (URigVMInjectionInfo* InjectionInfo = InNode->GetInjectionInfo())
 		{
@@ -6328,8 +6438,10 @@ bool URigVMController::RemoveNodes(TArray<URigVMNode*> InNodes, bool bSetupUndoR
 							
 							if (FunctionReferencePtr.IsValid())
 							{
-								FRigVMControllerGraphGuard GraphGuard(this, FunctionReferencePtr->GetGraph(), bSetupUndoRedo);
-								RemoveNode(FunctionReferencePtr.Get(), bSetupUndoRedo, false);
+								if(URigVMController* ReferenceController = GetControllerForGraph(FunctionReferencePtr->GetGraph()))
+								{
+									ReferenceController->RemoveNode(FunctionReferencePtr.Get(), bSetupUndoRedo, false);
+								}
 							}
 						}
 					}
@@ -6378,11 +6490,21 @@ bool URigVMController::RemoveNodes(TArray<URigVMNode*> InNodes, bool bSetupUndoR
 		if (URigVMCollapseNode* CollapseNode = Cast<URigVMCollapseNode>(InNode))
 		{
 			URigVMGraph* SubGraph = CollapseNode->GetContainedGraph();
-			FRigVMControllerGraphGuard GraphGuard(this, SubGraph, false);
-		
-			TArray<URigVMNode*> ContainedNodes = SubGraph->GetNodes();
-			TGuardValue<bool> SuspendTemplateComputation(bSuspendTemplateComputation, true);
-			RemoveNodes(ContainedNodes, false, false);
+
+			if(URigVMController* SubGraphController = GetControllerForGraph(SubGraph))
+			{
+				TArray<URigVMNode*> ContainedNodes = SubGraph->GetNodes();
+				TGuardValue<bool> SuspendTemplateComputation(bSuspendTemplateComputation, true);
+				SubGraphController->RemoveNodes(ContainedNodes, false, false);
+
+				if(IRigVMClientHost* ClientHost = GetImplementingOuter<IRigVMClientHost>())
+				{
+					if(FRigVMClient* Client = ClientHost->GetRigVMClient())
+					{
+						(void)Client->RemoveController(SubGraph);
+					}
+				}
+			}
 		}
 		
 		Graph->Nodes.Remove(InNode);
@@ -6400,7 +6522,7 @@ bool URigVMController::RemoveNodes(TArray<URigVMNode*> InNodes, bool bSetupUndoR
 
 		if (Graph->IsA<URigVMFunctionLibrary>())
 		{
-			for(const URigVMNode* InNode : InNodes)
+			for(const URigVMNode* InNode : FilteredNodes)
 			{
 				const FString NodeName = GetSanitizedNodeName(InNode->GetName());
 				
@@ -6412,7 +6534,7 @@ bool URigVMController::RemoveNodes(TArray<URigVMNode*> InNodes, bool bSetupUndoR
 		else
 		{
 			TArray<FString> NodePaths;
-			Algo::Transform(InNodes, NodePaths, [](const URigVMNode* Node)
+			Algo::Transform(FilteredNodes, NodePaths, [](const URigVMNode* Node)
 			{
 				static constexpr TCHAR Format[] = TEXT("'%s'");
 				return FString::Printf(Format, *GetSanitizedPinPath(Node->GetNodePath()));
@@ -6427,7 +6549,7 @@ bool URigVMController::RemoveNodes(TArray<URigVMNode*> InNodes, bool bSetupUndoR
 		}
 	}
 
-	for(URigVMNode* InNode : InNodes)
+	for(URigVMNode* InNode : FilteredNodes)
 	{
 		if (URigVMVariableNode* VariableNode = Cast<URigVMVariableNode>(InNode))
 		{
@@ -6553,8 +6675,10 @@ bool URigVMController::RenameNode(URigVMNode* InNode, const FName& InNewName, bo
 			{
 				BuildData->ForEachFunctionReference(LibraryNode->GetFunctionIdentifier(), [this, InNewName](URigVMFunctionReferenceNode* ReferenceNode)
 				{
-					FRigVMControllerGraphGuard GraphGuard(this, ReferenceNode->GetGraph(), false);
-					RenameNode(ReferenceNode, InNewName, false);
+					if(URigVMController* ReferenceController = GetControllerForGraph(ReferenceNode->GetGraph()))
+					{
+						ReferenceController->RenameNode(ReferenceNode, InNewName, false);
+					}
 				});
 			}
 
@@ -6913,8 +7037,10 @@ bool URigVMController::SetNodeColor(URigVMNode* InNode, const FLinearColor& InCo
 		{
 			FunctionLibrary->ForEachReference(LibraryNode->GetFName(), [this](URigVMFunctionReferenceNode* ReferenceNode)
             {
-                FRigVMControllerGraphGuard GraphGuard(this, ReferenceNode->GetGraph(), false);
-				Notify(ERigVMGraphNotifType::NodeColorChanged, ReferenceNode);
+				if(URigVMController* ReferenceController = GetControllerForGraph(ReferenceNode->GetGraph()))
+				{
+					ReferenceController->Notify(ERigVMGraphNotifType::NodeColorChanged, ReferenceNode);
+				}
             });
 		}
 	}
@@ -8134,47 +8260,49 @@ FString URigVMController::AddAggregatePin(URigVMNode* InNode, const FString& InP
 		AggregateNode = Cast<URigVMAggregateNode>(CollapseNode);
 		if (AggregateNode)
 		{
-			FRigVMControllerGraphGuard GraphGuard(this, AggregateNode->GetContainedGraph(), false);
-			TGuardValue<bool> GuardEditGraph(GetGraph()->bEditable, true);
-
-			for(int32 Index = 0; Index < InNode->GetPins().Num(); Index++)
+			if(URigVMController* AggregateController = GetControllerForGraph(AggregateNode->GetContainedGraph()))
 			{
-				URigVMPin* Pin = InNode->GetPins()[Index];
-				const FName PinName = Pin->GetFName();
-				
-				if (URigVMPin* AggregatePin = AggregateNode->FindPin(PinName.ToString()))
+				TGuardValue<bool> GuardEditGraph(AggregateController->GetGraph()->bEditable, true);
+
+				for(int32 Index = 0; Index < InNode->GetPins().Num(); Index++)
 				{
-					SetExposedPinIndex(PinName, Index, false);
-					continue;
-				}
+					URigVMPin* Pin = InNode->GetPins()[Index];
+					const FName PinName = Pin->GetFName();
+				
+					if (URigVMPin* AggregatePin = AggregateNode->FindPin(PinName.ToString()))
+					{
+						AggregateController->SetExposedPinIndex(PinName, Index, false);
+						continue;
+					}
 
-				const FName ExposedPinName = AddExposedPin(PinName, Pin->GetDirection(), Pin->GetCPPType(), *Pin->GetCPPTypeObject()->GetPathName(), Pin->GetDefaultValue(), false, false);
+					const FName ExposedPinName = AggregateController->AddExposedPin(PinName, Pin->GetDirection(), Pin->GetCPPType(), *Pin->GetCPPTypeObject()->GetPathName(), Pin->GetDefaultValue(), false, false);
 
-				const FString PinNameStr = PinName.ToString();
-				const FString ExposedPinNameStr = ExposedPinName.ToString();
+					const FString PinNameStr = PinName.ToString();
+					const FString ExposedPinNameStr = ExposedPinName.ToString();
 
-				if(URigVMPin* ExposedPin = AggregateNode->FindPin(ExposedPinNameStr))
-				{
-					ExposedPin->SetDisplayName(Pin->GetDisplayName());
-				}
+					if(URigVMPin* ExposedPin = AggregateNode->FindPin(ExposedPinNameStr))
+					{
+						ExposedPin->SetDisplayName(Pin->GetDisplayName());
+					}
 
-				if(URigVMPin* ExposedPin = AggregateNode->GetEntryNode()->FindPin(ExposedPinNameStr))
-				{
-					ExposedPin->SetDisplayName(Pin->GetDisplayName());
-				}
+					if(URigVMPin* ExposedPin = AggregateNode->GetEntryNode()->FindPin(ExposedPinNameStr))
+					{
+						ExposedPin->SetDisplayName(Pin->GetDisplayName());
+					}
 
-				if(URigVMPin* ExposedPin = AggregateNode->GetReturnNode()->FindPin(ExposedPinNameStr))
-				{
-					ExposedPin->SetDisplayName(Pin->GetDisplayName());
-				}
+					if(URigVMPin* ExposedPin = AggregateNode->GetReturnNode()->FindPin(ExposedPinNameStr))
+					{
+						ExposedPin->SetDisplayName(Pin->GetDisplayName());
+					}
 
-				if (Pin->GetDirection() == ERigVMPinDirection::Input)
-				{
-					AddLink(FString::Printf(TEXT("Entry.%s"), *ExposedPinNameStr), FString::Printf(TEXT("%s.%s"), *InNode->GetName(), *PinNameStr), false);
-				}
-				else
-				{
-					AddLink(FString::Printf(TEXT("%s.%s"), *InNode->GetName(), *PinNameStr), FString::Printf(TEXT("Return.%s"), *ExposedPinNameStr), false);
+					if (Pin->GetDirection() == ERigVMPinDirection::Input)
+					{
+						AggregateController->AddLink(FString::Printf(TEXT("Entry.%s"), *ExposedPinNameStr), FString::Printf(TEXT("%s.%s"), *InNode->GetName(), *PinNameStr), false);
+					}
+					else
+					{
+						AggregateController->AddLink(FString::Printf(TEXT("%s.%s"), *InNode->GetName(), *PinNameStr), FString::Printf(TEXT("Return.%s"), *ExposedPinNameStr), false);
+					}
 				}
 			}
 		}
@@ -8194,14 +8322,14 @@ FString URigVMController::AddAggregatePin(URigVMNode* InNode, const FString& InP
 	}
 
 	URigVMPin* NewPin = nullptr;	
+	if(URigVMController* AggregateController = GetControllerForGraph(AggregateNode->GetContainedGraph()))
 	{
-		FRigVMControllerGraphGuard GraphGuard(this, AggregateNode->GetContainedGraph(), false);
-		TGuardValue<bool> GuardEditGraph(GetGraph()->bEditable, true);
+		TGuardValue<bool> GuardEditGraph(AggregateController->GetGraph()->bEditable, true);
 
 		URigVMNode* InnerNode = (AggregateNode == nullptr) ? InNode : AggregateNode->GetFirstInnerNode();
 
-		const FString InnerNodeContent = ExportNodesToText({InnerNode->GetFName()});
-		const TArray<FName> NewNodeNames = ImportNodesFromText(InnerNodeContent, false);
+		const FString InnerNodeContent = AggregateController->ExportNodesToText({InnerNode->GetFName()});
+		const TArray<FName> NewNodeNames = AggregateController->ImportNodesFromText(InnerNodeContent, false);
 		
 		if(NewNodeNames.IsEmpty())
 		{
@@ -8242,7 +8370,7 @@ FString URigVMController::AddAggregatePin(URigVMNode* InNode, const FString& InP
 		}
 		
 		const URigVMPin* Arg1 = AggregateNode->GetFirstAggregatePin();
-		FName NewExposedPinName = AddExposedPin(NewPinName, Arg1->GetDirection(), Arg1->GetCPPType(), *Arg1->GetCPPTypeObject()->GetPathName(), InDefaultValue, false);
+		FName NewExposedPinName = AggregateController->AddExposedPin(NewPinName, Arg1->GetDirection(), Arg1->GetCPPType(), *Arg1->GetCPPTypeObject()->GetPathName(), InDefaultValue, false);
 		NewPin = AggregateNode->FindPin(NewExposedPinName.ToString());
 		URigVMPin* NewUnitPinArg1 = NewNode->GetFirstAggregatePin();
 		URigVMPin* NewUnitPinArg2 = NewNode->GetSecondAggregatePin();
@@ -8256,10 +8384,10 @@ FString URigVMController::AddAggregatePin(URigVMNode* InNode, const FString& InP
 			URigVMPin* PreviousReturnPin = ReturnPin->GetLinkedSourcePins()[0];
 			PreviousNode = PreviousReturnPin->GetNode();
 		
-			BreakAllLinks(ReturnPin, true, false);
-			AddLink(PreviousReturnPin, NewUnitPinArg1, false);						
-			AddLink(EntryPin, NewUnitPinArg2, false);
-			AddLink(NewUnitPinOppositeArg, ReturnPin, false);
+			AggregateController->BreakAllLinks(ReturnPin, true, false);
+			AggregateController->AddLink(PreviousReturnPin, NewUnitPinArg1, false);						
+			AggregateController->AddLink(EntryPin, NewUnitPinArg2, false);
+			AggregateController->AddLink(NewUnitPinOppositeArg, ReturnPin, false);
 		}
 		else
 		{
@@ -8269,17 +8397,17 @@ FString URigVMController::AddAggregatePin(URigVMNode* InNode, const FString& InP
 			URigVMPin* PreviousReturnPin = OldReturnPin->GetLinkedSourcePins()[0];
 			PreviousNode = PreviousReturnPin->GetNode();
 
-			BreakAllLinks(OldReturnPin, true, false);
-			AddLink(PreviousReturnPin, NewUnitPinOppositeArg, false);						
-			AddLink(NewUnitPinArg1, OldReturnPin, false);
-			AddLink(NewUnitPinArg2, NewReturnPin, false);
+			AggregateController->BreakAllLinks(OldReturnPin, true, false);
+			AggregateController->AddLink(PreviousReturnPin, NewUnitPinOppositeArg, false);						
+			AggregateController->AddLink(NewUnitPinArg1, OldReturnPin, false);
+			AggregateController->AddLink(NewUnitPinArg2, NewReturnPin, false);
 		}
 
 		// Rearrange the graph nodes
 		URigVMFunctionReturnNode* ReturnNode = AggregateNode->GetReturnNode();
 		FVector2D NodeDimensions(200, 150);
-		SetNodePosition(NewNode, PreviousNode->GetPosition() + NodeDimensions, false);
-		SetNodePosition(ReturnNode, NewNode->GetPosition() + NodeDimensions, false);
+		AggregateController->SetNodePosition(NewNode, PreviousNode->GetPosition() + NodeDimensions, false);
+		AggregateController->SetNodePosition(ReturnNode, NewNode->GetPosition() + NodeDimensions, false);
 
 		// Connect other input pins
 		for (URigVMPin* OtherInputPin : AggregateNode->GetFirstInnerNode()->GetPins())
@@ -8289,7 +8417,7 @@ FString URigVMController::AddAggregatePin(URigVMNode* InNode, const FString& InP
 				OtherInputPin->GetName() != NewUnitPinOppositeArg->GetName())
 			{
 				URigVMPin* OtherEntryPin = AggregateNode->GetEntryNode()->FindPin(OtherInputPin->GetName());
-				AddLink(OtherEntryPin, NewNode->FindPin(OtherEntryPin->GetName()), false);
+				AggregateController->AddLink(OtherEntryPin, NewNode->FindPin(OtherEntryPin->GetName()), false);
 			}
 		}
 
@@ -8303,8 +8431,10 @@ FString URigVMController::AddAggregatePin(URigVMNode* InNode, const FString& InP
 
 	if(!PinStates.IsEmpty())
 	{
-		FRigVMControllerGraphGuard GraphGuard(this, AggregateNode->GetGraph(), false);
-		ApplyPinStates(AggregateNode, PinStates);
+		if(URigVMController* AggregateController = GetControllerForGraph(AggregateNode->GetContainedGraph()))
+		{
+			AggregateController->ApplyPinStates(AggregateNode, PinStates);
+		}
 	}
 
 	if (bSetupUndoRedo)
@@ -8398,28 +8528,30 @@ bool URigVMController::RemoveAggregatePin(URigVMPin* InPin, bool bSetupUndoRedo,
 				{
 					if (EntryPin->GetLinkedTargetPins().Num() > 0)
 					{
-						FRigVMControllerGraphGuard GraphGuard(this, AggregateNode->GetContainedGraph(), false);
-						TGuardValue<bool> GuardEditGraph(GetGraph()->bEditable, true);
+						if(URigVMController* AggregateController = GetControllerForGraph(AggregateNode->GetContainedGraph()))
+						{
+							TGuardValue<bool> GuardEditGraph(AggregateController->GetGraph()->bEditable, true);
 						
-						URigVMPin* TargetPin = EntryPin->GetLinkedTargetPins()[0];
-					
-						URigVMNode* NodeToRemove = TargetPin->GetNode();
-						URigVMPin* ResultPin = NodeToRemove->GetOppositeAggregatePin();
-						URigVMPin* NextNodePin = ResultPin->GetLinkedTargetPins()[0];
+							URigVMPin* TargetPin = EntryPin->GetLinkedTargetPins()[0];
+						
+							URigVMNode* NodeToRemove = TargetPin->GetNode();
+							URigVMPin* ResultPin = NodeToRemove->GetOppositeAggregatePin();
+							URigVMPin* NextNodePin = ResultPin->GetLinkedTargetPins()[0];
 
-						if (NodeToRemove == AggregateNode->FirstInnerNodeCache || NodeToRemove == AggregateNode->LastInnerNodeCache)
-						{
-							AggregateNode->InvalidateCache();
-						}
+							if (NodeToRemove == AggregateNode->FirstInnerNodeCache || NodeToRemove == AggregateNode->LastInnerNodeCache)
+							{
+								AggregateNode->InvalidateCache();
+							}
 
-						const FString FirstAggregatePin = AggregateNode->GetFirstAggregatePin()->GetName();
-						const FString SecondAggregatePin = AggregateNode->GetSecondAggregatePin()->GetName();
-						FString OtherArg = TargetPin->GetName() == FirstAggregatePin ? SecondAggregatePin : FirstAggregatePin;
-						BreakAllLinks(NextNodePin, true, false);
-						RewireLinks(NodeToRemove->FindPin(OtherArg), NextNodePin, true, false);
-						RemoveNode(NodeToRemove, false);
-						RemoveExposedPin(*InPin->GetName(), false);
-						bSuccess = true;
+							const FString FirstAggregatePin = AggregateNode->GetFirstAggregatePin()->GetName();
+							const FString SecondAggregatePin = AggregateNode->GetSecondAggregatePin()->GetName();
+							FString OtherArg = TargetPin->GetName() == FirstAggregatePin ? SecondAggregatePin : FirstAggregatePin;
+							AggregateController->BreakAllLinks(NextNodePin, true, false);
+							AggregateController->RewireLinks(NodeToRemove->FindPin(OtherArg), NextNodePin, true, false);
+							AggregateController->RemoveNode(NodeToRemove, false);
+							AggregateController->RemoveExposedPin(*InPin->GetName(), false);
+							bSuccess = true;
+						}
 					}
 				}
 			}
@@ -8432,29 +8564,31 @@ bool URigVMController::RemoveAggregatePin(URigVMPin* InPin, bool bSetupUndoRedo,
 				{
 					if (ReturnPin->GetLinkedSourcePins().Num() > 0)
 					{
-						FRigVMControllerGraphGuard GraphGuard(this, AggregateNode->GetContainedGraph(), false);
-						TGuardValue<bool> GuardEditGraph(GetGraph()->bEditable, true);
+						if(URigVMController* AggregateController = GetControllerForGraph(AggregateNode->GetContainedGraph()))
+						{
+							TGuardValue<bool> GuardEditGraph(AggregateController->GetGraph()->bEditable, true);
 						
-						URigVMPin* SourcePin = ReturnPin->GetLinkedSourcePins()[0];
+							URigVMPin* SourcePin = ReturnPin->GetLinkedSourcePins()[0];
 					
-						URigVMNode* NodeToRemove = SourcePin->GetNode();
-						URigVMPin* OppositePin = NodeToRemove->GetOppositeAggregatePin();
-						URigVMPin* NextNodePin = OppositePin->GetLinkedSourcePins()[0];
-						URigVMNode* NextNode = NextNodePin->GetNode();
+							URigVMNode* NodeToRemove = SourcePin->GetNode();
+							URigVMPin* OppositePin = NodeToRemove->GetOppositeAggregatePin();
+							URigVMPin* NextNodePin = OppositePin->GetLinkedSourcePins()[0];
+							URigVMNode* NextNode = NextNodePin->GetNode();
 
-						if (NodeToRemove == AggregateNode->FirstInnerNodeCache || NodeToRemove == AggregateNode->LastInnerNodeCache)
-						{
-							AggregateNode->InvalidateCache();
-						}
+							if (NodeToRemove == AggregateNode->FirstInnerNodeCache || NodeToRemove == AggregateNode->LastInnerNodeCache)
+							{
+								AggregateNode->InvalidateCache();
+							}
 
-						const FString FirstAggregatePin = AggregateNode->GetFirstAggregatePin()->GetName();
-						const FString SecondAggregatePin = AggregateNode->GetSecondAggregatePin()->GetName();
-						FString OtherArg = SourcePin->GetName() == FirstAggregatePin ? SecondAggregatePin : FirstAggregatePin;
-						BreakAllLinks(NextNodePin, false, false);
-						RewireLinks(NodeToRemove->FindPin(OtherArg), NextNodePin, false, false);
-						RemoveNode(NodeToRemove, false);
-						RemoveExposedPin(*InPin->GetName(), false);
-						bSuccess = true;
+							const FString FirstAggregatePin = AggregateNode->GetFirstAggregatePin()->GetName();
+							const FString SecondAggregatePin = AggregateNode->GetSecondAggregatePin()->GetName();
+							FString OtherArg = SourcePin->GetName() == FirstAggregatePin ? SecondAggregatePin : FirstAggregatePin;
+							AggregateController->BreakAllLinks(NextNodePin, false, false);
+							AggregateController->RewireLinks(NodeToRemove->FindPin(OtherArg), NextNodePin, false, false);
+							AggregateController->RemoveNode(NodeToRemove, false);
+							AggregateController->RemoveExposedPin(*InPin->GetName(), false);
+							bSuccess = true;
+						}
 					}
 				}
 			}			
@@ -10341,13 +10475,14 @@ FName URigVMController::AddExposedPin(const FName& InPinName, ERigVMPinDirection
 
 	if (Pin->IsStruct())
 	{
-		FRigVMControllerGraphGuard GraphGuard(this, LibraryNode->GetGraph(), bSetupUndoRedo);
-
-		FString DefaultValue = InDefaultValue;
-		CreateDefaultValueForStructIfRequired(Pin->GetScriptStruct(), DefaultValue);
+		if(URigVMController* LibraryController = GetControllerForGraph(LibraryNode->GetGraph()))
 		{
-			TGuardValue<bool> SuspendNotifications(bSuspendNotifications, true);
-			AddPinsForStruct(Pin->GetScriptStruct(), LibraryNode, Pin, Pin->Direction, DefaultValue, false);
+			FString DefaultValue = InDefaultValue;
+			LibraryController->CreateDefaultValueForStructIfRequired(Pin->GetScriptStruct(), DefaultValue);
+			{
+				TGuardValue<bool> SuspendNotifications(LibraryController->bSuspendNotifications, true);
+				LibraryController->AddPinsForStruct(Pin->GetScriptStruct(), LibraryNode, Pin, Pin->Direction, DefaultValue, false);
+			}
 		}
 	}
 
@@ -10358,15 +10493,17 @@ FName URigVMController::AddExposedPin(const FName& InPinName, ERigVMPinDirection
 		GetActionStack()->BeginAction(Action);
 	}
 
+	if(URigVMController* LibraryController = GetControllerForGraph(LibraryNode->GetGraph()))
 	{
-		FRigVMControllerGraphGuard GraphGuard(this, LibraryNode->GetGraph(), bSetupUndoRedo);
-		Notify(ERigVMGraphNotifType::PinAdded, Pin);
+		LibraryController->Notify(ERigVMGraphNotifType::PinAdded, Pin);
 	}
 
 	if (!InDefaultValue.IsEmpty())
 	{
-		FRigVMControllerGraphGuard GraphGuard(this, Pin->GetGraph(), bSetupUndoRedo);
-		SetPinDefaultValue(Pin, InDefaultValue, true, bSetupUndoRedo, false);
+		if(URigVMController* PinController = GetControllerForGraph(Pin->GetGraph()))
+		{
+			PinController->SetPinDefaultValue(Pin, InDefaultValue, true, bSetupUndoRedo, false);
+		}
 	}
 
 	URigVMFunctionEntryNode* EntryNode = Graph->GetEntryNode();
@@ -10484,9 +10621,9 @@ bool URigVMController::RemoveExposedPin(const FName& InPinName, bool bSetupUndoR
 
 	bool bSuccessfullyRemovedPin = false;
 	{
+		if(URigVMController* LibraryController = GetControllerForGraph(LibraryNode->GetGraph()))
 		{
-			FRigVMControllerGraphGuard GraphGuard(this, LibraryNode->GetGraph(), bSetupUndoRedo);
-			bSuccessfullyRemovedPin = RemovePin(Pin, bSetupUndoRedo, true);
+			bSuccessfullyRemovedPin = LibraryController->RemovePin(Pin, bSetupUndoRedo, true);
 		}
 
 		TArray<URigVMVariableNode*> NodesToRemove;
@@ -10611,13 +10748,16 @@ bool URigVMController::RenameExposedPin(const FName& InOldPinName, const FName&
 
 	struct Local
 	{
-		static bool RenamePin(URigVMController* InController, URigVMPin* InPin, const FName& InNewName)
+		static bool RenamePin(const URigVMController* InController, URigVMPin* InPin, const FName& InNewName)
 		{
-			FRigVMControllerGraphGuard GraphGuard(InController, InPin->GetGraph(), false);
-
-			TArray<FLinkedPath> LinkedPaths = GetLinkedPaths(InPin, true, true);
+			URigVMController* PinController = InController->GetControllerForGraph(InPin->GetGraph());
+			if(PinController == nullptr)
+			{
+				return false;
+			}
+			TArray<FLinkedPath> LinkedPaths = PinController->GetLinkedPaths(InPin, true, true);
 
-			InController->FastBreakLinkedPaths(LinkedPaths);
+			PinController->FastBreakLinkedPaths(LinkedPaths);
 
 			const FString OldPinPath = InPin->GetPinPath();
 			static constexpr TCHAR PinPrefixFormat[] = TEXT("%s.");
@@ -10630,7 +10770,7 @@ bool URigVMController::RenameExposedPin(const FName& InOldPinName, const FName&
 			
 			InPin->SetDisplayName(InNewName);
 			
-			InController->Notify(ERigVMGraphNotifType::PinRenamed, InPin);
+			PinController->Notify(ERigVMGraphNotifType::PinRenamed, InPin);
 
 			FRestoreLinkedPathSettings Settings;
 			Settings.RemapDelegates = {{
@@ -10653,7 +10793,7 @@ bool URigVMController::RenameExposedPin(const FName& InOldPinName, const FName&
 					return InPinPath;
 				})
 			}};
-			InController->RestoreLinkedPaths(LinkedPaths, Settings);
+			PinController->RestoreLinkedPaths(LinkedPaths, Settings);
 
 			return true;
 		}
@@ -10683,8 +10823,10 @@ bool URigVMController::RenameExposedPin(const FName& InOldPinName, const FName&
         {
 			if (URigVMPin* EntryPin = ReferenceNode->FindPin(InOldPinName.ToString()))
 			{
-                FRigVMControllerGraphGuard GraphGuard(this, ReferenceNode->GetGraph(), false);
-                Local::RenamePin(this, EntryPin, PinName);
+				if(URigVMController* ReferenceController = GetControllerForGraph(ReferenceNode->GetGraph()))
+				{
+					Local::RenamePin(ReferenceController, EntryPin, PinName);
+				}
             }
         });
 	}
@@ -10875,14 +11017,14 @@ bool URigVMController::ChangeExposedPinType(const FName& InPinName, const FStrin
 
 	// Change pin type of the library node in the function library
 	{
-		bool bSuccessChangingType;
+		bool bSuccessChangingType = false;
+		if(URigVMController* LibraryController = GetControllerForGraph(LibraryNode->GetGraph()))
 		{
-			FRigVMControllerGraphGuard GraphGuard(this, LibraryNode->GetGraph(), bSetupUndoRedo);
-			bSuccessChangingType = ChangePinType(Pin, InCPPType, InCPPTypeObjectPath, bSetupUndoRedo, bSetupOrphanPins);
+			bSuccessChangingType = LibraryController->ChangePinType(Pin, InCPPType, InCPPTypeObjectPath, bSetupUndoRedo, bSetupOrphanPins);
 
 			if (bSuccessChangingType)
 			{
-				RemoveUnusedOrphanedPins(LibraryNode);
+				LibraryController->RemoveUnusedOrphanedPins(LibraryNode);
 			}
 		}
 		if (!bSuccessChangingType)
@@ -10916,9 +11058,11 @@ bool URigVMController::ChangeExposedPinType(const FName& InPinName, const FStrin
         {
 			if (URigVMPin* ReferencedNodePin = ReferenceNode->FindPin(Pin->GetName()))
 			{
-				FRigVMControllerGraphGuard GraphGuard(this, ReferenceNode->GetGraph(), bSetupUndoRedo);
-				ChangePinType(ReferencedNodePin, InCPPType, InCPPTypeObjectPath, bSetupUndoRedo, bSetupOrphanPins);
-				RemoveUnusedOrphanedPins(ReferenceNode);
+				if(URigVMController* ReferenceController = GetControllerForGraph(ReferenceNode->GetGraph()))
+				{
+					ReferenceController->ChangePinType(ReferencedNodePin, InCPPType, InCPPTypeObjectPath, bSetupUndoRedo, bSetupOrphanPins);
+					ReferenceController->RemoveUnusedOrphanedPins(ReferenceNode);
+				}
 			}
         });
 	}
@@ -11015,8 +11159,10 @@ bool URigVMController::SetExposedPinIndex(const FName& InPinName, int32 InNewInd
 		LibraryNode->Pins.Remove(Pin);
 		LibraryNode->Pins.Insert(Pin, InNewIndex);
 
-		FRigVMControllerGraphGuard GraphGuard(this, LibraryNode->GetGraph(), false);
-		Notify(ERigVMGraphNotifType::PinIndexChanged, Pin);
+		if(const URigVMController* LibraryController = GetControllerForGraph(LibraryNode->GetGraph()))
+		{
+			LibraryController->Notify(ERigVMGraphNotifType::PinIndexChanged, Pin);
+		}
 	}
 
 	RefreshFunctionPins(LibraryNode->GetEntryNode());
@@ -11336,18 +11482,20 @@ URigVMLibraryNode* URigVMController::AddFunctionToLibrary(const FName& InFunctio
 	if (bMutable)
 	{
 		const UScriptStruct* ExecuteContextStruct = FRigVMExecuteContext::StaticStruct();
-		
-		FRigVMControllerGraphGuard GraphGuard(this, CollapseNode->ContainedGraph, false);
-		AddExposedPin(FRigVMStruct::ExecuteContextName,
-			ERigVMPinDirection::IO,
-			FString::Printf(TEXT("F%s"), *ExecuteContextStruct->GetName()),
-			*ExecuteContextStruct->GetPathName(),
-			FString(),
-			false);
+
+		if(URigVMController* CollapseController = GetControllerForGraph(CollapseNode->GetContainedGraph()))
+		{
+			CollapseController->AddExposedPin(FRigVMStruct::ExecuteContextName,
+				ERigVMPinDirection::IO,
+				FString::Printf(TEXT("F%s"), *ExecuteContextStruct->GetName()),
+				*ExecuteContextStruct->GetPathName(),
+				FString(),
+				false);
+		}
 	}
 
+	if(URigVMController* CollapseController = GetControllerForGraph(CollapseNode->GetContainedGraph()))
 	{
-		FRigVMControllerGraphGuard GraphGuard(this, CollapseNode->GetContainedGraph(), false);
 		TGuardValue<bool> GuardEditGraph(CollapseNode->ContainedGraph->bEditable, true);
 
 		URigVMNode* EntryNode = CollapseNode->ContainedGraph->FindNode(TEXT("Entry"));
@@ -11358,32 +11506,32 @@ URigVMLibraryNode* URigVMController::AddFunctionToLibrary(const FName& InFunctio
 			EntryNode = NewObject<URigVMFunctionEntryNode>(CollapseNode->ContainedGraph, TEXT("Entry"));
 			CollapseNode->ContainedGraph->Nodes.Add(EntryNode);
 			{
-				TGuardValue<bool> SuspendNotifications(bSuspendNotifications, true);
-				RefreshFunctionPins(EntryNode);
+				TGuardValue<bool> SuspendNotifications(CollapseController->bSuspendNotifications, true);
+				CollapseController->RefreshFunctionPins(EntryNode);
 			}
-			Notify(ERigVMGraphNotifType::NodeAdded, EntryNode);
+			CollapseController->Notify(ERigVMGraphNotifType::NodeAdded, EntryNode);
 		}
 
 		EntryNode->Position = FVector2D(-250.f, 0.f);
-		Notify(ERigVMGraphNotifType::NodePositionChanged, EntryNode);
+		CollapseController->Notify(ERigVMGraphNotifType::NodePositionChanged, EntryNode);
 
 		if (ReturnNode == nullptr)
 		{
 			ReturnNode = NewObject<URigVMFunctionReturnNode>(CollapseNode->ContainedGraph, TEXT("Return"));
 			CollapseNode->ContainedGraph->Nodes.Add(ReturnNode);
 			{
-				TGuardValue<bool> SuspendNotifications(bSuspendNotifications, true);
-				RefreshFunctionPins(ReturnNode);
+				TGuardValue<bool> SuspendNotifications(CollapseController->bSuspendNotifications, true);
+				CollapseController->RefreshFunctionPins(ReturnNode);
 			}
-			Notify(ERigVMGraphNotifType::NodeAdded, ReturnNode);
+			CollapseController->Notify(ERigVMGraphNotifType::NodeAdded, ReturnNode);
 		}
 
 		ReturnNode->Position = FVector2D(250.f, 0.f);
-		Notify(ERigVMGraphNotifType::NodePositionChanged, ReturnNode);
+		CollapseController->Notify(ERigVMGraphNotifType::NodePositionChanged, ReturnNode);
 
 		if (bMutable)
 		{
-			AddLink(EntryNode->FindPin(FRigVMStruct::ExecuteContextName.ToString()), ReturnNode->FindPin(FRigVMStruct::ExecuteContextName.ToString()), false);
+			CollapseController->AddLink(EntryNode->FindPin(FRigVMStruct::ExecuteContextName.ToString()), ReturnNode->FindPin(FRigVMStruct::ExecuteContextName.ToString()), false);
 		}
 	}
 
@@ -12259,17 +12407,19 @@ TArray<TSoftObjectPtr<URigVMFunctionReferenceNode>> URigVMController::GetAffecte
 			{
 				if(URigVMLibraryNode* AffectedFunction = AffectedFunctionReferenceNode->FindFunctionForNode())
 				{
-					FRigVMControllerGraphGuard GraphGuard(this, AffectedFunction->GetContainedGraph(), false);
-					TGuardValue<bool> SuspendNotifications(bSuspendNotifications, true);
-					TArray<TSoftObjectPtr<URigVMFunctionReferenceNode>> AffectedFunctionReferencePtrs = GetAffectedReferences(InEditType, bForceLoad);
-					for(TSoftObjectPtr<URigVMFunctionReferenceNode> AffectedFunctionReferencePtr : AffectedFunctionReferencePtrs)
+					if(URigVMController* FunctionController = GetControllerForGraph(AffectedFunction->GetContainedGraph()))
 					{
-						const FString Key = AffectedFunctionReferencePtr.ToSoftObjectPath().ToString();
-						if(VisitedPaths.Contains(Key))
+						TGuardValue<bool> SuspendNotifications(FunctionController->bSuspendNotifications, true);
+						TArray<TSoftObjectPtr<URigVMFunctionReferenceNode>> AffectedFunctionReferencePtrs = FunctionController->GetAffectedReferences(InEditType, bForceLoad);
+						for(TSoftObjectPtr<URigVMFunctionReferenceNode> AffectedFunctionReferencePtr : AffectedFunctionReferencePtrs)
 						{
-							continue;
+							const FString Key = AffectedFunctionReferencePtr.ToSoftObjectPath().ToString();
+							if(VisitedPaths.Contains(Key))
+							{
+								continue;
+							}
+							VisitedPaths.Add(Key, FunctionReferencePtrs.Add(AffectedFunctionReferencePtr));
 						}
-						VisitedPaths.Add(Key, FunctionReferencePtrs.Add(AffectedFunctionReferencePtr));
 					}
 				}
 			}
@@ -13378,16 +13528,17 @@ bool URigVMController::CanAddNode(URigVMNode* InNode, bool bReportErrors, bool b
 			}
 		}
 	}
-	else if(URigVMCollapseNode* CollapseNode = Cast<URigVMCollapseNode>(InNode))
+	else if(const URigVMCollapseNode* CollapseNode = Cast<URigVMCollapseNode>(InNode))
 	{
-		FRigVMControllerGraphGuard GraphGuard(this, CollapseNode->GetContainedGraph(), false);
-
-		TArray<URigVMNode*> ContainedNodes = CollapseNode->GetContainedNodes();
-		for(URigVMNode* ContainedNode : ContainedNodes)
+		if(URigVMController* CollapseController = GetControllerForGraph(CollapseNode->GetContainedGraph()))
 		{
-			if(!CanAddNode(ContainedNode, bReportErrors, true))
+			TArray<URigVMNode*> ContainedNodes = CollapseNode->GetContainedNodes();
+			for(URigVMNode* ContainedNode : ContainedNodes)
 			{
-				return false;
+				if(!CollapseController->CanAddNode(ContainedNode, bReportErrors, true))
+				{
+					return false;
+				}
 			}
 		}
 	}
@@ -14664,14 +14815,16 @@ void URigVMController::RepopulatePinsOnNode(URigVMNode* InNode, bool bFollowCore
 		{
 			if(InCollapseNode)
 			{
-				FRigVMControllerGraphGuard GraphGuard(this, InCollapseNode->GetContainedGraph(), false);
-				TGuardValue<bool> GuardEditGraph(InCollapseNode->ContainedGraph->bEditable, true);
-				// need to get a copy of the node array since the following function could remove nodes from the graph
-				// we don't want to remove elements from the array we are iterating over.
-				TArray<URigVMNode*> ContainedNodes = InCollapseNode->GetContainedNodes();
-				for (URigVMNode* ContainedNode : ContainedNodes)
+				if(URigVMController* CollapseController = GetControllerForGraph(InCollapseNode->GetContainedGraph()))
 				{
-					RepopulatePinsOnNode(ContainedNode, bFollowCoreRedirectors, bSetupOrphanedPins);
+					TGuardValue<bool> GuardEditGraph(InCollapseNode->ContainedGraph->bEditable, true);
+					// need to get a copy of the node array since the following function could remove nodes from the graph
+					// we don't want to remove elements from the array we are iterating over.
+					TArray<URigVMNode*> ContainedNodes = InCollapseNode->GetContainedNodes();
+					for (URigVMNode* ContainedNode : ContainedNodes)
+					{
+						CollapseController->RepopulatePinsOnNode(ContainedNode, bFollowCoreRedirectors, bSetupOrphanedPins);
+					}
 				}
 			}
 		};
@@ -17147,13 +17300,38 @@ void URigVMController::RewireLinks(URigVMPin* InOldPin, URigVMPin* InNewPin, boo
 
 #endif
 
-bool URigVMController::RenameObject(UObject* InObjectToRename, const TCHAR* InNewName, UObject* InNewOuter)
+bool URigVMController::RenameObject(UObject* InObjectToRename, const TCHAR* InNewName, UObject* InNewOuter) const
 {
-	return InObjectToRename->Rename(InNewName, InNewOuter, REN_ForceNoResetLoaders | REN_DoNotDirty | REN_DontCreateRedirectors | REN_NonTransactional);
+	const bool bSuccess = InObjectToRename->Rename(InNewName, InNewOuter, REN_ForceNoResetLoaders | REN_DoNotDirty | REN_DontCreateRedirectors | REN_NonTransactional);
+	if(bSuccess)
+	{
+		if(const URigVMCollapseNode* CollapseNode = Cast<URigVMCollapseNode>(InObjectToRename))
+		{
+			if(IRigVMClientHost* ClientHost = GetImplementingOuter<IRigVMClientHost>())
+			{
+				if(FRigVMClient* Client = ClientHost->GetRigVMClient())
+				{
+					Client->OnCollapseNodeRenamed(CollapseNode);
+				}
+			}
+		}
+	}
+	return bSuccess;
 }
 
-void URigVMController::DestroyObject(UObject* InObjectToDestroy)
+void URigVMController::DestroyObject(UObject* InObjectToDestroy) const
 {
+	if(const URigVMCollapseNode* CollapseNode = Cast<URigVMCollapseNode>(InObjectToDestroy))
+	{
+		if(IRigVMClientHost* ClientHost = GetImplementingOuter<IRigVMClientHost>())
+		{
+			if(FRigVMClient* Client = ClientHost->GetRigVMClient())
+			{
+				Client->OnCollapseNodeRemoved(CollapseNode);
+			}
+		}
+	}
+
 	RenameObject(InObjectToDestroy, nullptr, GetTransientPackage());
 	InObjectToDestroy->RemoveFromRoot();
 	InObjectToDestroy->MarkAsGarbage();
@@ -17318,17 +17496,18 @@ void URigVMController::RefreshFunctionReferences(URigVMLibraryNode* InFunctionDe
 {
 	check(InFunctionDefinition);
 
-	if (URigVMFunctionLibrary* FunctionLibrary = Cast<URigVMFunctionLibrary>(InFunctionDefinition->GetGraph()))
+	if (const URigVMFunctionLibrary* FunctionLibrary = Cast<URigVMFunctionLibrary>(InFunctionDefinition->GetGraph()))
 	{
 		FunctionLibrary->ForEachReference(InFunctionDefinition->GetFName(), [this, bSetupUndoRedo](URigVMFunctionReferenceNode* ReferenceNode)
 		{
-			FRigVMControllerGraphGuard GraphGuard(this, ReferenceNode->GetGraph(), bSetupUndoRedo);
-
-			const TArray<FLinkedPath> LinkedPaths = GetLinkedPaths(ReferenceNode->GetLinks());
-			FastBreakLinkedPaths(LinkedPaths);
-			RepopulatePinsOnNode(ReferenceNode, false, false, true);
-			TGuardValue<bool> ReportGuard(bReportWarningsAndErrors, false);
-			RestoreLinkedPaths(LinkedPaths);
+			if(URigVMController* ReferenceController = GetControllerForGraph(ReferenceNode->GetGraph()))
+			{
+				const TArray<FLinkedPath> LinkedPaths = GetLinkedPaths(ReferenceNode->GetLinks());
+				ReferenceController->FastBreakLinkedPaths(LinkedPaths);
+				ReferenceController->RepopulatePinsOnNode(ReferenceNode, false, false, true);
+				TGuardValue<bool> ReportGuard(ReferenceController->bReportWarningsAndErrors, false);
+				ReferenceController->RestoreLinkedPaths(LinkedPaths);
+			}
 		});
 	}
 }
@@ -18336,18 +18515,20 @@ FRigVMClientPatchResult URigVMController::PatchInvalidLinksOnWildcards()
 		}
 
 		// Remove exposed pins of type wildcard
-		if (URigVMCollapseNode* CollapseNode = Cast<URigVMCollapseNode>(Graph->GetOuter()))
+		if (const URigVMCollapseNode* CollapseNode = Cast<URigVMCollapseNode>(Graph->GetOuter()))
 		{
 			TArray<URigVMPin*> ExposedPins = CollapseNode->GetPins();
-			for (URigVMPin* ExposedPin : ExposedPins)
+			for (const URigVMPin* ExposedPin : ExposedPins)
 			{
 				if (ExposedPin->IsWildCard())
 				{
-					FRigVMControllerGraphGuard GraphGuard(this, CollapseNode->GetContainedGraph());
-					if (!RemoveExposedPin(ExposedPin->GetFName(), false))
+					if(URigVMController* CollapseController = GetControllerForGraph(CollapseNode->GetContainedGraph()))
 					{
-						Result.ErrorMessages.Add(FString::Printf(TEXT("Error removing exposed pin %s PatchInvalidLinksOnWildcards"), *ExposedPin->GetPinPath(true)));
-						Result.bSucceeded = false;
+						if (!CollapseController->RemoveExposedPin(ExposedPin->GetFName(), false))
+						{
+							Result.ErrorMessages.Add(FString::Printf(TEXT("Error removing exposed pin %s PatchInvalidLinksOnWildcards"), *ExposedPin->GetPinPath(true)));
+							Result.bSucceeded = false;
+						}
 					}
 				}
 			}
diff --git a/Engine/Source/Developer/RigVMDeveloper/Private/RigVMModel/RigVMControllerActions.cpp b/Engine/Source/Developer/RigVMDeveloper/Private/RigVMModel/RigVMControllerActions.cpp
--- a/Engine/Source/Developer/RigVMDeveloper/Private/RigVMModel/RigVMControllerActions.cpp
+++ b/Engine/Source/Developer/RigVMDeveloper/Private/RigVMModel/RigVMControllerActions.cpp
@@ -41,10 +41,11 @@ const UScriptStruct* FRigVMActionWrapper::GetScriptStruct() const
 FRigVMBaseAction* FRigVMActionWrapper::GetAction(URigVMController* InLastController) const
 {
 	FRigVMBaseAction* Action = (FRigVMBaseAction*)StructOnScope->GetStructMemory();
+	Action->EnsureControllerValidity();
 
 	// reuse the last controller so that we don't have to rely on resolving it
 	// from the soft object path each time
-	if(InLastController != nullptr && !Action->WeakController.IsValid())
+	if(Action->WeakController.IsValid() && IsValid(InLastController) && InLastController->IsValidGraph())
 	{
 		if(Action->ControllerPath.ToString() == InLastController->GetPathName())
 		{
@@ -245,6 +246,7 @@ void URigVMActionStack::LogAction(const UScriptStruct* InActionStruct, const FRi
 
 URigVMController* FRigVMBaseAction::GetController() const
 {
+	EnsureControllerValidity();
 	if(!WeakController.IsValid())
 	{
 		if(URigVMController* Controller = Cast<URigVMController>(ControllerPath.ResolveObject()))
@@ -255,6 +257,20 @@ URigVMController* FRigVMBaseAction::GetController() const
 	return WeakController.Get();
 }
 
+void FRigVMBaseAction::EnsureControllerValidity() const
+{
+	if(WeakController.IsValid())
+	{
+		if(const URigVMController* Controller = WeakController.Get())
+		{
+			if(!IsValid(Controller) || !Controller->IsValidGraph())
+			{
+				WeakController.Reset();
+			}
+		}
+	}
+}
+
 bool FRigVMBaseAction::Merge(const FRigVMBaseAction* Other)
 {
 	return SubActions.Num() == 0 && Other->SubActions.Num() == 0;
@@ -284,12 +300,13 @@ bool FRigVMBaseAction::Undo()
 	for (int32 KeyIndex = SubActions.Num() - 1; KeyIndex >= 0; KeyIndex--)
 	{
 		FRigVMActionWrapper Wrapper(SubActions[KeyIndex]);
+		FRigVMBaseAction* SubAction = Wrapper.GetAction(Controller);
 #if RIGVM_ACTIONSTACK_VERBOSE_LOG		
-		Wrapper.GetAction(Controller)->LogAction(UndoPrefix);
+		SubAction->LogAction(UndoPrefix);
 #endif
-		if(!Wrapper.GetAction(Controller)->Undo())
+		if(!SubAction->Undo())
 		{
-			Controller->ReportAndNotifyErrorf(TEXT("Error while undoing action '%s'."), *Wrapper.GetAction(Controller)->Title);
+			Controller->ReportAndNotifyErrorf(TEXT("Error while undoing action '%s'."), *SubAction->Title);
 			Result = false;
 		}
 	}
@@ -1528,118 +1545,6 @@ bool FRigVMRenameNodeAction::Redo()
 	return false;
 }
 
-FRigVMPushGraphAction::FRigVMPushGraphAction()
-: FRigVMBaseAction(nullptr)
-{
-}
-
-FRigVMPushGraphAction::FRigVMPushGraphAction(URigVMController* InController, UObject* InGraph)
-: FRigVMBaseAction(InController)
-{
-	GraphPath = TSoftObjectPtr<URigVMGraph>(Cast<URigVMGraph>(InGraph)).GetUniqueID();
-}
-
-bool FRigVMPushGraphAction::Undo()
-{
-	if (!FRigVMBaseAction::Undo())
-	{
-		return false;
-	}
-	
-	if(URigVMGraph* PoppedGraph = GetController()->PopGraph(false))
-	{
-		const FSoftObjectPath PoppedGraphPath = TSoftObjectPtr<URigVMGraph>(PoppedGraph).GetUniqueID();
-		return PoppedGraphPath == GraphPath;
-	}
-	
-	return false;
-}
-
-bool FRigVMPushGraphAction::Redo()
-{
-	if(!CanUndoRedo())
-	{
-		return false;
-	}
-	TSoftObjectPtr<URigVMGraph> GraphPtr(GraphPath);
-	if(URigVMGraph* Graph = GraphPtr.Get())
-	{
-		GetController()->PushGraph(Graph, false);
-		return FRigVMBaseAction::Redo();
-	}
-	return false;
-}
-
-bool FRigVMPushGraphAction::MakesObsolete(const FRigVMBaseAction* Other) const
-{
-	if (Other->GetScriptStruct() == FRigVMPopGraphAction::StaticStruct())
-	{
-		const FRigVMPopGraphAction* PopGraphAction = (const FRigVMPopGraphAction*)Other;
-		if(PopGraphAction->GraphPath == GraphPath)
-		{
-			return true;
-		}
-	}
-	return false;
-}
-
-FRigVMPopGraphAction::FRigVMPopGraphAction()
-: FRigVMBaseAction(nullptr)
-{
-}
-
-FRigVMPopGraphAction::FRigVMPopGraphAction(URigVMController* InController, UObject* InGraph)
-: FRigVMBaseAction(InController)
-{
-	GraphPath = TSoftObjectPtr<URigVMGraph>(Cast<URigVMGraph>(InGraph)).GetUniqueID();
-}
-
-bool FRigVMPopGraphAction::Undo()
-{
-	if (!FRigVMBaseAction::Undo())
-	{
-		return false;
-	}
-
-	TSoftObjectPtr<URigVMGraph> GraphPtr(GraphPath);
-	if (URigVMGraph* Graph = GraphPtr.Get())
-	{
-		GetController()->PushGraph(Graph, false);
-		return true;
-	}
-	return false;
-}
-
-bool FRigVMPopGraphAction::Redo()
-{
-	if(!CanUndoRedo())
-	{
-		return false;
-	}
-	if(URigVMGraph* PoppedGraph = GetController()->PopGraph(false))
-	{
-		const FSoftObjectPath PoppedGraphPath = TSoftObjectPtr<URigVMGraph>(PoppedGraph).GetUniqueID();
-		if(PoppedGraphPath == GraphPath)
-		{
-			return FRigVMBaseAction::Redo();
-		}
-	}
-	return false;
-}
-
-bool FRigVMPopGraphAction::MakesObsolete(const FRigVMBaseAction* Other) const
-{
-	if (Other->GetScriptStruct() == FRigVMPushGraphAction::StaticStruct())
-	{
-		const FRigVMPushGraphAction* PushGraphAction = (const FRigVMPushGraphAction*)Other;
-		if(PushGraphAction->GraphPath == GraphPath)
-		{
-			return true;
-		}
-	}
-	return false;
-}
-
 FRigVMAddExposedPinAction::FRigVMAddExposedPinAction()
 : FRigVMBaseAction(nullptr)
 , Direction(ERigVMPinDirection::Input)
diff --git a/Engine/Source/Developer/RigVMDeveloper/Public/RigVMModel/RigVMClient.h b/Engine/Source/Developer/RigVMDeveloper/Public/RigVMModel/RigVMClient.h
--- a/Engine/Source/Developer/RigVMDeveloper/Public/RigVMModel/RigVMClient.h
+++ b/Engine/Source/Developer/RigVMDeveloper/Public/RigVMModel/RigVMClient.h
@@ -112,6 +112,7 @@ public:
 	URigVMController* GetOrCreateController(const FString& InNodePathOrName);
 	URigVMController* GetOrCreateController(const URigVMGraph* InModel);
 	URigVMController* GetOrCreateController(const UObject* InEditorSideObject);
+	bool RemoveController(const URigVMGraph* InModel);
 	URigVMFunctionLibrary* GetFunctionLibrary() const { return FunctionLibrary; }
 	URigVMFunctionLibrary* GetOrCreateFunctionLibrary(bool bSetupUndoRedo, const FObjectInitializer* ObjectInitializer = nullptr, bool bCreateController = true);
 	TArray<FName> GetEntryNames() const;
@@ -123,7 +124,8 @@ public:
 	bool RemoveModel(const FString& InNodePathOrName, bool bSetupUndoRedo);
 	FName RenameModel(const FString& InNodePathOrName, const FName& InNewName, bool bSetupUndoRedo);
 	void PostTransacted(const FTransactionObjectEvent& TransactionEvent);
-	void OnSubGraphRenamed(const FSoftObjectPath& OldObjectPath, const FSoftObjectPath& NewObjectPath);
+	void OnCollapseNodeRenamed(const URigVMCollapseNode* InCollapseNode);
+	void OnCollapseNodeRemoved(const URigVMCollapseNode* InCollapseNode);
 
 	URigVMNode* FindNode(const FString& InNodePathOrName) const;
 	URigVMPin* FindPin(const FString& InPinPath) const;
diff --git a/Engine/Source/Developer/RigVMDeveloper/Public/RigVMModel/RigVMController.h b/Engine/Source/Developer/RigVMDeveloper/Public/RigVMModel/RigVMController.h
--- a/Engine/Source/Developer/RigVMDeveloper/Public/RigVMModel/RigVMController.h
+++ b/Engine/Source/Developer/RigVMDeveloper/Public/RigVMModel/RigVMController.h
@@ -200,23 +200,27 @@ public:
 
 	// Sets the currently edited Graph of this controller.
 	// This causes a GraphChanged modified event.
-	UFUNCTION(BlueprintCallable, Category = RigVMController)
+	UFUNCTION(BlueprintCallable, Category = RigVMController, meta=(DeprecatedFunction, DeprecationMessage="Function has been deprecated, please rely on GetControllerForGraph instead."))
 	void SetGraph(URigVMGraph* InGraph);
 
 	// Pushes a new graph to the stack
 	// This causes a GraphChanged modified event.
-	UFUNCTION(BlueprintCallable, Category = RigVMController)
+	UFUNCTION(BlueprintCallable, Category = RigVMController, meta=(DeprecatedFunction, DeprecationMessage="Function has been deprecated, please rely on GetControllerForGraph instead."))
 	bool PushGraph(URigVMGraph* InGraph, bool bSetupUndoRedo = true);
 
 	// Pops the last graph off the stack
 	// This causes a GraphChanged modified event.
-	UFUNCTION(BlueprintCallable, Category = RigVMController)
+	UFUNCTION(BlueprintCallable, Category = RigVMController, meta=(DeprecatedFunction, DeprecationMessage="Function has been deprecated, please rely on GetControllerForGraph instead."))
 	URigVMGraph* PopGraph(bool bSetupUndoRedo = true);
 	
 	// Returns the top level graph
 	UFUNCTION(BlueprintCallable, Category = RigVMController)
 	URigVMGraph* GetTopLevelGraph() const;
 
+	// Returns another controller for a given graph
+	UFUNCTION(BlueprintCallable, Category = RigVMController)
+	URigVMController* GetControllerForGraph(const URigVMGraph* InGraph) const;
+
 	// The Modified event used to subscribe to changes
 	// happening within the Graph. This is broadcasted to 
 	// for any change happening - not only the changes 
@@ -1171,8 +1175,8 @@ private:
 #endif
 
 
-	bool RenameObject(UObject* InObjectToRename, const TCHAR* InNewName, UObject* InNewOuter = nullptr);
-	void DestroyObject(UObject* InObjectToDestroy);
+	bool RenameObject(UObject* InObjectToRename, const TCHAR* InNewName, UObject* InNewOuter = nullptr) const;
+	void DestroyObject(UObject* InObjectToDestroy) const ;
 	static URigVMPin* MakeExecutePin(URigVMNode* InNode, const FName& InName);
 	static void MakeExecutePin(URigVMPin* InOutPin);
 	static void AddNodePin(URigVMNode* InNode, URigVMPin* InPin);
@@ -1425,51 +1429,10 @@ private:
 	friend struct FRigVMEjectNodeFromPinAction;
 	friend class FRigVMParserAST;
 	friend class FRigVMControllerCompileBracketScope;
-	friend class FRigVMControllerGraphGuard;
 	friend struct FPinInfoArray;
 	friend class FRigVMControllerNotifGuard;
 	friend struct FRigVMClient;
-};
-
-class FRigVMControllerGraphGuard
-{
-public:
-
-	FRigVMControllerGraphGuard(URigVMController* InController, URigVMGraph* InGraph, bool bSetupUndoRedo = true)
-		: Controller(InController)
-		, bUndo(bSetupUndoRedo)
-		, bEnabled(true)
-	{
-		bEnabled = Controller->PushGraph(InGraph, bUndo);
-		if(bEnabled)
-		{
-			NumGraphs = Controller->Graphs.Num();
-		}
-	}
-
-	~FRigVMControllerGraphGuard()
-	{
-		if(!bEnabled)
-		{
-			return;
-		}
-		
-		// an action can be cancelled in the middle of a graph guard,
-		// in that case CancelAction should have already popped the graph
-		if (Controller->Graphs.Num() < NumGraphs)
-		{
-			return;
-		}
-		Controller->PopGraph(bUndo);
-	}
-
-private:
-
-	URigVMController* Controller;
-	bool bUndo;
-	bool bEnabled;
-
-	int32 NumGraphs;
+	friend struct FRigVMActionWrapper;
 };
 
 class FRigVMControllerNotifGuard
diff --git a/Engine/Source/Developer/RigVMDeveloper/Public/RigVMModel/RigVMControllerActions.h b/Engine/Source/Developer/RigVMDeveloper/Public/RigVMModel/RigVMControllerActions.h
--- a/Engine/Source/Developer/RigVMDeveloper/Public/RigVMModel/RigVMControllerActions.h
+++ b/Engine/Source/Developer/RigVMDeveloper/Public/RigVMModel/RigVMControllerActions.h
@@ -199,6 +199,8 @@ protected:
 	{
 	}
 
+	void EnsureControllerValidity() const;
+
 	mutable TWeakObjectPtr<URigVMController> WeakController;
 
 	UPROPERTY()
@@ -1107,51 +1109,6 @@ public:
 	FString NewNodeName;
 };
 
-/**
- * An action pushing a graph to the graph stack of the controller
- */
-USTRUCT()
-struct FRigVMPushGraphAction : public FRigVMBaseAction
-{
-	GENERATED_BODY()
-
-public:
-
-	FRigVMPushGraphAction();
-	FRigVMPushGraphAction(URigVMController* InController, UObject* InGraph);
-	virtual ~FRigVMPushGraphAction() {};
-	virtual UScriptStruct* GetScriptStruct() const override { return FRigVMPushGraphAction::StaticStruct(); }
-	virtual bool Undo() override;
-	virtual bool Redo() override;
-	virtual bool MakesObsolete(const FRigVMBaseAction* Other) const override;
-
-	UPROPERTY()
-	FSoftObjectPath GraphPath;
-};
-
-/**
- * An action popping a graph from the graph stack of the controller
- */
-USTRUCT()
-struct FRigVMPopGraphAction : public FRigVMBaseAction
-{
-	GENERATED_BODY()
-
-public:
-
-	FRigVMPopGraphAction();
-	FRigVMPopGraphAction(URigVMController* InController, UObject* InGraph);
-	virtual ~FRigVMPopGraphAction() {};
-	virtual UScriptStruct* GetScriptStruct() const override { return FRigVMPopGraphAction::StaticStruct(); }
-	virtual bool Undo() override;
-	virtual bool Redo() override;
-	virtual bool MakesObsolete(const FRigVMBaseAction* Other) const override;
-
-	UPROPERTY()
-	FSoftObjectPath GraphPath;
-};
-
-
 /**
  * An action exposing a pin as a parameter
  */
